[Mass] MassEntity moved over to Engine/Source/Runtime/

diff --git a/Engine/Plugins/Runtime/MassEntity/MassEntity.uplugin b/Engine/Plugins/Runtime/MassEntity/MassEntity.uplugin
--- a/Engine/Plugins/Runtime/MassEntity/MassEntity.uplugin
+++ b/Engine/Plugins/Runtime/MassEntity/MassEntity.uplugin
@@ -15,11 +15,6 @@
 	"Installed": false,
 	"EnabledByDefault": false,
 	"Modules": [
-		{
-			"Name": "MassEntity",
-			"Type": "Runtime",
-			"LoadingPhase": "PreLoadingScreen"
-		},
 		{
 			"Name": "MassEntityTestSuite",
 			"Type": "UncookedOnly",
diff --git a/Engine/Source/Runtime/MassEntity/MassEntity.Build.cs b/Engine/Source/Runtime/MassEntity/MassEntity.Build.cs
new file mode 100644
--- /dev/null
+++ b/Engine/Source/Runtime/MassEntity/MassEntity.Build.cs
@@ -0,0 +1,41 @@
+// Copyright Epic Games, Inc. All Rights Reserved.
+
+namespace UnrealBuildTool.Rules
+{
+	public class MassEntity : ModuleRules
+	{
+		public MassEntity(ReadOnlyTargetRules Target) : base(Target)
+		{
+			PCHUsage = ModuleRules.PCHUsageMode.UseExplicitOrSharedPCHs;
+
+			UnsafeTypeCastWarningLevel = WarningLevel.Warning;
+
+			PublicDependencyModuleNames.AddRange(
+				new string[] {
+					"Core",
+					"CoreUObject",
+					"Engine",
+					"DeveloperSettings",
+				}
+			);
+
+			if (Target.bBuildEditor || Target.bCompileAgainstEditor)
+			{
+				PrivateDependencyModuleNames.Add("UnrealEd");
+				PrivateDependencyModuleNames.Add("EditorSubsystem");
+			}
+
+			if (Target.bBuildDeveloperTools == true)
+			{
+				PrivateDependencyModuleNames.Add("MessageLog");
+			}
+
+			if (Target.Configuration != UnrealTargetConfiguration.Shipping
+				&& Target.Configuration != UnrealTargetConfiguration.Test)
+			{
+				// pulling this one in for the testableEnsureMsgf
+				PrivateDependencyModuleNames.Add("AITestSuite");
+			}
+		}
+	}
+}
\ No newline at end of file
diff --git a/Engine/Source/Runtime/MassEntity/Private/MassEntityModule.cpp b/Engine/Source/Runtime/MassEntity/Private/MassEntityModule.cpp
new file mode 100644
--- /dev/null
+++ b/Engine/Source/Runtime/MassEntity/Private/MassEntityModule.cpp
@@ -0,0 +1,69 @@
+// Copyright Epic Games, Inc. All Rights Reserved.
+
+#include "CoreMinimal.h"
+#include "IMassEntityModule.h"
+#include "MassEntityTypes.h"
+#if WITH_UNREAL_DEVELOPER_TOOLS
+#include "MessageLogModule.h"
+#include "Engine/World.h"
+#include "Logging/MessageLog.h"
+#endif // WITH_UNREAL_DEVELOPER_TOOLS
+
+#define LOCTEXT_NAMESPACE "Mass"
+
+class FMassEntityModule : public IMassEntityModule
+{
+	/** IModuleInterface implementation */
+	virtual void StartupModule() override;
+	virtual void ShutdownModule() override;
+#if WITH_UNREAL_DEVELOPER_TOOLS
+	static void OnWorldCleanup(UWorld* /*World*/, bool /*bSessionEnded*/, bool /*bCleanupResources*/);
+	FDelegateHandle OnWorldCleanupHandle;
+#endif // WITH_UNREAL_DEVELOPER_TOOLS
+};
+
+IMPLEMENT_MODULE(FMassEntityModule, MassEntity)
+
+
+void FMassEntityModule::StartupModule()
+{
+	// This code will execute after your module is loaded into memory (but after global variables are initialized, of course.)
+
+#if WITH_UNREAL_DEVELOPER_TOOLS
+	{
+		FMessageLogModule& MessageLogModule = FModuleManager::LoadModuleChecked<FMessageLogModule>("MessageLog");
+		FMessageLogInitializationOptions InitOptions;
+		InitOptions.bShowPages = true;
+		InitOptions.bShowFilters = true;
+		MessageLogModule.RegisterLogListing("MassEntity", LOCTEXT("MassEntity", "MassEntity"), InitOptions);
+
+		OnWorldCleanupHandle = FWorldDelegates::OnWorldCleanup.AddStatic(&FMassEntityModule::OnWorldCleanup);
+	}
+#endif // WITH_UNREAL_DEVELOPER_TOOLS
+
+#if MASS_DO_PARALLEL
+	UE_LOG(LogMass, Log, TEXT("MassEntity running with MULTITHREADING support."));
+#else
+	UE_LOG(LogMass, Log, TEXT("MassEntity running in game thread."));
+#endif // MASS_DO_PARALLEL
+}
+
+
+void FMassEntityModule::ShutdownModule()
+{
+	// This function may be called during shutdown to clean up your module.  For modules that support dynamic reloading,
+	// we call this function before unloading the module.
+#if WITH_UNREAL_DEVELOPER_TOOLS
+	FWorldDelegates::OnWorldCleanup.Remove(OnWorldCleanupHandle);
+#endif // WITH_UNREAL_DEVELOPER_TOOLS
+}
+
+#if WITH_UNREAL_DEVELOPER_TOOLS
+void FMassEntityModule::OnWorldCleanup(UWorld* /*World*/, bool /*bSessionEnded*/, bool /*bCleanupResources*/)
+{
+	// clearing out messages from the world being cleaned up
+	FMessageLog("MassEntity").NewPage(FText::FromString(TEXT("MassEntity")));
+}
+#endif // WITH_UNREAL_DEVELOPER_TOOLS
+
+#undef LOCTEXT_NAMESPACE 
diff --git a/Engine/Source/Runtime/MassEntity/Private/MassExecutor.cpp b/Engine/Source/Runtime/MassEntity/Private/MassExecutor.cpp
new file mode 100644
--- /dev/null
+++ b/Engine/Source/Runtime/MassEntity/Private/MassExecutor.cpp
@@ -0,0 +1,245 @@
+// Copyright Epic Games, Inc. All Rights Reserved.
+
+#include "MassExecutor.h"
+#include "MassProcessingTypes.h"
+#include "MassProcessor.h"
+#include "MassCommandBuffer.h"
+#include "MassExecutionContext.h"
+#include "ProfilingDebugging/CpuProfilerTrace.h"
+
+namespace UE::Mass::Executor
+{
+
+FORCEINLINE void ExecuteProcessors(FMassEntityManager& EntityManager, TArrayView<UMassProcessor* const> Processors, FMassExecutionContext& ExecutionContext)
+{
+	for (UMassProcessor* Proc : Processors)
+	{
+		Proc->CallExecute(EntityManager, ExecutionContext);
+	}
+}
+
+void Run(FMassRuntimePipeline& RuntimePipeline, FMassProcessingContext& ProcessingContext)
+{
+	if (!ensure(ProcessingContext.EntityManager) || 
+		!ensure(ProcessingContext.DeltaSeconds >= 0.f) ||
+		!ensure(RuntimePipeline.GetProcessors().Find(nullptr) == INDEX_NONE))
+	{
+		return;
+	}
+
+	TRACE_CPUPROFILER_EVENT_SCOPE_STR("MassExecutor Run Pipeline")
+	RunProcessorsView(RuntimePipeline.GetMutableProcessors(), ProcessingContext);
+}
+
+void RunSparse(FMassRuntimePipeline& RuntimePipeline, FMassProcessingContext& ProcessingContext, FMassArchetypeHandle Archetype, TConstArrayView<FMassEntityHandle> Entities)
+{
+	if (!ensure(ProcessingContext.EntityManager) ||
+		!ensure(RuntimePipeline.GetProcessors().Find(nullptr) == INDEX_NONE) ||
+		RuntimePipeline.Num() == 0 ||
+		!ensureMsgf(Archetype.IsValid(), TEXT("The Archetype passed in to UE::Mass::Executor::RunSparse is invalid")))
+	{
+		return;
+	}
+
+	TRACE_CPUPROFILER_EVENT_SCOPE_STR("MassExecutor RunSparseEntities");
+
+	const FMassArchetypeEntityCollection EntityCollection(Archetype, Entities, FMassArchetypeEntityCollection::NoDuplicates);
+	RunProcessorsView(RuntimePipeline.GetMutableProcessors(), ProcessingContext, MakeArrayView(&EntityCollection, 1));
+}
+
+void RunSparse(FMassRuntimePipeline& RuntimePipeline, FMassProcessingContext& ProcessingContext, const FMassArchetypeEntityCollection& EntityCollection)
+{
+	if (!ensure(ProcessingContext.EntityManager) ||
+		!ensure(RuntimePipeline.GetProcessors().Find(nullptr) == INDEX_NONE) ||
+		RuntimePipeline.Num() == 0 ||
+		!ensureMsgf(EntityCollection.GetArchetype().IsValid(), TEXT("The Archetype of EntityCollection passed in to UE::Mass::Executor::RunSparse is invalid")))
+	{
+		return;
+	}
+
+	TRACE_CPUPROFILER_EVENT_SCOPE_STR("MassExecutor RunSparse");
+
+	RunProcessorsView(RuntimePipeline.GetMutableProcessors(), ProcessingContext, MakeArrayView(&EntityCollection, 1));
+}
+
+void Run(UMassProcessor& Processor, FMassProcessingContext& ProcessingContext)
+{
+	if (!ensure(ProcessingContext.EntityManager) || !ensure(ProcessingContext.DeltaSeconds >= 0.f))
+	{
+		return;
+	}
+
+	TRACE_CPUPROFILER_EVENT_SCOPE_STR("MassExecutor Run")
+
+	UMassProcessor* ProcPtr = &Processor;
+	RunProcessorsView(MakeArrayView(&ProcPtr, 1), ProcessingContext);
+}
+
+void RunProcessorsView(TArrayView<UMassProcessor* const> Processors, FMassProcessingContext& ProcessingContext, TConstArrayView<FMassArchetypeEntityCollection> EntityCollections)
+{
+	TRACE_CPUPROFILER_EVENT_SCOPE(RunProcessorsView);
+
+	if (!ProcessingContext.EntityManager)
+	{
+		UE_LOG(LogMass, Error, TEXT("%s ProcessingContext.EntityManager is null. Baling out."), ANSI_TO_TCHAR(__FUNCTION__));
+		return;
+	}
+#if WITH_MASSENTITY_DEBUG
+	if (Processors.Find(nullptr) != INDEX_NONE)
+	{
+		UE_LOG(LogMass, Error, TEXT("%s input Processors contains nullptr. Baling out."), ANSI_TO_TCHAR(__FUNCTION__));
+		return;
+	}
+#endif // WITH_MASSENTITY_DEBUG
+
+	TRACE_CPUPROFILER_EVENT_SCOPE_STR("MassExecutor RunProcessorsView")
+
+	FMassExecutionContext ExecutionContext(*ProcessingContext.EntityManager.Get(), ProcessingContext.DeltaSeconds);
+
+	// if ProcessingContext points at a valid CommandBuffer use that one, otherwise manually create a new command buffer 
+	// to let the default one still be used by code unaware of mass processing
+	TSharedPtr<FMassCommandBuffer> CommandBuffer = ProcessingContext.CommandBuffer 
+		? ProcessingContext.CommandBuffer : MakeShareable(new FMassCommandBuffer());
+	ExecutionContext.SetDeferredCommandBuffer(CommandBuffer);
+	ExecutionContext.SetFlushDeferredCommands(false);
+	ExecutionContext.SetAuxData(ProcessingContext.AuxData);
+	ExecutionContext.SetExecutionType(EMassExecutionContextType::Processor);
+
+	FMassEntityManager& EntityManager = *ProcessingContext.EntityManager.Get();
+	{
+		TRACE_CPUPROFILER_EVENT_SCOPE_STR("Execute Processors")
+		
+		FMassEntityManager::FScopedProcessing ProcessingScope = EntityManager.NewProcessingScope();
+
+		if (EntityCollections.Num() == 0)
+		{
+			ExecuteProcessors(*ProcessingContext.EntityManager, Processors, ExecutionContext);
+		}
+		else
+		{
+			// @todo change ExecutionContext to contain a TConstArrayView of collections. Will lead up to changes in the entity query as well.
+			for (const FMassArchetypeEntityCollection& Collection : EntityCollections)
+			{
+				ExecutionContext.SetEntityCollection(Collection);
+				ExecuteProcessors(*ProcessingContext.EntityManager, Processors, ExecutionContext);
+				ExecutionContext.ClearEntityCollection();
+			}
+		}
+	}
+	
+	if (ProcessingContext.bFlushCommandBuffer)
+	{		
+		TRACE_CPUPROFILER_EVENT_SCOPE_STR("Flush Deferred Commands")
+		
+		ExecutionContext.SetFlushDeferredCommands(true);
+		// append the commands added from other, non-processor sources (like MassAgentSubsystem)
+		ensure(!EntityManager.Defer().IsFlushing());
+		ensure(!ExecutionContext.Defer().IsFlushing());
+		ExecutionContext.Defer().MoveAppend(EntityManager.Defer());
+		ExecutionContext.FlushDeferred();
+	}
+	// else make sure we don't just lose the commands. Append to the command buffer requested via
+	// ProcessingContext.CommandBuffer or to the default EntityManager's command buffer.
+	else if (CommandBuffer->HasPendingCommands() && CommandBuffer != ProcessingContext.CommandBuffer)
+	{
+		if (ProcessingContext.CommandBuffer)
+		{
+			ProcessingContext.CommandBuffer->MoveAppend(*CommandBuffer.Get());
+		}
+		else if (CommandBuffer.Get() != &EntityManager.Defer())
+		{
+			EntityManager.AppendCommands(CommandBuffer);
+		}
+	}
+}
+
+struct FMassExecutorDoneTask
+{
+	FMassExecutorDoneTask(const FMassExecutionContext& InExecutionContext, TFunction<void()> InOnDoneNotification, const FString& InDebugName)
+		: ExecutionContext(InExecutionContext)
+		, OnDoneNotification(InOnDoneNotification)
+		, DebugName(InDebugName)
+	{
+	}
+	static TStatId GetStatId()
+	{
+		RETURN_QUICK_DECLARE_CYCLE_STAT(FMassExecutorDoneTask, STATGROUP_TaskGraphTasks);
+	}
+
+	static ENamedThreads::Type GetDesiredThread() { return ENamedThreads::GameThread; }
+	static ESubsequentsMode::Type GetSubsequentsMode() { return ESubsequentsMode::TrackSubsequents; }
+	void DoTask(ENamedThreads::Type CurrentThread, const FGraphEventRef& MyCompletionGraphEvent)
+	{
+		TRACE_CPUPROFILER_EVENT_SCOPE_STR("Flush Deferred Commands Parallel");
+		SCOPE_CYCLE_COUNTER(STAT_Mass_Total);
+
+		FMassEntityManager& EntityManagerRef = ExecutionContext.GetEntityManagerChecked();
+
+		if (&ExecutionContext.Defer() != &EntityManagerRef.Defer())
+		{
+			ExecutionContext.Defer().MoveAppend(EntityManagerRef.Defer());
+		}
+
+		UE_LOG(LogMass, Verbose, TEXT("MassExecutor %s tasks DONE"), *DebugName);
+		ExecutionContext.SetFlushDeferredCommands(true);
+		ExecutionContext.FlushDeferred();
+
+		OnDoneNotification();
+	}
+private:
+	FMassExecutionContext ExecutionContext;
+	TFunction<void()> OnDoneNotification;
+	FString DebugName;
+};
+
+FGraphEventRef TriggerParallelTasks(UMassProcessor& Processor, FMassProcessingContext& ProcessingContext, TFunction<void()> OnDoneNotification)
+{
+	if (!ProcessingContext.EntityManager)
+	{
+		UE_LOG(LogMass, Error, TEXT("%s ProcessingContext.EntityManager is null. Baling out."), ANSI_TO_TCHAR(__FUNCTION__));
+		return FGraphEventRef();
+	}
+
+	TRACE_CPUPROFILER_EVENT_SCOPE_STR("MassExecutor RunParallel")
+
+	// not going through FMassEntityManager::CreateExecutionContext on purpose - we do need a separate command buffer
+	FMassExecutionContext ExecutionContext(*ProcessingContext.EntityManager.Get(), ProcessingContext.DeltaSeconds);
+	TSharedPtr<FMassCommandBuffer> CommandBuffer = ProcessingContext.CommandBuffer
+		? ProcessingContext.CommandBuffer : MakeShareable(new FMassCommandBuffer());
+	ExecutionContext.SetDeferredCommandBuffer(CommandBuffer);
+	ExecutionContext.SetFlushDeferredCommands(false);
+	ExecutionContext.SetAuxData(ProcessingContext.AuxData);
+	ExecutionContext.SetExecutionType(EMassExecutionContextType::Processor);
+
+	FGraphEventRef CompletionEvent;
+	{
+		TRACE_CPUPROFILER_EVENT_SCOPE_STR("Dispatch Processors")
+		CompletionEvent = Processor.DispatchProcessorTasks(ProcessingContext.EntityManager, ExecutionContext, {});
+	}
+
+	if (CompletionEvent.IsValid())
+	{
+		const FGraphEventArray Prerequisites = { CompletionEvent };
+		CompletionEvent = TGraphTask<FMassExecutorDoneTask>::CreateTask(&Prerequisites)
+			.ConstructAndDispatchWhenReady(ExecutionContext, OnDoneNotification, Processor.GetName());
+	}
+
+	return CompletionEvent;
+}
+
+//-----------------------------------------------------------------------------
+// DEPRECATED
+//-----------------------------------------------------------------------------
+void RunProcessorsView(TArrayView<UMassProcessor* const> Processors, FMassProcessingContext& ProcessingContext, const FMassArchetypeEntityCollection* EntityCollection)
+{
+	if (EntityCollection)
+	{
+		RunProcessorsView(Processors, ProcessingContext, MakeArrayView(EntityCollection, 1));
+	}
+	else
+	{
+		RunProcessorsView(Processors, ProcessingContext);
+	}
+}
+
+} // namespace UE::Mass::Executor
diff --git a/Engine/Source/Runtime/MassEntity/Private/MassProcessorDependencySolver.cpp b/Engine/Source/Runtime/MassEntity/Private/MassProcessorDependencySolver.cpp
new file mode 100644
--- /dev/null
+++ b/Engine/Source/Runtime/MassEntity/Private/MassProcessorDependencySolver.cpp
@@ -0,0 +1,1210 @@
+// Copyright Epic Games, Inc. All Rights Reserved.
+
+#include "MassProcessorDependencySolver.h"
+#include "MassProcessingTypes.h"
+#include "MassProcessor.h"
+#include "Logging/MessageLog.h"
+#include "HAL/FileManager.h"
+#include "Misc/Paths.h"
+#include "Engine/World.h"
+
+#define LOCTEXT_NAMESPACE "Mass"
+
+namespace UE::Mass::Private
+{
+	FString NameViewToString(TConstArrayView<FName> View)
+	{
+		if (View.Num() == 0)
+		{
+			return TEXT("[]");
+		}
+		FString ReturnVal = FString::Printf(TEXT("[%s"), *View[0].ToString());
+		for (int i = 1; i < View.Num(); ++i)
+		{
+			ReturnVal += FString::Printf(TEXT(", %s"), *View[i].ToString());
+		}
+		return ReturnVal + TEXT("]");
+	}
+
+	bool DoArchetypeContainersOverlap(TConstArrayView<FMassArchetypeHandle> A, const TArray<FMassArchetypeHandle>& B)
+	{
+		for (const FMassArchetypeHandle& HandleA : A)
+		{
+			if (B.Contains(HandleA))
+			{
+				return true;
+			}
+		}
+		return false;
+	}
+}
+
+//----------------------------------------------------------------------//
+//  FMassExecutionRequirements
+//----------------------------------------------------------------------//
+void FMassExecutionRequirements::Append(const FMassExecutionRequirements& Other)
+{
+	for (int i = 0; i < EMassAccessOperation::MAX; ++i)
+	{
+		Fragments[i] += Other.Fragments[i];
+		ChunkFragments[i] += Other.ChunkFragments[i];
+		SharedFragments[i] += Other.SharedFragments[i];
+		RequiredSubsystems[i] += Other.RequiredSubsystems[i];
+	}
+	ConstSharedFragments.Read += Other.ConstSharedFragments.Read;
+
+	RequiredAllTags += Other.RequiredAllTags;
+	RequiredAnyTags += Other.RequiredAnyTags;
+	RequiredNoneTags += Other.RequiredNoneTags;
+	// note that we're deliberately ignoring optional tags, they play no role here.
+
+	// signal that it requires recalculation;
+	ResourcesUsedCount = INDEX_NONE;
+}
+
+void FMassExecutionRequirements::CountResourcesUsed()
+{
+	ResourcesUsedCount = ConstSharedFragments.Read.CountStoredTypes();
+
+	for (int i = 0; i < EMassAccessOperation::MAX; ++i)
+	{
+		ResourcesUsedCount += Fragments[i].CountStoredTypes();
+		ResourcesUsedCount += ChunkFragments[i].CountStoredTypes();
+		ResourcesUsedCount += SharedFragments[i].CountStoredTypes();
+		ResourcesUsedCount += RequiredSubsystems[i].CountStoredTypes();
+	}
+}
+
+int32 FMassExecutionRequirements::GetTotalBitsUsedCount()
+{
+	CountResourcesUsed();
+
+	return ResourcesUsedCount + RequiredAllTags.CountStoredTypes()
+		+ RequiredAnyTags.CountStoredTypes() + RequiredNoneTags.CountStoredTypes();
+}
+
+bool FMassExecutionRequirements::IsEmpty() const
+{
+	return Fragments.IsEmpty() && ChunkFragments.IsEmpty() 
+		&& SharedFragments.IsEmpty() && ConstSharedFragments.IsEmpty() && RequiredSubsystems.IsEmpty()
+		&& RequiredAllTags.IsEmpty() && RequiredAnyTags.IsEmpty() && RequiredNoneTags.IsEmpty();
+}
+
+FMassArchetypeCompositionDescriptor FMassExecutionRequirements::AsCompositionDescriptor() const
+{
+	return FMassArchetypeCompositionDescriptor(Fragments.Read + Fragments.Write
+		, RequiredAllTags + RequiredAnyTags
+		, ChunkFragments.Read + ChunkFragments.Write
+		, SharedFragments.Read + SharedFragments.Write
+		, ConstSharedFragments.Read);
+}
+
+//----------------------------------------------------------------------//
+//  FProcessorDependencySolver::FResourceUsage
+//----------------------------------------------------------------------//
+FMassProcessorDependencySolver::FResourceUsage::FResourceUsage(const TArray<FNode>& InAllNodes)
+	: AllNodesView(InAllNodes)
+{
+	for (int i = 0; i < EMassAccessOperation::MAX; ++i)
+	{
+		FragmentsAccess[i].Access.AddZeroed(FMassFragmentBitSet::GetMaxNum());
+		ChunkFragmentsAccess[i].Access.AddZeroed(FMassChunkFragmentBitSet::GetMaxNum());
+		SharedFragmentsAccess[i].Access.AddZeroed(FMassSharedFragmentBitSet::GetMaxNum());
+		RequiredSubsystemsAccess[i].Access.AddZeroed(FMassExternalSubsystemBitSet::GetMaxNum());
+	}
+}
+
+template<typename TBitSet>
+void FMassProcessorDependencySolver::FResourceUsage::HandleElementType(TMassExecutionAccess<FResourceAccess>& ElementAccess
+	, const TMassExecutionAccess<TBitSet>& TestedRequirements, FMassProcessorDependencySolver::FNode& InOutNode, const int32 NodeIndex)
+{
+	using UE::Mass::Private::DoArchetypeContainersOverlap;
+
+	// for every bit set in TestedRequirements we do the following:
+	// 1. For every read-only requirement we make InOutNode depend on the currently stored Writer of this resource
+	//    - note that this operation is not destructive, meaning we don't destructively consume the data, since all 
+	//      subsequent read access to the given resource will also depend on the Writer
+	//    - note 2: we also fine tune what we store as a dependency for InOutNode by checking if InOutNode's archetype
+	//      overlap with whoever the current Writer is 
+	//    - this will result in InOutNode wait for the current Writer to finish before starting its own work and 
+	//      that's exactly what we need to do to avoid accessing data while it's potentially being written
+	// 2. For every read-write requirement we make InOutNode depend on all the readers and writers currently stored. 
+	//    - once that's done we clean currently stored Readers and Writers since every subsequent operation on this 
+	//      resource will be blocked by currently considered InOutNode (as the new Writer)
+	//    - again, we do check corresponding archetype collections overlap
+	//    - similarly to the read operation waiting on write operations in pt 1. we want to hold off the write 
+	//      operations to be performed by InOutNode until all currently registered (and conflicting) writers and readers 
+	//      are done with their operations 
+	// 3. For all accessed resources we store information that InOutNode is accessing it
+	//    - we do this so that the following nodes know that they'll have to wait for InOutNode if an access 
+	//      conflict arises. 
+
+	// 1. For every read only requirement we make InOutNode depend on the currently stored Writer of this resource
+	for (auto It = TestedRequirements.Read.GetIndexIterator(); It; ++It)
+	{
+		for (int32 UserIndex : ElementAccess.Write.Access[*It].Users)
+		{
+			if (DoArchetypeContainersOverlap(AllNodesView[UserIndex].ValidArchetypes, InOutNode.ValidArchetypes))
+			{
+				InOutNode.OriginalDependencies.Add(UserIndex);
+			}
+		}
+	}
+
+	// 2. For every read-write requirement we make InOutNode depend on all the readers and writers currently stored. 
+	for (auto It = TestedRequirements.Write.GetIndexIterator(); It; ++It)
+	{
+		for (int32 i = ElementAccess.Read.Access[*It].Users.Num() - 1; i >= 0; --i)
+		{
+			const int32 UserIndex = ElementAccess.Read.Access[*It].Users[i];
+			if (DoArchetypeContainersOverlap(AllNodesView[UserIndex].ValidArchetypes, InOutNode.ValidArchetypes))
+			{	
+				InOutNode.OriginalDependencies.Add(UserIndex);
+				ElementAccess.Read.Access[*It].Users.RemoveAtSwap(i);
+			}
+		}
+
+		for (int32 i = ElementAccess.Write.Access[*It].Users.Num() - 1; i >= 0; --i)
+		{
+			const int32 UserIndex = ElementAccess.Write.Access[*It].Users[i];
+			if (DoArchetypeContainersOverlap(AllNodesView[UserIndex].ValidArchetypes, InOutNode.ValidArchetypes))
+			{
+				InOutNode.OriginalDependencies.Add(UserIndex);
+				ElementAccess.Write.Access[*It].Users.RemoveAtSwap(i);
+			}
+		}
+	}
+
+	// 3. For all accessed resources we store information that InOutNode is accessing it
+	for (auto It = TestedRequirements.Read.GetIndexIterator(); It; ++It)
+	{
+		// mark Element at index indicated by It as being used in mode EMassAccessOperation(i) by NodeIndex
+		ElementAccess.Read.Access[*It].Users.Add(NodeIndex);
+	}
+	for (auto It = TestedRequirements.Write.GetIndexIterator(); It; ++It)
+	{
+		// mark Element at index indicated by It as being used in mode EMassAccessOperation(i) by NodeIndex
+		ElementAccess.Write.Access[*It].Users.Add(NodeIndex);
+	}
+}
+
+template<typename TBitSet>
+bool FMassProcessorDependencySolver::FResourceUsage::CanAccess(const TMassExecutionAccess<TBitSet>& StoredElements, const TMassExecutionAccess<TBitSet>& TestedElements)
+{
+	// see if there's an overlap of tested write operations with existing read & write operations, as well as 
+	// tested read operations with existing write operations
+	
+	return !(
+		// if someone's already writing to what I want to write
+		TestedElements.Write.HasAny(StoredElements.Write)
+		// or if someone's already reading what I want to write
+		|| TestedElements.Write.HasAny(StoredElements.Read)
+		// or if someone's already writing what I want to read
+		|| TestedElements.Read.HasAny(StoredElements.Write)
+	);
+}
+
+bool FMassProcessorDependencySolver::FResourceUsage::HasArchetypeConflict(TMassExecutionAccess<FResourceAccess> ElementAccess, const TArray<FMassArchetypeHandle>& InArchetypes) const
+{
+	using UE::Mass::Private::DoArchetypeContainersOverlap;
+
+	// this function is being run when we've already determined there's an access conflict on given ElementsAccess,
+	// meaning whoever's asking is trying to access Elements that are already being used. We can still grant access 
+	// though provided that none of the current users of Element access the same archetypes the querier does (as provided 
+	// by InArchetypes).
+	// @todo this operation could be even more efficient and precise if we tracked which operation (read/write) and which
+	// specific Element were conflicting and the we could limit the check to that. That would however significantly 
+	// complicate the code and would require a major refactor to keep things clean.
+	for (int i = 0; i < EMassAccessOperation::MAX; ++i)
+	{
+		for (const FResourceUsers& Resource : ElementAccess[i].Access)
+		{
+			for (const int32 UserIndex : Resource.Users)
+			{
+				if (DoArchetypeContainersOverlap(AllNodesView[UserIndex].ValidArchetypes, InArchetypes))
+				{
+					return true;
+				}
+			}
+		}
+	}
+	return false;
+}
+
+bool FMassProcessorDependencySolver::FResourceUsage::CanAccessRequirements(const FMassExecutionRequirements& TestedRequirements, const TArray<FMassArchetypeHandle>& InArchetypes) const
+{
+	// note that on purpose we're not checking ConstSharedFragments - those are always only read, no danger of conflicting access
+	bool bCanAccess = (CanAccess<FMassFragmentBitSet>(Requirements.Fragments, TestedRequirements.Fragments) || !HasArchetypeConflict(FragmentsAccess, InArchetypes))
+		&& (CanAccess<FMassChunkFragmentBitSet>(Requirements.ChunkFragments, TestedRequirements.ChunkFragments) || !HasArchetypeConflict(ChunkFragmentsAccess, InArchetypes))
+		&& (CanAccess<FMassSharedFragmentBitSet>(Requirements.SharedFragments, TestedRequirements.SharedFragments) || !HasArchetypeConflict(SharedFragmentsAccess, InArchetypes))
+		&& CanAccess<FMassExternalSubsystemBitSet>(Requirements.RequiredSubsystems, TestedRequirements.RequiredSubsystems);
+
+	return bCanAccess;
+}
+
+void FMassProcessorDependencySolver::FResourceUsage::SubmitNode(const int32 NodeIndex, FNode& InOutNode)
+{
+	HandleElementType<FMassFragmentBitSet>(FragmentsAccess, InOutNode.Requirements.Fragments, InOutNode, NodeIndex);
+	HandleElementType<FMassChunkFragmentBitSet>(ChunkFragmentsAccess, InOutNode.Requirements.ChunkFragments, InOutNode, NodeIndex);
+	HandleElementType<FMassSharedFragmentBitSet>(SharedFragmentsAccess, InOutNode.Requirements.SharedFragments, InOutNode, NodeIndex);
+	HandleElementType<FMassExternalSubsystemBitSet>(RequiredSubsystemsAccess, InOutNode.Requirements.RequiredSubsystems, InOutNode, NodeIndex);
+	// note that on purpose we're not pushing ConstSharedFragments - those are always only read, no danger of conflicting access
+	// so there's no point in tracking them
+
+	Requirements.Append(InOutNode.Requirements);
+}
+
+//----------------------------------------------------------------------//
+//  FProcessorDependencySolver::FNode
+//----------------------------------------------------------------------//
+void FMassProcessorDependencySolver::FNode::IncreaseWaitingNodesCount(TArrayView<FMassProcessorDependencySolver::FNode> InAllNodes)
+{
+	// cycle-protection check. If true it means we have a cycle and the whole algorithm result will be unreliable 
+	if (TotalWaitingNodes >= FMath::Square(InAllNodes.Num()))
+	{
+		return;
+	}
+
+	++TotalWaitingNodes;
+
+	for (const int32 DependencyIndex : OriginalDependencies)
+	{
+		check(&InAllNodes[DependencyIndex] != this);
+		InAllNodes[DependencyIndex].IncreaseWaitingNodesCount(InAllNodes);
+	}
+}
+//----------------------------------------------------------------------//
+//  FProcessorDependencySolver
+//----------------------------------------------------------------------//
+FMassProcessorDependencySolver::FMassProcessorDependencySolver(TArrayView<UMassProcessor* const> InProcessors, const bool bIsGameRuntime)
+	: Processors(InProcessors)
+	, bGameRuntime(bIsGameRuntime)
+{}
+
+bool FMassProcessorDependencySolver::PerformSolverStep(FResourceUsage& ResourceUsage, TArray<int32>& InOutIndicesRemaining, TArray<int32>& OutNodeIndices)
+{
+	int32 AcceptedNodeIndex = INDEX_NONE;
+	int32 FallbackAcceptedNodeIndex = INDEX_NONE;
+
+	for (int32 i = 0; i < InOutIndicesRemaining.Num(); ++i)
+	{
+		const int32 NodeIndex = InOutIndicesRemaining[i];
+		if (AllNodes[NodeIndex].TransientDependencies.Num() == 0)
+		{
+			// if we're solving dependencies for a single thread use we don't need to fine-tune the order based on resources nor archetypes
+			if (bSingleThreadTarget || ResourceUsage.CanAccessRequirements(AllNodes[NodeIndex].Requirements, AllNodes[NodeIndex].ValidArchetypes))
+			{
+				AcceptedNodeIndex = NodeIndex;
+				break;
+			}
+			else if (FallbackAcceptedNodeIndex == INDEX_NONE)
+			{
+				// if none of the nodes left can "cleanly" execute (i.e. without conflicting with already stored nodes)
+				// we'll just pick this one up and go with it. 
+				FallbackAcceptedNodeIndex = NodeIndex;
+			}
+		}
+	}
+
+	if (AcceptedNodeIndex != INDEX_NONE || FallbackAcceptedNodeIndex != INDEX_NONE)
+	{
+		const int32 NodeIndex = AcceptedNodeIndex != INDEX_NONE ? AcceptedNodeIndex : FallbackAcceptedNodeIndex;
+
+		FNode& Node = AllNodes[NodeIndex];
+
+		// Note that this is not an unexpected event and will happen during every dependency solving. It's a part 
+		// of the algorithm. We initially look for all the things we can run without conflicting with anything else. 
+		// But that can't last forever, at some point we'll end up in a situation where every node left waits for 
+		// something that has been submitted already. Then we just pick one of the waiting ones (the one indicated by 
+		// FallbackAcceptedNodeIndex), "run it" and proceed.
+		UE_CLOG(AcceptedNodeIndex == INDEX_NONE, LogMass, Verbose, TEXT("No dependency-free node can be picked, due to resource requirements. Picking %s as the next node.")
+			, *Node.Name.ToString());
+
+		ResourceUsage.SubmitNode(NodeIndex, Node);
+		InOutIndicesRemaining.RemoveSingle(NodeIndex);
+		OutNodeIndices.Add(NodeIndex);
+		for (const int32 DependencyIndex : Node.OriginalDependencies)
+		{
+			Node.SequencePositionIndex = FMath::Max(Node.SequencePositionIndex, AllNodes[DependencyIndex].SequencePositionIndex);
+		}
+		++Node.SequencePositionIndex;
+
+		for (const int32 RemainingNodeIndex : InOutIndicesRemaining)
+		{
+			AllNodes[RemainingNodeIndex].TransientDependencies.RemoveSingleSwap(NodeIndex, EAllowShrinking::No);
+		}
+		
+		return true;
+	}
+
+	return false;
+}
+
+void FMassProcessorDependencySolver::CreateSubGroupNames(FName InGroupName, TArray<FString>& SubGroupNames)
+{
+	// the function will convert composite group name into a series of progressively more precise group names
+	// so "A.B.C" will result in ["A", "A.B", "A.B.C"]
+
+	SubGroupNames.Reset();
+	FString GroupNameAsString = InGroupName.ToString();
+	FString TopGroupName;
+
+	while (GroupNameAsString.Split(TEXT("."), &TopGroupName, &GroupNameAsString))
+	{
+		SubGroupNames.Add(TopGroupName);
+	}
+	SubGroupNames.Add(GroupNameAsString);
+	
+	for (int i = 1; i < SubGroupNames.Num(); ++i)
+	{
+		SubGroupNames[i] = FString::Printf(TEXT("%s.%s"), *SubGroupNames[i - 1], *SubGroupNames[i]);
+	}
+}
+
+int32 FMassProcessorDependencySolver::CreateNodes(UMassProcessor& Processor)
+{
+	check(Processor.GetClass());
+	// for processors supporting multiple instances we use processor name rather than processor's class name for
+	// dependency calculations. This makes the user responsible for fine-tuning per-processor dependencies. 
+	const FName ProcName = Processor.ShouldAllowMultipleInstances() 
+		? Processor.GetFName()
+		: Processor.GetClass()->GetFName();
+
+	if (const int32* NodeIndexPtr = NodeIndexMap.Find(ProcName))
+	{
+		if (Processor.ShouldAllowMultipleInstances())
+		{
+			UE_LOG(LogMass, Warning, TEXT("%hs Processor %s, name %s, already registered. This processor class does suport duplicates, but individual instances need to have a unique name.")
+				, __FUNCTION__, *Processor.GetFullName(), *ProcName.ToString());
+		}
+		else
+		{
+			UE_LOG(LogMass, Warning, TEXT("%hs Processor %s already registered. Duplicates are not supported by this processor class.")
+				, __FUNCTION__, *ProcName.ToString());
+		}
+		return *NodeIndexPtr;
+	}
+
+	const FMassProcessorExecutionOrder& ExecutionOrder = Processor.GetExecutionOrder();
+
+	// first figure out the groups so that the group nodes come before the processor nodes, this is required for child
+	// nodes to inherit group's dependencies like in scenarios where some processor required to ExecuteBefore a given group
+	int32 ParentGroupNodeIndex = INDEX_NONE;
+	if (ExecutionOrder.ExecuteInGroup.IsNone() == false)
+	{
+		TArray<FString> AllGroupNames;
+		CreateSubGroupNames(ExecutionOrder.ExecuteInGroup, AllGroupNames);
+	
+		check(AllGroupNames.Num() > 0);
+
+		for (const FString& GroupName : AllGroupNames)
+		{
+			const FName GroupFName(GroupName);
+			int32* LocalGroupIndex = NodeIndexMap.Find(GroupFName);
+			// group name hasn't been encountered yet - create it
+			if (LocalGroupIndex == nullptr)
+			{
+				int32 NewGroupNodeIndex = AllNodes.Num();
+				NodeIndexMap.Add(GroupFName, NewGroupNodeIndex);
+				FNode& GroupNode = AllNodes.Add_GetRef({ GroupFName, nullptr, NewGroupNodeIndex });
+				// just ignore depending on the dummy "root" node
+				if (ParentGroupNodeIndex != INDEX_NONE)
+				{
+					GroupNode.OriginalDependencies.Add(ParentGroupNodeIndex);
+					AllNodes[ParentGroupNodeIndex].SubNodeIndices.Add(NewGroupNodeIndex);
+				}
+
+				ParentGroupNodeIndex = NewGroupNodeIndex;
+			}
+			else
+			{	
+				ParentGroupNodeIndex = *LocalGroupIndex;
+			}
+
+		}
+	}
+
+	const int32 NodeIndex = AllNodes.Num();
+	NodeIndexMap.Add(ProcName, NodeIndex);
+	FNode& ProcessorNode = AllNodes.Add_GetRef({ ProcName, &Processor, NodeIndex });
+
+	ProcessorNode.ExecuteAfter.Append(ExecutionOrder.ExecuteAfter);
+	ProcessorNode.ExecuteBefore.Append(ExecutionOrder.ExecuteBefore);
+	Processor.ExportRequirements(ProcessorNode.Requirements);
+	ProcessorNode.Requirements.CountResourcesUsed();
+
+	if (ParentGroupNodeIndex > 0)
+	{
+		AllNodes[ParentGroupNodeIndex].SubNodeIndices.Add(NodeIndex);
+	}
+
+	return NodeIndex;
+}
+
+void FMassProcessorDependencySolver::BuildDependencies()
+{
+	// at this point we have collected all the known processors and groups in AllNodes so we can transpose 
+	// A.ExecuteBefore(B) type of dependencies into B.ExecuteAfter(A)
+	for (int32 NodeIndex = 0; NodeIndex < AllNodes.Num(); ++NodeIndex)
+	{
+		for (int i = 0; i < AllNodes[NodeIndex].ExecuteBefore.Num(); ++i)
+		{
+			const FName BeforeDependencyName = AllNodes[NodeIndex].ExecuteBefore[i];
+			int32 DependentNodeIndex = INDEX_NONE;
+			int32* DependentNodeIndexPtr = NodeIndexMap.Find(BeforeDependencyName);
+			if (DependentNodeIndexPtr == nullptr)
+			{
+				// missing dependency. Adding a "dummy" node representing those to still support ordering based on missing groups or processors 
+				// For example, if Processor A and B declare dependency, respectively, "Before C" and "After C" we still 
+				// expect A to come before B regardless of whether C exists or not.
+				
+				DependentNodeIndex = AllNodes.Num();
+				NodeIndexMap.Add(BeforeDependencyName, DependentNodeIndex);
+				AllNodes.Add({ BeforeDependencyName, nullptr, DependentNodeIndex });
+
+				UE_LOG(LogMass, Log, TEXT("Unable to find dependency \"%s\" declared by %s. Creating a dummy dependency node.")
+					, *BeforeDependencyName.ToString(), *AllNodes[NodeIndex].Name.ToString());
+			}
+			else
+			{
+				DependentNodeIndex = *DependentNodeIndexPtr;
+			}
+
+			check(AllNodes.IsValidIndex(DependentNodeIndex));
+			AllNodes[DependentNodeIndex].ExecuteAfter.Add(AllNodes[NodeIndex].Name);
+		}
+		AllNodes[NodeIndex].ExecuteBefore.Reset();
+	}
+
+	// at this point all nodes contain:
+	// - single "original dependency" pointing at its parent group
+	// - ExecuteAfter populated with node names
+
+	// Now, for every Name in ExecuteAfter we do the following:
+	//	if Name represents a processor, add it as "original dependency"
+	//	else, if Name represents a group:
+	//		- append all group's child node names to ExecuteAfter
+	// 
+	for (int32 NodeIndex = 0; NodeIndex < AllNodes.Num(); ++NodeIndex)
+	{
+		for (int i = 0; i < AllNodes[NodeIndex].ExecuteAfter.Num(); ++i)
+		{
+			const FName AfterDependencyName = AllNodes[NodeIndex].ExecuteAfter[i];
+			int32* PrerequisiteNodeIndexPtr = NodeIndexMap.Find(AfterDependencyName);
+			int32 PrerequisiteNodeIndex = INDEX_NONE;
+
+			if (PrerequisiteNodeIndexPtr == nullptr)
+			{
+				// missing dependency. Adding a "dummy" node representing those to still support ordering based on missing groups or processors 
+				// For example, if Processor A and B declare dependency, respectively, "Before C" and "After C" we still 
+				// expect A to come before B regardless of whether C exists or not.
+
+				PrerequisiteNodeIndex = AllNodes.Num();
+				NodeIndexMap.Add(AfterDependencyName, PrerequisiteNodeIndex);
+				AllNodes.Add({ AfterDependencyName, nullptr, PrerequisiteNodeIndex });
+
+				UE_LOG(LogMass, Log, TEXT("Unable to find dependency \"%s\" declared by %s. Creating a dummy dependency node.")
+					, *AfterDependencyName.ToString(), *AllNodes[NodeIndex].Name.ToString());
+			}
+			else
+			{
+				PrerequisiteNodeIndex = *PrerequisiteNodeIndexPtr;
+			}
+
+			const FNode& PrerequisiteNode = AllNodes[PrerequisiteNodeIndex];
+
+			if (PrerequisiteNode.IsGroup())
+			{
+				for (int32 SubNodeIndex : PrerequisiteNode.SubNodeIndices)
+				{
+					AllNodes[NodeIndex].ExecuteAfter.AddUnique(AllNodes[SubNodeIndex].Name);
+				}
+			}
+			else
+			{
+				AllNodes[NodeIndex].OriginalDependencies.AddUnique(PrerequisiteNodeIndex);
+			}
+		}
+
+		// if this node is a group push all the dependencies down on all the children
+		// by design all child nodes come after group nodes so the child nodes' dependencies have not been processed yet
+		if (AllNodes[NodeIndex].IsGroup() && AllNodes[NodeIndex].SubNodeIndices.Num())
+		{
+			for (int32 PrerequisiteNodeIndex : AllNodes[NodeIndex].OriginalDependencies)
+			{
+				checkSlow(PrerequisiteNodeIndex != NodeIndex);
+				// in case of processor nodes we can store it directly
+				if (AllNodes[PrerequisiteNodeIndex].IsGroup() == false)
+				{
+					for (int32 ChildNodeIndex : AllNodes[NodeIndex].SubNodeIndices)
+					{
+						AllNodes[ChildNodeIndex].OriginalDependencies.AddUnique(PrerequisiteNodeIndex);
+					}
+				}
+				// special case - if dependency is a group and we haven't processed that group yet, we need to add it by name
+				else if (PrerequisiteNodeIndex > NodeIndex)
+				{
+					const FName& PrerequisiteName = AllNodes[PrerequisiteNodeIndex].Name;
+					for (int32 ChildNodeIndex : AllNodes[NodeIndex].SubNodeIndices)
+					{
+						AllNodes[ChildNodeIndex].ExecuteAfter.AddUnique(PrerequisiteName);
+					}
+				}
+			}
+		}
+	}
+}
+
+void FMassProcessorDependencySolver::LogNode(const FNode& Node, int Indent)
+{
+	using UE::Mass::Private::NameViewToString;
+
+	if (Node.IsGroup())
+	{
+		UE_LOG(LogMass, Log, TEXT("%*s%s before:%s after:%s"), Indent, TEXT(""), *Node.Name.ToString()
+			, *NameViewToString(Node.ExecuteBefore)
+			, *NameViewToString(Node.ExecuteAfter));
+
+		for (const int32 NodeIndex : Node.SubNodeIndices)
+		{
+			LogNode(AllNodes[NodeIndex], Indent + 4);
+		}
+	}
+	else
+	{
+		CA_ASSUME(Node.Processor); // as implied by Node.IsGroup() == false
+		UE_LOG(LogMass, Log, TEXT("%*s%s before:%s after:%s"), Indent, TEXT(""), *Node.Name.ToString()
+			, *NameViewToString(Node.Processor->GetExecutionOrder().ExecuteBefore)
+			, *NameViewToString(Node.Processor->GetExecutionOrder().ExecuteAfter));
+	}
+}
+
+#if 0 // disabled the graph building for now, leaving the old code here for reference
+struct FDumpGraphDependencyUtils
+{
+	static void DumpGraphNode(FArchive& LogFile, const FProcessorDependencySolver::FNode& Node, int Indent, TSet<const FProcessorDependencySolver::FNode*>& AllNodes, const bool bRoot)
+	{
+		const FString NodeName = Node.Name.ToString();
+		if (!Node.Processor)
+		{
+			const FString ClusterNodeName = NodeName.Replace(TEXT("."), TEXT("_"));
+			const FString GraphNodeName = NodeName.Replace(TEXT("."), TEXT(" "));
+
+			int32 Index = -1;
+			Node.Name.ToString().FindLastChar(TEXT('.'), Index);
+			const FString GroupName = NodeName.Mid(Index+1);
+
+			LogFile.Logf(TEXT("%*ssubgraph cluster_%s"), Indent, TEXT(""), *ClusterNodeName);
+			LogFile.Logf(TEXT("%*s{"), Indent, TEXT(""));
+			LogFile.Logf(TEXT("%*slabel =\"%s\";"), Indent + 4, TEXT(""), *GroupName);
+			LogFile.Logf(TEXT("%*s\"%s Start\"%s;"), Indent + 4, TEXT(""), *GraphNodeName, bRoot ? TEXT("") : TEXT("[shape=point style=invis]"));
+			LogFile.Logf(TEXT("%*s\"%s End\"%s;"), Indent + 4, TEXT(""), *GraphNodeName, bRoot ? TEXT("") : TEXT("[shape=point style=invis]"));
+			for (const FProcessorDependencySolver::FNode& SubNode : Node.SubNodes)
+			{
+				DumpGraphNode(LogFile, SubNode, Indent + 4, AllNodes, false);
+			}
+			LogFile.Logf(TEXT("%*s}"), Indent, TEXT(""));
+		}
+		else
+		{
+			LogFile.Logf(TEXT("%*s\"%s\""), Indent, TEXT(""), *NodeName);
+			AllNodes.Add(&Node);
+		}
+	}
+
+	static const FProcessorDependencySolver::FNode* FindDepNodeInParents(const TArray<const FProcessorDependencySolver::FNode*>& Parents, FName DependencyName)
+	{
+		for (int32 i = Parents.Num() - 1; i >= 0; --i)
+		{
+			const FProcessorDependencySolver::FNode* CurNode = Parents[i];
+			int32 DependencyIndex = CurNode->FindNodeIndex(DependencyName);
+			while(DependencyIndex != INDEX_NONE)
+			{
+				if(CurNode->SubNodes[DependencyIndex].Name == DependencyName)
+				{
+					return &CurNode->SubNodes[DependencyIndex];
+				}
+				else
+				{
+					// Dig down the chain
+					CurNode = &CurNode->SubNodes[DependencyIndex];
+					DependencyIndex = CurNode->FindNodeIndex(DependencyName);
+				}
+			}
+		}
+		return nullptr;
+	};
+
+	static bool DoAllSubNodeHasDependency(const FProcessorDependencySolver::FNode& Node, const FName DependencyName, bool bBeforeDep)
+	{
+		bool bDepExistInAllSibling = true;
+		for (const FProcessorDependencySolver::FNode& SiblingSubNode : Node.SubNodes)
+		{
+			if (SiblingSubNode.Processor)
+			{
+				const TArray<FName>& ExecuteDep = bBeforeDep ? SiblingSubNode.Processor->GetExecutionOrder().ExecuteBefore : SiblingSubNode.Processor->GetExecutionOrder().ExecuteAfter;
+				if (ExecuteDep.Find(DependencyName) == INDEX_NONE)
+				{
+					bDepExistInAllSibling = false;
+					break;
+				}
+			}
+			else if (!DoAllSubNodeHasDependency(SiblingSubNode, DependencyName, bBeforeDep))
+			{
+				bDepExistInAllSibling = false;
+				break;
+			}
+		}
+		return bDepExistInAllSibling;
+	}
+
+	static bool DoesDependencyExistIndirectly(const TArray<const FProcessorDependencySolver::FNode*>& Parents, const FProcessorDependencySolver::FNode& Node, const FName DepNameToFind, bool bBeforeDep)
+	{
+		check(Node.Processor);
+		const TArray<FName>& ExecuteDep = bBeforeDep ? Node.Processor->GetExecutionOrder().ExecuteBefore : Node.Processor->GetExecutionOrder().ExecuteAfter;
+		for (const FName& DependencyName : ExecuteDep)
+		{
+			if (DependencyName == DepNameToFind)
+			{
+				continue;
+			}
+			if (const FProcessorDependencySolver::FNode* DepNode = FindDepNodeInParents(Parents, DependencyName))
+			{
+				if(DepNode->Processor)
+				{
+					const TArray<FName>& ExecuteDepDep = bBeforeDep ? DepNode->Processor->GetExecutionOrder().ExecuteBefore : DepNode->Processor->GetExecutionOrder().ExecuteAfter;
+					if(ExecuteDepDep.Find(DepNameToFind) != INDEX_NONE)
+					{
+						return true;
+					}
+				}
+				else if(DoAllSubNodeHasDependency(*DepNode, DepNameToFind, bBeforeDep))
+				{
+					return true;
+				}
+			}
+		}
+		return false;
+	}
+
+	static void RemoveAllProcessorFromSet(const FProcessorDependencySolver::FNode& DepNode, TSet<const FProcessorDependencySolver::FNode*>& Set)
+	{
+		if (DepNode.Processor)
+		{
+			Set.Remove(&DepNode);
+		}
+		else
+		{
+			for (const FProcessorDependencySolver::FNode& SubNode : DepNode.SubNodes)
+			{
+				RemoveAllProcessorFromSet(SubNode, Set);
+			}
+		}
+	}
+
+	static bool AreAllProcessorInSet(const FProcessorDependencySolver::FNode& DepNode, TSet<const FProcessorDependencySolver::FNode*>& Set)
+	{
+		if (DepNode.Processor)
+		{
+			if (!Set.Contains(&DepNode))
+			{
+				return false;
+			}
+		}
+		else
+		{
+			for (const FProcessorDependencySolver::FNode& SubNode : DepNode.SubNodes)
+			{
+				if (!AreAllProcessorInSet(SubNode, Set))
+				{
+					return false;
+				}
+			}
+		}
+		return true;
+	}
+
+	static void DumpGraphDependencies(FArchive& LogFile, const FProcessorDependencySolver::FNode& Node, TArray<const FProcessorDependencySolver::FNode*> Parents, TArray<const FProcessorDependencySolver::FNode*> CommonBeforeDep, TArray<const FProcessorDependencySolver::FNode*> CommonAfterDep, TSet<const FProcessorDependencySolver::FNode*>& DependsOnStart, TSet<const FProcessorDependencySolver::FNode*>& LinkToEnd)
+	{
+		const FString NodeName = Node.Name.ToString();
+		const FProcessorDependencySolver::FNode* ParentNode = Parents.Num() > 0 ? Parents.Last() : nullptr;
+		const FString ParentNodeName = ParentNode ? ParentNode->Name.ToString() : TEXT("");
+		const FString ParentGraphNodeName = ParentNodeName.Replace(TEXT("."), TEXT(" "));
+		const FString ParentClusterNodeName = ParentNodeName.Replace(TEXT("."), TEXT("_"));
+		if (Node.Processor)
+		{
+			for (const FName& DependencyName : Node.Processor->GetExecutionOrder().ExecuteBefore)
+			{
+				if (const FProcessorDependencySolver::FNode* DepNode = FindDepNodeInParents(Parents, DependencyName))
+				{
+					LinkToEnd.Remove(&Node);
+					RemoveAllProcessorFromSet(*DepNode, DependsOnStart);
+
+					if (CommonBeforeDep.Find(DepNode) != INDEX_NONE)
+					{
+						continue;
+					}
+					if (DoesDependencyExistIndirectly(Parents, Node, DependencyName, true))
+					{
+						continue;
+					}
+
+					const FString DepNodeName = DepNode->Name.ToString();
+					if (DepNode->Processor)
+					{
+						LogFile.Logf(TEXT("    \"%s\" -> \"%s\";"), *NodeName, *DepNodeName);
+					}
+					else
+					{
+						const FString DepClusterNodeName = DepNodeName.Replace(TEXT("."), TEXT("_"));
+						const FString DepGraphNodeName = DepNodeName.Replace(TEXT("."), TEXT(" "));
+						LogFile.Logf(TEXT("    \"%s\" -> \"%s Start\"[lhead=cluster_%s];"), *NodeName, *DepGraphNodeName, *DepClusterNodeName);
+					}
+
+				}
+			}
+			for (const FName& DependencyName : Node.Processor->GetExecutionOrder().ExecuteAfter)
+			{
+				if (const FProcessorDependencySolver::FNode* DepNode = FindDepNodeInParents(Parents, DependencyName))
+				{
+					DependsOnStart.Remove(&Node);
+					RemoveAllProcessorFromSet(*DepNode, LinkToEnd);
+
+					if (CommonAfterDep.Find(DepNode) != INDEX_NONE)
+					{
+						continue;
+					}
+					if (DoesDependencyExistIndirectly(Parents, Node, DependencyName, false))
+					{
+						continue;
+					}
+
+					const FString DepNodeName = DepNode->Name.ToString();
+					if (DepNode->Processor)
+					{
+						LogFile.Logf(TEXT("    \"%s\" -> \"%s\";"), *DepNodeName, *NodeName);
+					}
+					else
+					{
+						const FString DepClusterNodeName = DepNodeName.Replace(TEXT("."), TEXT("_"));
+						const FString DepGraphNodeName = DepNodeName.Replace(TEXT("."), TEXT(" "));
+						LogFile.Logf(TEXT("    \"%s End\" -> \"%s\"[ltail=cluster_%s];"), *DepGraphNodeName, *NodeName, *DepClusterNodeName);
+					}
+				}
+			}
+
+			// Layouting
+			LogFile.Logf(TEXT("    \"%s Start\" -> \"%s\" -> \"%s End\"[style=invis];"), *ParentGraphNodeName, *NodeName, *ParentGraphNodeName);
+		}
+		else if (ParentNode)
+		{
+
+			const FString ClusterNodeName = NodeName.Replace(TEXT("."), TEXT("_"));
+			const FString GraphNodeName = NodeName.Replace(TEXT("."), TEXT(" "));
+
+			// Find common dependency through out all sub nodes
+ 			TSet<FName> BeforeDepOuputed;
+ 			TSet<FName> AfterDepOutputed;
+			for (const FProcessorDependencySolver::FNode& SubNode : Node.SubNodes)
+			{
+				if(SubNode.Processor)
+				{
+					for (const FName& DependencyName : SubNode.Processor->GetExecutionOrder().ExecuteBefore)
+					{
+						if (const FProcessorDependencySolver::FNode* DepNode = FindDepNodeInParents(Parents, DependencyName))
+						{
+							if (DoAllSubNodeHasDependency(Node, DependencyName, true /*bBeforeDep*/))
+							{
+								CommonBeforeDep.Add(DepNode);
+
+								if (!DoesDependencyExistIndirectly(Parents, SubNode, DependencyName, true))
+								{
+									if (BeforeDepOuputed.Contains(DependencyName))
+									{
+										continue;
+									}
+									BeforeDepOuputed.Add(DependencyName);
+
+									const FString DepNodeName = DepNode->Name.ToString();
+									if (DepNode->Processor)
+									{
+										LogFile.Logf(TEXT("    \"%s End\" -> \"%s\"[ltail=cluster_%s];"), *GraphNodeName, *DepNodeName, *ClusterNodeName);
+									}
+									else
+									{
+										const FString DepClusterNodeName = DepNodeName.Replace(TEXT("."), TEXT("_"));
+										const FString DepGraphNodeName = DepNodeName.Replace(TEXT("."), TEXT(" "));
+										LogFile.Logf(TEXT("    \"%s End\" -> \"%s Start\"[ltail=cluster_%s, lhead=cluster_%s];"), *GraphNodeName, *DepGraphNodeName, *ClusterNodeName, *DepClusterNodeName);
+									}
+								}
+							}
+						}
+					}
+
+					for (const FName& DependencyName : SubNode.Processor->GetExecutionOrder().ExecuteAfter)
+					{
+						if (const FProcessorDependencySolver::FNode* DepNode = FindDepNodeInParents(Parents, DependencyName))
+						{
+							if (DoAllSubNodeHasDependency(Node, DependencyName, false /*bBeforeDep*/))
+							{
+								CommonAfterDep.Add(DepNode);
+
+								if (!DoesDependencyExistIndirectly(Parents, SubNode, DependencyName, false))
+								{
+									if (AfterDepOutputed.Contains(DependencyName))
+									{
+										continue;
+									}
+									AfterDepOutputed.Add(DependencyName);
+
+									const FString DepNodeName = DepNode->Name.ToString();
+									if (DepNode->Processor)
+									{
+										LogFile.Logf(TEXT("    \"%s\" -> \"%s Start\"[lhead=cluster_%s];"), *DepNodeName, *GraphNodeName, *ClusterNodeName);
+									}
+									else
+									{
+										const FString DepClusterNodeName = DepNodeName.Replace(TEXT("."), TEXT("_"));
+										const FString DepGraphNodeName = DepNodeName.Replace(TEXT("."), TEXT(" "));
+										LogFile.Logf(TEXT("    \"%s End\" -> \"%s Start\"[ltail=cluster_%s, lhead=cluster_%s];"), *DepGraphNodeName,*GraphNodeName, *DepClusterNodeName, *ClusterNodeName);
+									}
+								}
+							}
+						}
+					}
+				}
+			}
+
+ 			// Layouting
+			LogFile.Logf(TEXT("    \"%s Start\" -> \"%s Start\" -> \"%s End\" -> \"%s End\"[style=invis];"), *ParentGraphNodeName, *GraphNodeName, *GraphNodeName, *ParentGraphNodeName);
+		}
+
+		Parents.Push(&Node);
+		for (const FProcessorDependencySolver::FNode& SubNode : Node.SubNodes)
+		{
+			DumpGraphDependencies(LogFile, SubNode, Parents, CommonBeforeDep, CommonAfterDep, DependsOnStart, LinkToEnd);
+		}
+	}
+
+	static void PromoteStartAndEndDependency(FArchive& LogFile, const FString& GraphName, const FProcessorDependencySolver::FNode& Node, TSet<const FProcessorDependencySolver::FNode*>& DependsOnStart, TSet<const FProcessorDependencySolver::FNode*>& LinkToEnd)
+	{
+		if (Node.Processor)
+		{
+			return;
+		}
+
+		bool bAllDependsOnStart = true;
+		bool bAllLinkedToEnd = true;
+		for (const FProcessorDependencySolver::FNode& SubNode : Node.SubNodes)
+		{
+			if (!AreAllProcessorInSet(SubNode, DependsOnStart))
+			{
+				bAllDependsOnStart = false;
+			}
+
+			if (!AreAllProcessorInSet(SubNode, LinkToEnd))
+			{
+				bAllLinkedToEnd = false;
+			}
+
+			if (!bAllDependsOnStart && !bAllLinkedToEnd)
+			{
+				break;
+			}
+		}
+
+		const FString NodeName = Node.Name.ToString();
+		const FString ClusterNodeName = NodeName.Replace(TEXT("."), TEXT("_"));
+		const FString GraphNodeName = NodeName.Replace(TEXT("."), TEXT(" "));
+		if (bAllDependsOnStart)
+		{
+			RemoveAllProcessorFromSet(Node, DependsOnStart);
+			LogFile.Logf(TEXT("    \"%s Start\" -> \"%s Start\"[lhead=cluster_%s, weight=0];"), *GraphName, *GraphNodeName, *ClusterNodeName);
+		}
+
+		if (bAllLinkedToEnd)
+		{
+			RemoveAllProcessorFromSet(Node, LinkToEnd);
+			LogFile.Logf(TEXT("    \"%s End\" -> \"%s End\"[ltail=cluster_%s, weight=0];"), *GraphNodeName, *GraphName, *ClusterNodeName);
+		}
+
+		if(!bAllDependsOnStart || !bAllLinkedToEnd)
+		{
+			for (const FProcessorDependencySolver::FNode& SubNode : Node.SubNodes)
+			{
+				PromoteStartAndEndDependency(LogFile, GraphName, SubNode, DependsOnStart, LinkToEnd);
+			}
+		}
+	}
+
+	static void DumpAllGraphDependencies(FArchive& LogFile, const FProcessorDependencySolver::FNode& GroupRootNode, const TSet<const FProcessorDependencySolver::FNode*>& AllNodes)
+	{
+		TSet<const FProcessorDependencySolver::FNode*> DependsOnStart = AllNodes;
+		TSet<const FProcessorDependencySolver::FNode*> LinkToEnd = AllNodes;
+		const FString GraphName = GroupRootNode.Name.ToString();
+		FDumpGraphDependencyUtils::DumpGraphDependencies(LogFile, GroupRootNode, TArray<const FProcessorDependencySolver::FNode*>(), TArray<const FProcessorDependencySolver::FNode*>(), TArray<const FProcessorDependencySolver::FNode*>(), DependsOnStart, LinkToEnd);
+		FDumpGraphDependencyUtils::PromoteStartAndEndDependency(LogFile, GraphName, GroupRootNode, DependsOnStart, LinkToEnd);
+		for (const FProcessorDependencySolver::FNode* Node : DependsOnStart)
+		{
+			LogFile.Logf(TEXT("    \"%s Start\" -> \"%s\"[weight=0];"), *GraphName, *Node->Name.ToString());
+		}
+		for (const FProcessorDependencySolver::FNode* Node : LinkToEnd)
+		{
+			LogFile.Logf(TEXT("    \"%s\" -> \"%s End\"[weight=0];"), *Node->Name.ToString(), *GraphName);
+		}
+	}
+};
+
+void FProcessorDependencySolver::DumpGraph(FArchive& LogFile) const
+{
+	TSet<const FNode*> AllNodes;
+	LogFile.Logf(TEXT("digraph MassProcessorGraph"));
+	LogFile.Logf(TEXT("{"));
+	LogFile.Logf(TEXT("    compound = true;"));
+	LogFile.Logf(TEXT("    newrank = true;"));
+	FDumpGraphDependencyUtils::DumpGraphNode(LogFile, GroupRootNode, 4/* Indent */, AllNodes, true/* bRoot */);
+	FDumpGraphDependencyUtils::DumpAllGraphDependencies(LogFile, GroupRootNode, AllNodes);
+	LogFile.Logf(TEXT("}"));
+}
+
+#endif // 0; disabled the graph building for now, leaving the old code here for reference
+
+void FMassProcessorDependencySolver::Solve(TArray<FMassProcessorOrderInfo>& OutResult)
+{
+	using UE::Mass::Private::NameViewToString;
+
+	if (AllNodes.Num() == 0)
+	{
+		return;
+	}
+
+	for (FNode& Node : AllNodes)
+	{
+		Node.TransientDependencies = Node.OriginalDependencies;
+		Node.TotalWaitingNodes = 0;
+	}
+
+	TArray<int32> IndicesRemaining;
+	IndicesRemaining.Reserve(AllNodes.Num());
+	for (int32 i = 0; i < AllNodes.Num(); ++i)
+	{
+		// skip all the group nodes, all group dependencies have already been converted to individual processor dependencies
+		if (AllNodes[i].IsGroup() == false)
+		{
+			IndicesRemaining.Add(i);
+			AllNodes[i].IncreaseWaitingNodesCount(AllNodes);
+		}
+	}
+
+	IndicesRemaining.Sort([this](const int32 IndexA, const int32 IndexB){
+		return AllNodes[IndexA].TotalWaitingNodes > AllNodes[IndexB].TotalWaitingNodes;
+	});
+
+	// this is where we'll be tracking what's being accessed by whom
+	FResourceUsage ResourceUsage(AllNodes);
+
+	TArray<int32> SortedNodeIndices;
+	SortedNodeIndices.Reserve(AllNodes.Num());
+
+	while (IndicesRemaining.Num())
+	{
+		const bool bStepSuccessful = PerformSolverStep(ResourceUsage, IndicesRemaining, SortedNodeIndices);
+
+		if (bStepSuccessful == false)
+		{
+			bAnyCyclesDetected = true;
+
+			UE_LOG(LogMass, Error, TEXT("Detected processing dependency cycle:"));
+			for (const int32 Index : IndicesRemaining)
+			{
+				UMassProcessor* Processor = AllNodes[Index].Processor;
+				if (Processor)
+				{
+					UE_LOG(LogMass, Warning, TEXT("\t%s, group: %s, before: %s, after %s")
+						, *Processor->GetName()
+						, *Processor->GetExecutionOrder().ExecuteInGroup.ToString()
+						, *NameViewToString(Processor->GetExecutionOrder().ExecuteBefore)
+						, *NameViewToString(Processor->GetExecutionOrder().ExecuteAfter));
+				}
+				else
+				{
+					// group
+					UE_LOG(LogMass, Warning, TEXT("\tGroup %s"), *AllNodes[Index].Name.ToString());
+				}
+			}
+			UE_LOG(LogMass, Warning, TEXT("Cutting the chain at an arbitrary location."));
+
+			// remove first dependency
+			// note that if we're in a cycle handling scenario every node does have some dependencies left
+			const int32 DependencyNodeIndex = AllNodes[IndicesRemaining[0]].TransientDependencies.Pop(EAllowShrinking::No);
+			// we need to remove this dependency from original dependencies as well, otherwise we'll still have the cycle
+			// in the data being produces as a result of the whole algorithm
+			AllNodes[IndicesRemaining[0]].OriginalDependencies.Remove(DependencyNodeIndex);
+		}
+	}
+
+	// now we have the desired order in SortedNodeIndices. We have to traverse it to add to OutResult
+	for (int i = 0; i < SortedNodeIndices.Num(); ++i)
+	{
+		const int32 NodeIndex = SortedNodeIndices[i];
+
+		TArray<FName> DependencyNames;
+		for (const int32 DependencyIndex : AllNodes[NodeIndex].OriginalDependencies)
+		{
+			DependencyNames.AddUnique(AllNodes[DependencyIndex].Name);
+		}
+
+		// at this point we expect SortedNodeIndices to only point to regular processors (i.e. no groups)
+		if (ensure(AllNodes[NodeIndex].Processor != nullptr))
+		{
+			OutResult.Add({ AllNodes[NodeIndex].Name, AllNodes[NodeIndex].Processor, FMassProcessorOrderInfo::EDependencyNodeType::Processor, DependencyNames, AllNodes[NodeIndex].SequencePositionIndex });
+		}
+	}
+}
+
+void FMassProcessorDependencySolver::ResolveDependencies(TArray<FMassProcessorOrderInfo>& OutResult, TSharedPtr<FMassEntityManager> EntityManager, FMassProcessorDependencySolver::FResult* InOutOptionalResult)
+{
+	TRACE_CPUPROFILER_EVENT_SCOPE_STR("Mass ResolveDependencies");
+
+	if (Processors.Num() == 0)
+	{
+		return;
+	}
+
+	FScopedCategoryAndVerbosityOverride LogOverride(TEXT("LogMass"), ELogVerbosity::Log);
+
+	if (InOutOptionalResult)
+	{
+		DependencyGraphFileName = InOutOptionalResult->DependencyGraphFileName;
+	}
+
+	bAnyCyclesDetected = false;
+
+	UE_LOG(LogMass, Log, TEXT("Gathering dependencies data:"));
+
+	AllNodes.Reset();
+	NodeIndexMap.Reset();
+	// as the very first node we add a "root" node that represents the "top level group" and also simplifies the rest
+	// of the lookup code - if a processor declares it's in group None or depends on Node it we don't need to check that 
+	// explicitly. 
+	AllNodes.Add(FNode(FName(), nullptr, 0));
+	NodeIndexMap.Add(FName(), 0);
+
+	const bool bCreateVirtualArchetypes = (!EntityManager);
+	if (bCreateVirtualArchetypes)
+	{
+		// create FMassEntityManager instance that we'll use to sort out processors' overlaps
+		// the idea for this is that for every processor we have we create an archetype matching given processor's requirements. 
+		// Once that's done we have a collection of "virtual" archetypes our processors expect. Then we ask every processor 
+		// to cache the archetypes they'd accept, using processors' owned queries. The idea is that some of the nodes will 
+		// end up with more than just the virtual archetype created for that specific node. The practice proved the idea correct. 
+		EntityManager = MakeShareable(new FMassEntityManager());
+	}
+
+	// gather the processors information first
+	for (UMassProcessor* Processor : Processors)
+	{
+		if (Processor == nullptr)
+		{
+			UE_LOG(LogMass, Warning, TEXT("%s nullptr found in Processors collection being processed"), ANSI_TO_TCHAR(__FUNCTION__));
+			continue;
+		}
+
+		const int32 ProcessorNodeIndex = CreateNodes(*Processor);
+
+		if (bCreateVirtualArchetypes)
+		{
+			// this line is a part of a nice trick we're doing here utilizing EntityManager's archetype creation based on 
+			// what each processor expects, and EntityQuery's capability to cache archetypes matching its requirements (used below)
+			EntityManager->CreateArchetype(AllNodes[ProcessorNodeIndex].Requirements.AsCompositionDescriptor());
+		}
+	}
+
+	UE_LOG(LogMass, Verbose, TEXT("Pruning processors..."));
+
+	int32 PrunedProcessorsCount = 0;
+	for (FNode& Node : AllNodes)
+	{
+		if (Node.IsGroup() == false)
+		{
+			CA_ASSUME(Node.Processor); // as implied by Node.IsGroup() == false
+
+			// for each processor-representing node we cache information on which archetypes among the once we've created 
+			// above (see the EntityManager.CreateArchetype call in the previous loop) match this processor. 
+			Node.Processor->GetArchetypesMatchingOwnedQueries(*EntityManager.Get(), Node.ValidArchetypes);
+
+			// prune the archetype-less processors
+			if (Node.ValidArchetypes.Num() == 0 && Node.Processor->ShouldAllowQueryBasedPruning(bGameRuntime))
+			{
+				UE_LOG(LogMass, Verbose, TEXT("\t%s"), *Node.Processor->GetName());
+
+				if (InOutOptionalResult)
+				{
+					InOutOptionalResult->PrunedProcessorClasses.Add(Node.Processor->GetClass());
+				}
+
+				// clearing out the processor will result in the rest of the algorithm to treat it as a group - we still 
+				// want to preserve the configured ExecuteBefore and ExecuteAfter dependencies
+				Node.Processor = nullptr;
+				++PrunedProcessorsCount;
+			}
+		}
+	}
+
+	UE_LOG(LogMass, Verbose, TEXT("Number of processors pruned: %d"), PrunedProcessorsCount);
+
+	check(AllNodes.Num());
+	LogNode(AllNodes[0]);
+
+	BuildDependencies();
+
+	// now none of the processor nodes depend on groups - we replaced these dependencies with depending directly 
+	// on individual processors. However, we keep the group nodes around since we store the dependencies via index, so 
+	// removing nodes would mess that up. Solve below ignores group nodes and OutResult will not have any groups once its done.
+
+	Solve(OutResult);
+
+	UE_LOG(LogMass, Verbose, TEXT("Dependency order:"));
+	for (const FMassProcessorOrderInfo& Info : OutResult)
+	{
+		UE_LOG(LogMass, Verbose, TEXT("\t%s"), *Info.Name.ToString());
+	}
+
+	int32 MaxSequenceLength = 0;
+	for (FNode& Node : AllNodes)
+	{
+		MaxSequenceLength = FMath::Max(MaxSequenceLength, Node.SequencePositionIndex);
+	}
+
+	UE_LOG(LogMass, Verbose, TEXT("Max sequence length: %d"), MaxSequenceLength);
+
+	if (InOutOptionalResult)
+	{
+		InOutOptionalResult->MaxSequenceLength = MaxSequenceLength;
+		InOutOptionalResult->ArchetypeDataVersion = EntityManager->GetArchetypeDataVersion();
+	}
+}
+
+bool FMassProcessorDependencySolver::IsResultUpToDate(const FMassProcessorDependencySolver::FResult& InResult, TSharedPtr<FMassEntityManager> EntityManager)
+{
+	if (InResult.PrunedProcessorClasses.Num() == 0 || !EntityManager || InResult.ArchetypeDataVersion == EntityManager->GetArchetypeDataVersion())
+	{
+		return true;
+	}
+
+	// this is inefficient right now since we're using CDOs and need to check all archetypes every time.
+	// Would be more efficient if we had a common place where all processors live, both active and inactive, so that we can utilize those. 
+	for (const TSubclassOf<UMassProcessor>& ProcessorClass : InResult.PrunedProcessorClasses)
+	{
+		if (UMassProcessor* ProcessorCDO = ProcessorClass.GetDefaultObject())
+		{
+			if (ProcessorCDO->DoesAnyArchetypeMatchOwnedQueries(*EntityManager.Get()))
+			{
+				return false;
+			}
+		}
+	}
+	return true;
+}
+
+#undef LOCTEXT_NAMESPACE 
\ No newline at end of file
diff --git a/Engine/Source/Runtime/MassEntity/Public/IMassEntityModule.h b/Engine/Source/Runtime/MassEntity/Public/IMassEntityModule.h
new file mode 100644
--- /dev/null
+++ b/Engine/Source/Runtime/MassEntity/Public/IMassEntityModule.h
@@ -0,0 +1,40 @@
+// Copyright Epic Games, Inc. All Rights Reserved.
+
+#pragma once
+
+#include "CoreMinimal.h"
+#include "Modules/ModuleInterface.h"
+#include "Modules/ModuleManager.h"
+
+
+/**
+ * The public interface to this module.  In most cases, this interface is only public to sibling modules 
+ * within this plugin.
+ */
+class IMassEntityModule : public IModuleInterface
+{
+
+public:
+
+	/**
+	 * Singleton-like access to this module's interface.  This is just for convenience!
+	 * Beware of calling this during the shutdown phase, though.  Your module might have been unloaded already.
+	 *
+	 * @return Returns singleton instance, loading the module on demand if needed
+	 */
+	static inline IMassEntityModule& Get()
+	{
+		return FModuleManager::LoadModuleChecked<IMassEntityModule>( "MassEntity" );
+	}
+
+	/**
+	 * Checks to see if this module is loaded and ready.  It is only valid to call Get() if IsAvailable() returns true.
+	 *
+	 * @return True if the module is loaded and ready to use
+	 */
+	static inline bool IsAvailable()
+	{
+		return FModuleManager::Get().IsModuleLoaded( "MassEntity" );
+	}
+};
+
diff --git a/Engine/Source/Runtime/MassEntity/Public/MassEntitySettings.h b/Engine/Source/Runtime/MassEntity/Public/MassEntitySettings.h
new file mode 100644
--- /dev/null
+++ b/Engine/Source/Runtime/MassEntity/Public/MassEntitySettings.h
@@ -0,0 +1,84 @@
+// Copyright Epic Games, Inc. All Rights Reserved.
+
+#pragma once
+
+#include "MassSettings.h"
+#include "MassProcessingPhaseManager.h"
+#include "MassProcessor.h"
+#include "StructUtils/InstancedStruct.h"
+#include "MassEntitySettings.generated.h"
+
+
+#define GET_MASS_CONFIG_VALUE(a) (GetMutableDefault<UMassEntitySettings>()->a)
+
+struct FPropertyChangedEvent;
+
+
+/**
+ * Implements the settings for MassEntity plugin
+ */
+UCLASS(config = Mass, defaultconfig, DisplayName = "Mass Entity")
+class MASSENTITY_API UMassEntitySettings : public UMassModuleSettings
+{
+	GENERATED_BODY()
+public:
+#if WITH_EDITORONLY_DATA
+	DECLARE_MULTICAST_DELEGATE_OneParam(FOnSettingsChange, const FPropertyChangedEvent& /*PropertyChangedEvent*/);
+#endif // WITH_EDITORONLY_DATA
+	DECLARE_MULTICAST_DELEGATE(FOnInitialized);
+
+	UMassEntitySettings(const FObjectInitializer& ObjectInitializer = FObjectInitializer::Get());
+
+	void BuildProcessorListAndPhases();
+	void AddToActiveProcessorsList(TSubclassOf<UMassProcessor> ProcessorClass);
+
+	TConstArrayView<FMassProcessingPhaseConfig> GetProcessingPhasesConfig();
+	const FMassProcessingPhaseConfig& GetProcessingPhaseConfig(const EMassProcessingPhase ProcessingPhase) const { check(ProcessingPhase != EMassProcessingPhase::MAX); return ProcessingPhasesConfig[int(ProcessingPhase)]; }
+
+	static FOnInitialized& GetOnInitializedEvent() { return GET_MASS_CONFIG_VALUE(OnInitializedEvent); }
+#if WITH_EDITOR
+	FOnSettingsChange& GetOnSettingsChange() { return OnSettingsChange; }
+
+	static bool IsInitialized() { return GET_MASS_CONFIG_VALUE(bInitialized); }
+
+protected:
+	virtual void PostEditChangeProperty(struct FPropertyChangedEvent& PropertyChangedEvent) override;
+	virtual void PostEditChangeChainProperty(struct FPropertyChangedChainEvent& PropertyChangedEvent) override;
+#endif // WITH_EDITOR
+
+	virtual void PostInitProperties() override;
+	virtual void BeginDestroy() override;
+
+	void OnPostEngineInit();
+	void BuildPhases();
+	void BuildProcessorList();
+
+public:
+	UPROPERTY(EditDefaultsOnly, Category = Mass, config, AdvancedDisplay)
+	int32 ChunkMemorySize = 128 * 1024;
+
+	/**
+	 * The name of the file to dump the processor dependency graph. T
+	 * The dot file will be put in the project log folder.
+	 * To generate a svg out of that file, simply run dot executable with following parameters: -Tsvg -O filename.dot 
+	 */
+	UPROPERTY(EditDefaultsOnly, Category = Mass, Transient)
+	FString DumpDependencyGraphFileName;
+
+	/** Lets users configure processing phases including the composite processor class to be used as a container for the phases' processors. */
+	UPROPERTY(EditDefaultsOnly, Category = Mass, config)
+	FMassProcessingPhaseConfig ProcessingPhasesConfig[(uint8)EMassProcessingPhase::MAX];
+
+	/** This list contains all the processors available in the given binary (including plugins). The contents are sorted by display name.*/
+	UPROPERTY(VisibleAnywhere, Category = Mass, Transient, Instanced, EditFixedSize)
+	TArray<TObjectPtr<UMassProcessor>> ProcessorCDOs;
+
+#if WITH_EDITORONLY_DATA
+protected:
+	FOnSettingsChange OnSettingsChange;
+#endif // WITH_EDITORONLY_DATA
+	bool bInitialized = false;
+	bool bEngineInitialized = false;
+
+	FOnInitialized OnInitializedEvent;
+};
diff --git a/Engine/Source/Runtime/MassEntity/Public/MassEntitySubsystem.h b/Engine/Source/Runtime/MassEntity/Public/MassEntitySubsystem.h
new file mode 100644
--- /dev/null
+++ b/Engine/Source/Runtime/MassEntity/Public/MassEntitySubsystem.h
@@ -0,0 +1,38 @@
+// Copyright Epic Games, Inc. All Rights Reserved.
+
+#pragma once
+
+#include "MassSubsystemBase.h"
+#include "MassEntityManager.h"
+#include "MassEntitySubsystem.generated.h"
+
+
+/** 
+ * The sole responsibility of this world subsystem class is to host the default instance of FMassEntityManager
+ * for a given UWorld. All the gameplay-related use cases of Mass (found in MassGameplay and related plugins) 
+ * use this by default. 
+ */
+UCLASS()
+class MASSENTITY_API UMassEntitySubsystem : public UMassSubsystemBase
+{
+	GENERATED_BODY()
+
+	//~USubsystem interface
+	virtual void Initialize(FSubsystemCollectionBase& Collection) override;
+	virtual void PostInitialize() override;
+	virtual void Deinitialize() override;
+	//~End of USubsystem interface
+
+public:
+	UMassEntitySubsystem();
+
+	//~UObject interface
+	virtual void GetResourceSizeEx(FResourceSizeEx& CumulativeResourceSize) override;
+	//~End of UObject interface
+
+	const FMassEntityManager& GetEntityManager() const { check(EntityManager); return *EntityManager.Get(); }
+	FMassEntityManager& GetMutableEntityManager() { check(EntityManager); return *EntityManager.Get(); }
+
+protected:
+	TSharedPtr<FMassEntityManager> EntityManager;
+};
diff --git a/Engine/Source/Runtime/MassEntity/Public/MassExecutor.h b/Engine/Source/Runtime/MassEntity/Public/MassExecutor.h
new file mode 100644
--- /dev/null
+++ b/Engine/Source/Runtime/MassEntity/Public/MassExecutor.h
@@ -0,0 +1,50 @@
+// Copyright Epic Games, Inc. All Rights Reserved.
+
+#pragma once
+
+#include "MassArchetypeTypes.h"
+#include "Async/TaskGraphInterfaces.h"
+
+
+struct FMassRuntimePipeline;
+struct FMassProcessingContext;
+struct FMassEntityHandle;
+struct FMassArchetypeEntityCollection;
+class UMassProcessor;
+
+
+namespace UE::Mass::Executor
+{
+	/** Executes processors in a given RuntimePipeline */
+	MASSENTITY_API void Run(FMassRuntimePipeline& RuntimePipeline, FMassProcessingContext& ProcessingContext);
+
+	/** Executes given Processor. Used mainly for triggering calculations via MassCompositeProcessors, e.g processing phases */
+	MASSENTITY_API void Run(UMassProcessor& Processor, FMassProcessingContext& ProcessingContext);
+
+	/** Similar to the Run function, but instead of using all the entities hosted by ProcessingContext.EntitySubsystem 
+	 *  it is processing only the entities given by EntityID via the Entities input parameter. 
+	 *  Note that all the entities need to be of Archetype archetype. 
+	 *  Under the hood the function converts Archetype-Entities pair to FMassArchetypeEntityCollection and calls the other flavor of RunSparse
+	 */
+	MASSENTITY_API void RunSparse(FMassRuntimePipeline& RuntimePipeline, FMassProcessingContext& ProcessingContext, FMassArchetypeHandle Archetype, TConstArrayView<FMassEntityHandle> Entities);
+
+	/** Similar to the Run function, but instead of using all the entities hosted by ProcessingContext.EntitySubsystem 
+	 *  it is processing only the entities given by SparseEntities input parameter.
+	 *  @todo rename
+	 */
+	MASSENTITY_API void RunSparse(FMassRuntimePipeline& RuntimePipeline, FMassProcessingContext& ProcessingContext, const FMassArchetypeEntityCollection& EntityCollection);
+
+	/** Executes given Processors array view. This function gets called under the hood by the rest of Run* functions */
+	MASSENTITY_API void RunProcessorsView(TArrayView<UMassProcessor* const> Processors, FMassProcessingContext& ProcessingContext, TConstArrayView<FMassArchetypeEntityCollection> EntityCollections = {});
+
+	/** 
+	 *  Triggers tasks executing Processor (and potentially it's children) and returns the task graph event representing 
+	 *  the task (the event will be "completed" once all the processors finish running). 
+	 *  @param OnDoneNotification will be called after all the processors are done, just after flushing the command buffer.
+	 *    Note that OnDoneNotification will be executed on GameThread.
+	 */
+	MASSENTITY_API FGraphEventRef TriggerParallelTasks(UMassProcessor& Processor, FMassProcessingContext& ProcessingContext, TFunction<void()> OnDoneNotification);
+
+	UE_DEPRECATED(5.5, "This flavor of RunProcessorsView is deprecated. Use the one with TConstArrayView<FMassArchetypeEntityCollection> parameter instead.")
+	MASSENTITY_API void RunProcessorsView(TArrayView<UMassProcessor* const> Processors, FMassProcessingContext& ProcessingContext, const FMassArchetypeEntityCollection* EntityCollection);
+};
diff --git a/Engine/Source/Runtime/MassEntity/Public/MassProcessor.h b/Engine/Source/Runtime/MassEntity/Public/MassProcessor.h
new file mode 100644
--- /dev/null
+++ b/Engine/Source/Runtime/MassEntity/Public/MassProcessor.h
@@ -0,0 +1,281 @@
+// Copyright Epic Games, Inc. All Rights Reserved.
+
+#pragma once
+
+#include "CoreMinimal.h"
+#include "MassEntityManager.h"
+#include "MassProcessingTypes.h"
+#include "Async/TaskGraphInterfaces.h"
+#include "MassCommandBuffer.h"
+#include "MassProcessor.generated.h"
+
+
+struct FMassProcessingPhaseConfig;
+class UMassCompositeProcessor;
+struct FMassDebugger;
+
+enum class EProcessorCompletionStatus : uint8
+{
+	Invalid,
+	Threaded,
+	Postponed,
+	Done
+};
+
+USTRUCT()
+struct FMassProcessorExecutionOrder
+{
+	GENERATED_BODY()
+
+	/** Determines which processing group this processor will be placed in. Leaving it empty ("None") means "top-most group for my ProcessingPhase" */
+	UPROPERTY(EditAnywhere, Category = Processor, config)
+	FName ExecuteInGroup = FName();
+
+	UPROPERTY(EditAnywhere, Category = Processor, config)
+	TArray<FName> ExecuteBefore;
+
+	UPROPERTY(EditAnywhere, Category = Processor, config)
+	TArray<FName> ExecuteAfter;
+};
+
+
+UCLASS(abstract, EditInlineNew, CollapseCategories, config = Mass, defaultconfig, ConfigDoNotCheckDefaults)
+class MASSENTITY_API UMassProcessor : public UObject
+{
+	GENERATED_BODY()
+public:
+	UMassProcessor();
+	UMassProcessor(const FObjectInitializer& ObjectInitializer);
+
+	bool IsInitialized() const { return bInitialized; }
+	virtual void Initialize(UObject& Owner) { bInitialized = true; }
+	virtual FGraphEventRef DispatchProcessorTasks(const TSharedPtr<FMassEntityManager>& EntityManager, FMassExecutionContext& ExecutionContext, const FGraphEventArray& Prerequisites = FGraphEventArray());
+
+	EProcessorExecutionFlags GetExecutionFlags() const { return (EProcessorExecutionFlags)ExecutionFlags; }
+
+	/** Whether this processor should execute according the CurrentExecutionFlags parameters */
+	bool ShouldExecute(const EProcessorExecutionFlags CurrentExecutionFlags) const { return (GetExecutionFlags() & CurrentExecutionFlags) != EProcessorExecutionFlags::None; }
+	void CallExecute(FMassEntityManager& EntityManager, FMassExecutionContext& Context);
+
+	/** 
+	 * Controls whether there can be multiple instances of a given class in a single FMassRuntimePipeline and during 
+	 * dependency solving. 
+	 */
+	bool ShouldAllowMultipleInstances() const { return bAllowMultipleInstances; }
+
+	UE_DEPRECATED(5.2, "This function is deprecated. Use ShouldAllowDuplicates instead.")
+	bool AllowDuplicates() const { return ShouldAllowMultipleInstances(); }
+
+	void DebugOutputDescription(FOutputDevice& Ar) const { DebugOutputDescription(Ar, 0); }
+	virtual void DebugOutputDescription(FOutputDevice& Ar, int32 Indent) const;
+	virtual FString GetProcessorName() const { return GetName(); }
+	
+	//----------------------------------------------------------------------//
+	// Ordering functions 
+	//----------------------------------------------------------------------//
+	/** Indicates whether this processor can ever be pruned while considered for a phase processing graph. A processor
+	 *  can get pruned if none of its registered queries interact with archetypes instantiated at the moment of graph
+	 *  building. This can also happen for special processors that don't register any queries - if that's the case override 
+	 *  this function to return an appropriate value
+	 *  @param bRuntimeMode indicates whether the pruning is being done for game runtime (true) or editor-time presentation (false) */
+	virtual bool ShouldAllowQueryBasedPruning(const bool bRuntimeMode = true) const { return bRuntimeMode; }
+
+	virtual EMassProcessingPhase GetProcessingPhase() const { return ProcessingPhase; }
+	virtual void SetProcessingPhase(EMassProcessingPhase Phase) { ProcessingPhase = Phase; }
+	bool DoesRequireGameThreadExecution() const { return bRequiresGameThreadExecution; }
+	
+	const FMassProcessorExecutionOrder& GetExecutionOrder() const { return ExecutionOrder; }
+
+	/** By default fetches requirements declared entity queries registered via RegisterQuery. Processors can override 
+	 *	this function to supply additional requirements */
+	virtual void ExportRequirements(FMassExecutionRequirements& OutRequirements) const;
+
+	const FMassSubsystemRequirements& GetProcessorRequirements() const { return ProcessorRequirements; }
+
+	/** Adds Query to RegisteredQueries list. Query is required to be a member variable of this processor. Not meeting
+	 *  this requirement will cause check failure and the query won't be registered. */
+	void RegisterQuery(FMassEntityQuery& Query);
+
+	bool ShouldAutoAddToGlobalList() const { return bAutoRegisterWithProcessingPhases; }
+#if WITH_EDITOR
+	bool ShouldShowUpInSettings() const { return ShouldAutoAddToGlobalList() || bCanShowUpInSettings; }
+#endif // WITH_EDITORONLY_DATA
+
+	/** Sets bAutoRegisterWithProcessingPhases. Setting it to true will result in this processor class being always 
+	 * instantiated to be automatically evaluated every frame. @see FMassProcessingPhaseManager
+	 * Note that calling this function is only valid on CDOs. Calling it on a regular instance will fail an ensure and 
+	 * have no other effect, i.e. CDO's value won't change */
+	void SetShouldAutoRegisterWithGlobalList(const bool bAutoRegister);
+
+	void GetArchetypesMatchingOwnedQueries(const FMassEntityManager& EntityManager, TArray<FMassArchetypeHandle>& OutArchetype);
+	bool DoesAnyArchetypeMatchOwnedQueries(const FMassEntityManager& EntityManager);
+	
+#if CPUPROFILERTRACE_ENABLED
+	FString StatId;
+#endif
+	
+protected:
+	virtual void ConfigureQueries() PURE_VIRTUAL(UMassProcessor::ConfigureQueries);
+	virtual void PostInitProperties() override;
+	virtual void Execute(FMassEntityManager& EntityManager, FMassExecutionContext& Context) PURE_VIRTUAL(UMassProcessor::Execute);
+
+protected:
+	/** Whether this processor should be executed on StandAlone or Server or Client */
+	UPROPERTY(EditAnywhere, Category = "Pipeline", meta = (Bitmask, BitmaskEnum = "/Script/MassEntity.EProcessorExecutionFlags"), config)
+	int32 ExecutionFlags;
+
+	/** Processing phase this processor will be automatically run as part of. Needs to be set before the processor gets
+	 *  registered with MassProcessingPhaseManager, otherwise it will have no effect. This property is usually read via
+	 *  a given class's CDO, so it's recommended to set it in the constructor. */
+	UPROPERTY(EditDefaultsOnly, Category = Processor, config)
+	EMassProcessingPhase ProcessingPhase = EMassProcessingPhase::PrePhysics;
+
+	/** Configures when this given processor can be executed in relation to other processors and processing groups, within its processing phase. */
+	UPROPERTY(EditDefaultsOnly, Category = Processor, config)
+	FMassProcessorExecutionOrder ExecutionOrder;
+
+	/** Configures whether this processor should be automatically included in the global list of processors executed every tick (see ProcessingPhase and ExecutionOrder). */
+	UPROPERTY(EditDefaultsOnly, Category = Processor, config)
+	bool bAutoRegisterWithProcessingPhases = true;
+
+	/** Meant as a class property, make sure to set it in subclass' constructor. Controls whether there can be multiple
+	 *  instances of a given class in a single FMassRuntimePipeline and during dependency solving. */
+	bool bAllowMultipleInstances = false;
+
+	UE_DEPRECATED(5.2, "This property is deprecated. Use bAllowMultipleInstances instead")
+	/** meant as a class property, make sure to set it in subclass' constructor. Controls whether there can be multiple 
+	 *  instances of a given class in a single FMassRuntimePipeline */
+	bool bAllowDuplicates = false;
+
+	UPROPERTY(EditDefaultsOnly, Category = Processor, config)
+	bool bRequiresGameThreadExecution = false;
+
+#if WITH_EDITORONLY_DATA
+	/** Used to permanently remove a given processor class from PipeSetting's listing. Used primarily for test-time 
+	 *  processor classes, but can also be used by project-specific code to prune the processor list. */
+	UPROPERTY(config)
+	bool bCanShowUpInSettings = true;
+#endif // WITH_EDITORONLY_DATA
+
+	friend UMassCompositeProcessor;
+	friend FMassDebugger;
+
+	/** A query representing elements this processor is accessing in Execute function outside of query execution */
+	FMassSubsystemRequirements ProcessorRequirements;
+
+private:
+	/** Stores processor's queries registered via RegisterQuery. 
+	 *  @note that it's safe to store pointers here since RegisterQuery does verify that a given registered query is 
+	 *  a member variable of a given processor */
+	TArray<FMassEntityQuery*> OwnedQueries;
+
+	/** Used to track whether Initialized has been called. */
+	bool bInitialized = false;
+};
+
+
+UCLASS()
+class MASSENTITY_API UMassCompositeProcessor : public UMassProcessor
+{
+	GENERATED_BODY()
+
+	friend FMassDebugger;
+public:
+	struct FDependencyNode
+	{
+		FName Name;
+		UMassProcessor* Processor = nullptr;
+		TArray<int32> Dependencies;
+#if WITH_MASSENTITY_DEBUG
+		int32 SequenceIndex = INDEX_NONE;
+#endif // WITH_MASSENTITY_DEBUG
+	};
+
+public:
+	UMassCompositeProcessor();
+
+	void SetChildProcessors(TArray<UMassProcessor*>&& InProcessors);
+
+	virtual void Initialize(UObject& Owner) override;
+	virtual void DebugOutputDescription(FOutputDevice& Ar, int32 Indent = 0) const override;
+	virtual void SetProcessingPhase(EMassProcessingPhase Phase) override;
+
+	void SetGroupName(FName NewName);
+	FName GetGroupName() const { return GroupName; }
+
+	virtual void SetProcessors(TArrayView<UMassProcessor*> InProcessorInstances, const TSharedPtr<FMassEntityManager>& EntityManager = nullptr);
+
+	/** 
+	 * Builds flat processing graph that's being used for multithreading execution of hosted processors.
+	 */
+	virtual void BuildFlatProcessingGraph(TConstArrayView<FMassProcessorOrderInfo> SortedProcessors);
+
+	/**
+	 * Adds processors in InOutOrderedProcessors to ChildPipeline. 
+	 * Note that this operation is non-destructive for the existing processors - the ones of classes found in InOutOrderedProcessors 
+	 * will be retained and used instead of the instances provided via InOutOrderedProcessors. Respective entries in InOutOrderedProcessors
+	 * will be updated to reflect the reuse.
+	 * The described behavior however is available only for processors with bAllowMultipleInstances == false.
+	 */
+	void UpdateProcessorsCollection(TArrayView<FMassProcessorOrderInfo> InOutOrderedProcessors, EProcessorExecutionFlags InWorldExecutionFlags = EProcessorExecutionFlags::None);
+
+	/** adds SubProcessor to an appropriately named group. If RequestedGroupName == None then SubProcessor
+	 *  will be added directly to ChildPipeline. If not then the indicated group will be searched for in ChildPipeline 
+	 *  and if it's missing it will be created and AddGroupedProcessor will be called recursively */
+	void AddGroupedProcessor(FName RequestedGroupName, UMassProcessor& SubProcessor);
+
+	virtual FGraphEventRef DispatchProcessorTasks(const TSharedPtr<FMassEntityManager>& EntityManager, FMassExecutionContext& ExecutionContext, const FGraphEventArray& Prerequisites = FGraphEventArray()) override;
+
+	bool IsEmpty() const { return ChildPipeline.IsEmpty(); }
+
+	virtual FString GetProcessorName() const override { return GroupName.ToString(); }
+
+protected:
+	virtual void ConfigureQueries() override;
+	virtual void Execute(FMassEntityManager& EntityManager, FMassExecutionContext& Context) override;
+
+	/** RequestedGroupName can indicate a multi-level group name, like so: A.B.C
+	 *  We need to extract the highest-level group name ('A' in the example), and see if it already exists. 
+	 *  If not, create it. 
+	 *  @param RequestedGroupName name of the group for which we want to find or create the processor.
+	 *  @param OutRemainingGroupName contains the group name after cutting the high-level group. In the used example it
+	 *    will contain "B.C". This value is then used to recursively create subgroups */
+	UMassCompositeProcessor* FindOrAddGroupProcessor(FName RequestedGroupName, FString* OutRemainingGroupName = nullptr);
+
+protected:
+	UPROPERTY(VisibleAnywhere, Category=Mass)
+	FMassRuntimePipeline ChildPipeline;
+
+	/** Group name that will be used when resolving processor dependencies and grouping */
+	UPROPERTY()
+	FName GroupName;
+
+	TArray<FDependencyNode> FlatProcessingGraph;
+
+	struct FProcessorCompletion
+	{
+		FGraphEventRef CompletionEvent;
+		EProcessorCompletionStatus Status = EProcessorCompletionStatus::Invalid;
+
+		bool IsDone() const 
+		{
+			return Status == EProcessorCompletionStatus::Done || (CompletionEvent.IsValid() && CompletionEvent->IsComplete());
+		}
+
+		void Wait()
+		{
+			if (CompletionEvent.IsValid())
+			{
+				CompletionEvent->Wait();
+			}
+		}
+	};
+	TArray<FProcessorCompletion> CompletionStatus;
+
+	//-----------------------------------------------------------------------------
+	// DEPRECATED
+	//-----------------------------------------------------------------------------
+public:
+	UE_DEPRECATED(5.3, "Populate is deprecated. Please use UpdateProcessorsCollection instead.")
+	void Populate(TConstArrayView<FMassProcessorOrderInfo> OrderedProcessors);
+};
