Tweaked how we configure Mass command's debug names. The change is prompted by failing unit tests, it's an attempt to reduce the number of alternative code paths when instantiating mass commands

diff --git a/Engine/Source/Runtime/MassEntity/Public/MassCommands.h b/Engine/Source/Runtime/MassEntity/Public/MassCommands.h
--- a/Engine/Source/Runtime/MassEntity/Public/MassCommands.h
+++ b/Engine/Source/Runtime/MassEntity/Public/MassCommands.h
@@ -37,8 +37,12 @@ enum class EMassCommandCheckTime : bool
 
 #if CSV_PROFILER_STATS || WITH_MASSENTITY_DEBUG
 #	define DEBUG_NAME(Name) , FName(TEXT(Name))
+#	define DEBUG_NAME_PARAM(Name) , const FName InDebugName = TEXT(Name)
+#	define FORWARD_DEBUG_NAME_PARAM , InDebugName
 #else
 #	define DEBUG_NAME(Name)
+#	define DEBUG_NAME_PARAM(Name)
+#	define FORWARD_DEBUG_NAME_PARAM
 #endif // CSV_PROFILER_STATS || WITH_MASSENTITY_DEBUG
 
 namespace UE::Mass::Utils
@@ -127,16 +131,10 @@ struct FMassBatchedEntityCommand : public FMassBatchedCommand
 	using Super = FMassBatchedCommand;
 
 	FMassBatchedEntityCommand() = default;
-	explicit FMassBatchedEntityCommand(EMassCommandOperationType OperationType)
-		: Super(OperationType)
+	explicit FMassBatchedEntityCommand(EMassCommandOperationType OperationType DEBUG_NAME_PARAM("BatchedEntityCommand"))
+		: Super(OperationType FORWARD_DEBUG_NAME_PARAM)
 	{}
 
-#if CSV_PROFILER_STATS || WITH_MASSENTITY_DEBUG
-	FMassBatchedEntityCommand(EMassCommandOperationType OperationType, FName DebugName)
-		: Super(OperationType, DebugName)
-	{}
-#endif // CSV_PROFILER_STATS || WITH_MASSENTITY_DEBUG
-
 	void Add(FMassEntityHandle Entity)
 	{
 		UE_MT_SCOPED_WRITE_ACCESS(EntitiesAccessDetector);
@@ -253,19 +251,11 @@ struct FMassCommandChangeTags : public FMassBatchedEntityCommand
 		: Super(EMassCommandOperationType::ChangeComposition DEBUG_NAME("ChangeTags"))
 	{}
 
-	FMassCommandChangeTags(EMassCommandOperationType OperationType, FMassTagBitSet TagsToAdd, FMassTagBitSet TagsToRemove)
-		: Super(OperationType DEBUG_NAME("ChangeTags"))
-		, TagsToAdd(TagsToAdd)
-		, TagsToRemove(TagsToRemove)
-	{}
-
-#if CSV_PROFILER_STATS || WITH_MASSENTITY_DEBUG
-	FMassCommandChangeTags(EMassCommandOperationType OperationType, FMassTagBitSet TagsToAdd, FMassTagBitSet TagsToRemove, FName DebugName)
-		: Super(OperationType, DebugName)
+	FMassCommandChangeTags(EMassCommandOperationType OperationType, FMassTagBitSet TagsToAdd, FMassTagBitSet TagsToRemove DEBUG_NAME_PARAM("ChangeTags"))
+		: Super(OperationType FORWARD_DEBUG_NAME_PARAM)
 		, TagsToAdd(TagsToAdd)
 		, TagsToRemove(TagsToRemove)
 	{}
-#endif // CSV_PROFILER_STATS || WITH_MASSENTITY_DEBUG
 
 protected:
 	virtual void Execute(FMassEntityManager& System) const override
@@ -348,17 +338,11 @@ struct FMassCommandAddFragmentInstances : public FMassBatchedEntityCommand
 {
 	using Super = FMassBatchedEntityCommand;
 
-	FMassCommandAddFragmentInstances()
-		: Super(EMassCommandOperationType::Set DEBUG_NAME("AddFragmentInstanceList"))
+	FMassCommandAddFragmentInstances(EMassCommandOperationType OperationType = EMassCommandOperationType::Set DEBUG_NAME_PARAM("AddFragmentInstanceList"))
+		: Super(EMassCommandOperationType::Set FORWARD_DEBUG_NAME_PARAM)
 		, FragmentsAffected(UE::Mass::Utils::ConstructFragmentBitSet<EMassCommandCheckTime::CompileTimeCheck, TOthers...>())
 	{}
 
-#if CSV_PROFILER || WITH_MASSENTITY_DEBUG
-	FMassCommandAddFragmentInstances(EMassCommandOperationType OperationType, FName DebugName)
-		: Super(OperationType, DebugName)
-	{}
-#endif // CSV_PROFILER || WITH_MASSENTITY_DEBUG
-
 	void Add(FMassEntityHandle Entity, TOthers... InFragments)
 	{
 		Super::Add(Entity);
@@ -601,4 +585,6 @@ using FMassDeferredChangeCompositionCommand = FMassDeferredCommand<EMassCommandO
 using FMassDeferredSetCommand = FMassDeferredCommand<EMassCommandOperationType::Set>;
 using FMassDeferredDestroyCommand = FMassDeferredCommand<EMassCommandOperationType::Destroy>;
 
-#undef DEBUG_NAME
\ No newline at end of file
+#undef DEBUG_NAME
+#undef DEBUG_NAME_PARAM
+#undef FORWARD_DEBUG_NAME_PARAM
