[Mass] Fixed FMassEntityManager::RemoveConstSharedFragmentFromEntity's return value in case where the entity doesn't have the given shared fragment type. The return value was inconsistent with the documentation and the other "failed" case when the given type doesn't represent a const shared fragment.
The issue also caused failure of the System.Mass.SharedFragments.RemoveFromEntity unit test.

diff --git a/Engine/Source/Runtime/MassEntity/Private/MassEntityManager.cpp b/Engine/Source/Runtime/MassEntity/Private/MassEntityManager.cpp
--- a/Engine/Source/Runtime/MassEntity/Private/MassEntityManager.cpp
+++ b/Engine/Source/Runtime/MassEntity/Private/MassEntityManager.cpp
@@ -1402,8 +1402,8 @@ bool FMassEntityManager::RemoveConstSharedFragmentFromEntity(const FMassEntityHa
 	
 	if (!CurrentArchetype->GetCompositionDescriptor().ConstSharedFragments.Contains(ConstSharedFragmentType))
 	{
-		// Nothing to do
-		return true;
+		// Nothing to do. Returning false to indicate nothing has been removed, as per function's documentation 
+		return false;
 	}
 
 	FMassArchetypeCompositionDescriptor NewComposition(CurrentArchetype->GetCompositionDescriptor());
