Moved the rest of MassEntity modules over to the Engine's Source/ code.

diff --git a/Engine/Plugins/Runtime/MassEntity/MassEntity.uplugin b/Engine/Plugins/Runtime/MassEntity/MassEntity.uplugin
--- a/Engine/Plugins/Runtime/MassEntity/MassEntity.uplugin
+++ b/Engine/Plugins/Runtime/MassEntity/MassEntity.uplugin
@@ -14,25 +14,6 @@
 	"IsBetaVersion": true,
 	"Installed": false,
 	"EnabledByDefault": false,
-	"Modules": [
-		{
-			"Name": "MassEntityTestSuite",
-			"Type": "UncookedOnly",
-			"LoadingPhase": "Default"
-		},
-		{
-			"Name": "MassEntityEditor",
-			"Type": "Editor"
-		},
-		{
-			"Name": "MassEntityDebugger",
-			"Type": "DeveloperTool"
-		},
-		{
-			"Name": "MassEntityDebugger",
-			"Type": "Editor"
-		}
-	],
-	"Plugins": [
-	]
+	"Modules": [],
+	"Plugins": []
 }
\ No newline at end of file
diff --git a/Engine/Source/Editor/MassEntityEditor/Public/IMassEntityEditor.h b/Engine/Source/Editor/MassEntityEditor/Public/IMassEntityEditor.h
new file mode 100644
--- /dev/null
+++ b/Engine/Source/Editor/MassEntityEditor/Public/IMassEntityEditor.h
@@ -0,0 +1,13 @@
+// Copyright Epic Games, Inc. All Rights Reserved.
+
+#pragma once
+
+#include "CoreMinimal.h"
+#include "Toolkits/AssetEditorToolkit.h"
+
+
+/** Mass Editor public interface */
+class MASSENTITYEDITOR_API IMassEntityEditor : public FAssetEditorToolkit
+{
+public:
+};
diff --git a/Engine/Source/Editor/MassEntityEditor/Public/MassEditorStyle.h b/Engine/Source/Editor/MassEntityEditor/Public/MassEditorStyle.h
new file mode 100644
--- /dev/null
+++ b/Engine/Source/Editor/MassEntityEditor/Public/MassEditorStyle.h
@@ -0,0 +1,26 @@
+// Copyright Epic Games, Inc. All Rights Reserved.
+
+#pragma once
+
+#include "Styling/SlateStyle.h"
+
+class MASSENTITYEDITOR_API FMassEntityEditorStyle
+{
+public:
+	static void Initialize();
+	static void Shutdown();
+	static TSharedPtr<ISlateStyle> Get() { return StyleSet; }
+	static FName GetStyleSetName();
+
+	static const FSlateBrush* GetBrush(FName PropertyName, const ANSICHAR* Specifier = NULL)
+	{
+		return StyleSet->GetBrush(PropertyName, Specifier);
+	}
+
+	static FString InContent(const FString& RelativePath, const ANSICHAR* Extension);
+
+
+private:
+
+	static TSharedPtr<FSlateStyleSet> StyleSet;
+};
diff --git a/Engine/Source/Editor/MassEntityEditor/Public/MassEntityEditorModule.h b/Engine/Source/Editor/MassEntityEditor/Public/MassEntityEditorModule.h
new file mode 100644
--- /dev/null
+++ b/Engine/Source/Editor/MassEntityEditor/Public/MassEntityEditorModule.h
@@ -0,0 +1,47 @@
+// Copyright Epic Games, Inc. All Rights Reserved.
+
+#pragma  once
+
+#include "Modules/ModuleInterface.h"
+#include "Toolkits/AssetEditorToolkit.h"
+#include "Delegates/IDelegateInstance.h"
+
+class IMassEntityEditor;
+struct FGraphPanelNodeFactory;
+struct FGraphNodeClassHelper;
+class UWorld;
+
+/**
+* The public interface to this module
+*/
+class MASSENTITYEDITOR_API FMassEntityEditorModule : public IModuleInterface, public IHasMenuExtensibility, public IHasToolBarExtensibility
+{
+public:
+	virtual void StartupModule() override;
+	virtual void ShutdownModule() override;
+
+	virtual TSharedPtr<FExtensibilityManager> GetMenuExtensibilityManager() override { return MenuExtensibilityManager; }
+	virtual TSharedPtr<FExtensibilityManager> GetToolBarExtensibilityManager() override { return ToolBarExtensibilityManager; }
+
+	UE_DEPRECATED(5.3, "Please use GetProcessorClassCache instead")
+	TSharedPtr<FGraphNodeClassHelper> GetProcassorClassCache() { return ProcessorClassCache; }
+	TSharedPtr<FGraphNodeClassHelper> GetProcessorClassCache() { return ProcessorClassCache; }
+
+protected:
+#if WITH_UNREAL_DEVELOPER_TOOLS
+	static void OnWorldCleanup(UWorld* /*World*/, bool /*bSessionEnded*/, bool /*bCleanupResources*/);
+	FDelegateHandle OnWorldCleanupHandle;
+#endif // WITH_UNREAL_DEVELOPER_TOOLS
+
+	TSharedPtr<FGraphNodeClassHelper> ProcessorClassCache;
+
+	TSharedPtr<FExtensibilityManager> MenuExtensibilityManager;
+	TSharedPtr<FExtensibilityManager> ToolBarExtensibilityManager;
+};
+
+#if UE_ENABLE_INCLUDE_ORDER_DEPRECATED_IN_5_3
+#include "CoreMinimal.h"
+#include "Modules/ModuleManager.h"
+#include "AssetTypeCategories.h"
+#include "Toolkits/IToolkitHost.h"
+#endif
diff --git a/Engine/Source/Editor/UnrealEd/Private/EditorEngine.cpp b/Engine/Source/Editor/UnrealEd/Private/EditorEngine.cpp
--- a/Engine/Source/Editor/UnrealEd/Private/EditorEngine.cpp
+++ b/Engine/Source/Editor/UnrealEd/Private/EditorEngine.cpp
@@ -1480,7 +1480,8 @@ void UEditorEngine::LoadDefaultEditorModules()
 			TEXT("RenderResourceViewer"),
 			TEXT("UniversalObjectLocatorEditor"),
 			TEXT("StructUtilsEditor"),
-			TEXT("StructUtilsTestSuite")
+			TEXT("StructUtilsTestSuite"),
+			TEXT("MassEntityEditor")
 		};
 
 	FScopedSlowTask ModuleSlowTask((float)UE_ARRAY_COUNT(ModuleNames));
diff --git a/Engine/Source/Runtime/Engine/Engine.Build.cs b/Engine/Source/Runtime/Engine/Engine.Build.cs
--- a/Engine/Source/Runtime/Engine/Engine.Build.cs
+++ b/Engine/Source/Runtime/Engine/Engine.Build.cs
@@ -374,6 +374,8 @@ public class Engine : ModuleRules
 
 		}
 
+		DynamicallyLoadedModuleNames.Add("MassEntity");
+
 		if (Target.bBuildEditor == true)
 		{
 			PublicIncludePathModuleNames.AddRange(
@@ -420,7 +422,8 @@ public class Engine : ModuleRules
 					"AssetTools",
 					"HierarchicalLODUtilities",
 					"LevelInstanceEditor",
-					"NaniteBuilder"
+					"NaniteBuilder",
+					"MassEntityEditor"
 				}
 			);
 		}
diff --git a/Engine/Source/Runtime/Launch/Private/LaunchEngineLoop.cpp b/Engine/Source/Runtime/Launch/Private/LaunchEngineLoop.cpp
--- a/Engine/Source/Runtime/Launch/Private/LaunchEngineLoop.cpp
+++ b/Engine/Source/Runtime/Launch/Private/LaunchEngineLoop.cpp
@@ -4656,6 +4656,8 @@ bool FEngineLoop::LoadStartupCoreModules()
 	FModuleManager::Get().LoadModule(TEXT("PacketHandler"));
 	FModuleManager::Get().LoadModule(TEXT("NetworkReplayStreaming"));
 
+	FModuleManager::Get().LoadModule(TEXT("MassEntity"));
+
 	return bSuccess;
 }
 
diff --git a/Engine/Source/Runtime/MassEntity/MassEntity.Build.cs b/Engine/Source/Runtime/MassEntity/MassEntity.Build.cs
--- a/Engine/Source/Runtime/MassEntity/MassEntity.Build.cs
+++ b/Engine/Source/Runtime/MassEntity/MassEntity.Build.cs
@@ -1,36 +1,40 @@
 // Copyright Epic Games, Inc. All Rights Reserved.
 
-namespace UnrealBuildTool.Rules
+using UnrealBuildTool;
+
+public class MassEntity : ModuleRules
 {
-	public class MassEntity : ModuleRules
+	public MassEntity(ReadOnlyTargetRules Target) : base(Target)
 	{
-		public MassEntity(ReadOnlyTargetRules Target) : base(Target)
-		{
-			PCHUsage = ModuleRules.PCHUsageMode.UseExplicitOrSharedPCHs;
-
-			UnsafeTypeCastWarningLevel = WarningLevel.Warning;
+		PCHUsage = ModuleRules.PCHUsageMode.UseExplicitOrSharedPCHs;
 
-			PublicDependencyModuleNames.AddRange(
-				new string[] {
-					"Core",
-					"CoreUObject",
-					"Engine",
-					"DeveloperSettings",
-				}
-			);
+		UnsafeTypeCastWarningLevel = WarningLevel.Warning;
 
-			if (Target.bBuildEditor || Target.bCompileAgainstEditor)
-			{
-				PrivateDependencyModuleNames.Add("UnrealEd");
-				PrivateDependencyModuleNames.Add("EditorSubsystem");
+		PublicDependencyModuleNames.AddRange(
+			new string[] {
+				"Core",
+				"CoreUObject",
+				"Engine",
+				"DeveloperSettings",
 			}
+		);
 
-			if (Target.Configuration != UnrealTargetConfiguration.Shipping
-				&& Target.Configuration != UnrealTargetConfiguration.Test)
-			{
-				// pulling this one in for the testableEnsureMsgf
-				PrivateDependencyModuleNames.Add("AITestSuite");
-			}
+		if (Target.bBuildEditor || Target.bCompileAgainstEditor)
+		{
+			PrivateDependencyModuleNames.Add("UnrealEd");
+			PrivateDependencyModuleNames.Add("EditorSubsystem");
+		}
+
+		if (Target.Configuration != UnrealTargetConfiguration.Shipping
+			&& Target.Configuration != UnrealTargetConfiguration.Test)
+		{
+			// pulling this one in for the testableEnsureMsgf
+			PrivateDependencyModuleNames.Add("AITestSuite");
+		}
+
+		if (Target.bBuildDeveloperTools)
+		{
+			DynamicallyLoadedModuleNames.Add("MassEntityTestSuite");
 		}
 	}
-}
\ No newline at end of file
+}
diff --git a/Engine/Source/Runtime/MassEntity/Private/MassProcessingTypes.cpp b/Engine/Source/Runtime/MassEntity/Private/MassProcessingTypes.cpp
--- a/Engine/Source/Runtime/MassEntity/Private/MassProcessingTypes.cpp
+++ b/Engine/Source/Runtime/MassEntity/Private/MassProcessingTypes.cpp
@@ -12,23 +12,6 @@
 
 DEFINE_LOG_CATEGORY(LogMass);
 
-
-namespace UE::Mass::Private
-{
-	struct FMassBootReporter
-	{
-		FMassBootReporter()
-		{
-#if MASS_DO_PARALLEL
-			UE_LOG(LogMass, Log, TEXT("MassEntity running with MULTITHREADING support."));
-#else
-			UE_LOG(LogMass, Log, TEXT("MassEntity running in game thread."));
-#endif // MASS_DO_PARALLEL
-		}
-	};
-	static const FMassBootReporter MassBootReporter;
-}
-
 //----------------------------------------------------------------------//
 //  FMassProcessingContext
 //----------------------------------------------------------------------//
