[Backout] - CL35079176
Original CL Desc
-----------------------------------------------------------------
Moved the rest of MassEntity modules over to the Engine's Source/ code.

diff --git a/Engine/Plugins/Runtime/MassEntity/Config/BaseMassEntity.ini b/Engine/Plugins/Runtime/MassEntity/Config/BaseMassEntity.ini
new file mode 100644
--- /dev/null
+++ b/Engine/Plugins/Runtime/MassEntity/Config/BaseMassEntity.ini
@@ -0,0 +1,2 @@
+[CoreRedirects]
+; no changes past UE 5.0
diff --git a/Engine/Plugins/Runtime/MassEntity/MassEntity.uplugin b/Engine/Plugins/Runtime/MassEntity/MassEntity.uplugin
--- a/Engine/Plugins/Runtime/MassEntity/MassEntity.uplugin
+++ b/Engine/Plugins/Runtime/MassEntity/MassEntity.uplugin
@@ -14,6 +14,25 @@
 	"IsBetaVersion": true,
 	"Installed": false,
 	"EnabledByDefault": false,
-	"Modules": [],
-	"Plugins": []
+	"Modules": [
+		{
+			"Name": "MassEntityTestSuite",
+			"Type": "UncookedOnly",
+			"LoadingPhase": "Default"
+		},
+		{
+			"Name": "MassEntityEditor",
+			"Type": "Editor"
+		},
+		{
+			"Name": "MassEntityDebugger",
+			"Type": "DeveloperTool"
+		},
+		{
+			"Name": "MassEntityDebugger",
+			"Type": "Editor"
+		}
+	],
+	"Plugins": [
+	]
 }
\ No newline at end of file
diff --git a/Engine/Plugins/Runtime/MassEntity/Source/MassEntityEditor/Public/IMassEntityEditor.h b/Engine/Plugins/Runtime/MassEntity/Source/MassEntityEditor/Public/IMassEntityEditor.h
new file mode 100644
--- /dev/null
+++ b/Engine/Plugins/Runtime/MassEntity/Source/MassEntityEditor/Public/IMassEntityEditor.h
@@ -0,0 +1,13 @@
+// Copyright Epic Games, Inc. All Rights Reserved.
+
+#pragma once
+
+#include "CoreMinimal.h"
+#include "Toolkits/AssetEditorToolkit.h"
+
+
+/** Mass Editor public interface */
+class MASSENTITYEDITOR_API IMassEntityEditor : public FAssetEditorToolkit
+{
+public:
+};
diff --git a/Engine/Plugins/Runtime/MassEntity/Source/MassEntityEditor/Public/MassEditorStyle.h b/Engine/Plugins/Runtime/MassEntity/Source/MassEntityEditor/Public/MassEditorStyle.h
new file mode 100644
--- /dev/null
+++ b/Engine/Plugins/Runtime/MassEntity/Source/MassEntityEditor/Public/MassEditorStyle.h
@@ -0,0 +1,26 @@
+// Copyright Epic Games, Inc. All Rights Reserved.
+
+#pragma once
+
+#include "Styling/SlateStyle.h"
+
+class MASSENTITYEDITOR_API FMassEntityEditorStyle
+{
+public:
+	static void Initialize();
+	static void Shutdown();
+	static TSharedPtr<ISlateStyle> Get() { return StyleSet; }
+	static FName GetStyleSetName();
+
+	static const FSlateBrush* GetBrush(FName PropertyName, const ANSICHAR* Specifier = NULL)
+	{
+		return StyleSet->GetBrush(PropertyName, Specifier);
+	}
+
+	static FString InContent(const FString& RelativePath, const ANSICHAR* Extension);
+
+
+private:
+
+	static TSharedPtr<FSlateStyleSet> StyleSet;
+};
diff --git a/Engine/Plugins/Runtime/MassEntity/Source/MassEntityEditor/Public/MassEntityEditorModule.h b/Engine/Plugins/Runtime/MassEntity/Source/MassEntityEditor/Public/MassEntityEditorModule.h
new file mode 100644
--- /dev/null
+++ b/Engine/Plugins/Runtime/MassEntity/Source/MassEntityEditor/Public/MassEntityEditorModule.h
@@ -0,0 +1,40 @@
+// Copyright Epic Games, Inc. All Rights Reserved.
+
+#pragma  once
+
+#include "Modules/ModuleInterface.h"
+#include "Toolkits/AssetEditorToolkit.h"
+
+class IMassEntityEditor;
+struct FGraphPanelNodeFactory;
+struct FGraphNodeClassHelper;
+
+/**
+* The public interface to this module
+*/
+class MASSENTITYEDITOR_API FMassEntityEditorModule : public IModuleInterface, public IHasMenuExtensibility, public IHasToolBarExtensibility
+{
+public:
+	virtual void StartupModule() override;
+	virtual void ShutdownModule() override;
+
+	virtual TSharedPtr<FExtensibilityManager> GetMenuExtensibilityManager() override { return MenuExtensibilityManager; }
+	virtual TSharedPtr<FExtensibilityManager> GetToolBarExtensibilityManager() override { return ToolBarExtensibilityManager; }
+
+	UE_DEPRECATED(5.3, "Please use GetProcessorClassCache instead")
+	TSharedPtr<FGraphNodeClassHelper> GetProcassorClassCache() { return ProcessorClassCache; }
+	TSharedPtr<FGraphNodeClassHelper> GetProcessorClassCache() { return ProcessorClassCache; }
+
+protected:
+	TSharedPtr<FGraphNodeClassHelper> ProcessorClassCache;
+
+	TSharedPtr<FExtensibilityManager> MenuExtensibilityManager;
+	TSharedPtr<FExtensibilityManager> ToolBarExtensibilityManager;
+};
+
+#if UE_ENABLE_INCLUDE_ORDER_DEPRECATED_IN_5_3
+#include "CoreMinimal.h"
+#include "Modules/ModuleManager.h"
+#include "AssetTypeCategories.h"
+#include "Toolkits/IToolkitHost.h"
+#endif
diff --git a/Engine/Plugins/Runtime/MassEntity/Source/MassEntityTestSuite/Private/MassEntityTestSuiteModule.cpp b/Engine/Plugins/Runtime/MassEntity/Source/MassEntityTestSuite/Private/MassEntityTestSuiteModule.cpp
new file mode 100644
--- /dev/null
+++ b/Engine/Plugins/Runtime/MassEntity/Source/MassEntityTestSuite/Private/MassEntityTestSuiteModule.cpp
@@ -0,0 +1,13 @@
+// Copyright Epic Games, Inc. All Rights Reserved.
+
+#include "MassEntityTestSuiteModule.h"
+
+#define LOCTEXT_NAMESPACE "MassTest"
+
+class FMassEntityTestSuiteModule : public IMassEntityTestSuiteModule
+{
+};
+
+IMPLEMENT_MODULE(FMassEntityTestSuiteModule, MassEntityTestSuite)
+
+#undef LOCTEXT_NAMESPACE
diff --git a/Engine/Plugins/Runtime/MassEntity/Source/MassEntityTestSuite/Public/MassEntityTestSuiteModule.h b/Engine/Plugins/Runtime/MassEntity/Source/MassEntityTestSuite/Public/MassEntityTestSuiteModule.h
new file mode 100644
--- /dev/null
+++ b/Engine/Plugins/Runtime/MassEntity/Source/MassEntityTestSuite/Public/MassEntityTestSuiteModule.h
@@ -0,0 +1,37 @@
+// Copyright Epic Games, Inc. All Rights Reserved.
+
+#pragma once
+
+#include "CoreMinimal.h"
+#include "Modules/ModuleManager.h"
+
+/**
+* The public interface to this module
+*/
+class IMassEntityTestSuiteModule : public IModuleInterface
+{
+
+public:
+
+	/**
+	* Singleton-like access to this module's interface.  This is just for convenience!
+	* Beware of calling this during the shutdown phase, though.  Your module might have been unloaded already.
+	*
+	* @return Returns singleton instance, loading the module on demand if needed
+	*/
+	static inline IMassEntityTestSuiteModule& Get()
+	{
+		return FModuleManager::LoadModuleChecked<IMassEntityTestSuiteModule>("MassEntityTestSuite");
+	}
+
+	/**
+	* Checks to see if this module is loaded and ready.  It is only valid to call Get() if IsAvailable() returns true.
+	*
+	* @return True if the module is loaded and ready to use
+	*/
+	static inline bool IsAvailable()
+	{
+		return FModuleManager::Get().IsModuleLoaded("MassEntityTestSuite");
+	}
+};
+
diff --git a/Engine/Source/Editor/UnrealEd/Private/EditorEngine.cpp b/Engine/Source/Editor/UnrealEd/Private/EditorEngine.cpp
--- a/Engine/Source/Editor/UnrealEd/Private/EditorEngine.cpp
+++ b/Engine/Source/Editor/UnrealEd/Private/EditorEngine.cpp
@@ -1480,8 +1480,7 @@ void UEditorEngine::LoadDefaultEditorModules()
 			TEXT("RenderResourceViewer"),
 			TEXT("UniversalObjectLocatorEditor"),
 			TEXT("StructUtilsEditor"),
-			TEXT("StructUtilsTestSuite"),
-			TEXT("MassEntityEditor")
+			TEXT("StructUtilsTestSuite")
 		};
 
 	FScopedSlowTask ModuleSlowTask((float)UE_ARRAY_COUNT(ModuleNames));
diff --git a/Engine/Source/Runtime/Engine/Engine.Build.cs b/Engine/Source/Runtime/Engine/Engine.Build.cs
--- a/Engine/Source/Runtime/Engine/Engine.Build.cs
+++ b/Engine/Source/Runtime/Engine/Engine.Build.cs
@@ -374,8 +374,6 @@ public class Engine : ModuleRules
 
 		}
 
-		DynamicallyLoadedModuleNames.Add("MassEntity");
-
 		if (Target.bBuildEditor == true)
 		{
 			PublicIncludePathModuleNames.AddRange(
@@ -422,8 +420,7 @@ public class Engine : ModuleRules
 					"AssetTools",
 					"HierarchicalLODUtilities",
 					"LevelInstanceEditor",
-					"NaniteBuilder",
-					"MassEntityEditor"
+					"NaniteBuilder"
 				}
 			);
 		}
diff --git a/Engine/Source/Runtime/Launch/Private/LaunchEngineLoop.cpp b/Engine/Source/Runtime/Launch/Private/LaunchEngineLoop.cpp
--- a/Engine/Source/Runtime/Launch/Private/LaunchEngineLoop.cpp
+++ b/Engine/Source/Runtime/Launch/Private/LaunchEngineLoop.cpp
@@ -4656,8 +4656,6 @@ bool FEngineLoop::LoadStartupCoreModules()
 	FModuleManager::Get().LoadModule(TEXT("PacketHandler"));
 	FModuleManager::Get().LoadModule(TEXT("NetworkReplayStreaming"));
 
-	FModuleManager::Get().LoadModule(TEXT("MassEntity"));
-
 	return bSuccess;
 }
 
diff --git a/Engine/Source/Runtime/MassEntity/MassEntity.Build.cs b/Engine/Source/Runtime/MassEntity/MassEntity.Build.cs
--- a/Engine/Source/Runtime/MassEntity/MassEntity.Build.cs
+++ b/Engine/Source/Runtime/MassEntity/MassEntity.Build.cs
@@ -1,40 +1,36 @@
 // Copyright Epic Games, Inc. All Rights Reserved.
 
-using UnrealBuildTool;
-
-public class MassEntity : ModuleRules
+namespace UnrealBuildTool.Rules
 {
-	public MassEntity(ReadOnlyTargetRules Target) : base(Target)
+	public class MassEntity : ModuleRules
 	{
-		PCHUsage = ModuleRules.PCHUsageMode.UseExplicitOrSharedPCHs;
-
-		UnsafeTypeCastWarningLevel = WarningLevel.Warning;
+		public MassEntity(ReadOnlyTargetRules Target) : base(Target)
+		{
+			PCHUsage = ModuleRules.PCHUsageMode.UseExplicitOrSharedPCHs;
 
-		PublicDependencyModuleNames.AddRange(
-			new string[] {
-				"Core",
-				"CoreUObject",
-				"Engine",
-				"DeveloperSettings",
-			}
-		);
+			UnsafeTypeCastWarningLevel = WarningLevel.Warning;
 
-		if (Target.bBuildEditor || Target.bCompileAgainstEditor)
-		{
-			PrivateDependencyModuleNames.Add("UnrealEd");
-			PrivateDependencyModuleNames.Add("EditorSubsystem");
-		}
+			PublicDependencyModuleNames.AddRange(
+				new string[] {
+					"Core",
+					"CoreUObject",
+					"Engine",
+					"DeveloperSettings",
+				}
+			);
 
-		if (Target.Configuration != UnrealTargetConfiguration.Shipping
-			&& Target.Configuration != UnrealTargetConfiguration.Test)
-		{
-			// pulling this one in for the testableEnsureMsgf
-			PrivateDependencyModuleNames.Add("AITestSuite");
-		}
+			if (Target.bBuildEditor || Target.bCompileAgainstEditor)
+			{
+				PrivateDependencyModuleNames.Add("UnrealEd");
+				PrivateDependencyModuleNames.Add("EditorSubsystem");
+			}
 
-		if (Target.bBuildDeveloperTools)
-		{
-			DynamicallyLoadedModuleNames.Add("MassEntityTestSuite");
+			if (Target.Configuration != UnrealTargetConfiguration.Shipping
+				&& Target.Configuration != UnrealTargetConfiguration.Test)
+			{
+				// pulling this one in for the testableEnsureMsgf
+				PrivateDependencyModuleNames.Add("AITestSuite");
+			}
 		}
 	}
-}
+}
\ No newline at end of file
