Fix incorrectly conditional include, needed for IMPLEMENT_MODULE macro

#rnx

diff --git a/Engine/Source/Runtime/MassEntity/Private/MassEntityModule.cpp b/Engine/Source/Runtime/MassEntity/Private/MassEntityModule.cpp
--- a/Engine/Source/Runtime/MassEntity/Private/MassEntityModule.cpp
+++ b/Engine/Source/Runtime/MassEntity/Private/MassEntityModule.cpp
@@ -3,9 +3,7 @@
 #include "Modules/ModuleInterface.h"
 #include "MassEntityTypes.h"
 #include "MassProcessingTypes.h"
-#if WITH_UNREAL_DEVELOPER_TOOLS
 #include "Modules/ModuleManager.h"
-#endif // WITH_UNREAL_DEVELOPER_TOOLS
 
 #define LOCTEXT_NAMESPACE "Mass"
 
