[Mass] Added a basic functionality for module-independent "debug event" observation. The point is to reduce inter-module dependencies while supplying very flexible way of distributing and handling of Mass debug events

diff --git a/Engine/Source/Runtime/MassEntity/Private/MassDebugger.cpp b/Engine/Source/Runtime/MassEntity/Private/MassDebugger.cpp
--- a/Engine/Source/Runtime/MassEntity/Private/MassDebugger.cpp
+++ b/Engine/Source/Runtime/MassEntity/Private/MassDebugger.cpp
@@ -377,6 +377,7 @@ FMassDebugger::FOnEntitySelected FMassDebugger::OnEntitySelectedDelegate;
 
 FMassDebugger::FOnMassEntityManagerEvent FMassDebugger::OnEntityManagerInitialized;
 FMassDebugger::FOnMassEntityManagerEvent FMassDebugger::OnEntityManagerDeinitialized;
+FMassDebugger::FOnDebugEvent FMassDebugger::OnDebugEvent;
 TArray<FMassDebugger::FEnvironment> FMassDebugger::ActiveEnvironments;
 UE::FSpinLock FMassDebugger::EntityManagerRegistrationLock;
 
diff --git a/Engine/Source/Runtime/MassEntity/Public/MassDebugger.h b/Engine/Source/Runtime/MassEntity/Public/MassDebugger.h
--- a/Engine/Source/Runtime/MassEntity/Public/MassDebugger.h
+++ b/Engine/Source/Runtime/MassEntity/Public/MassDebugger.h
@@ -8,7 +8,8 @@
 #include "MassEntityQuery.h"
 #include "MassProcessor.h"
 #include "Misc/SpinLock.h"
-
+#include "StructUtils/InstancedStruct.h"
+#include "Logging/TokenizedMessage.h"
 
 class FOutputDevice;
 class UMassProcessor;
@@ -20,6 +21,7 @@ struct FMassFragmentRequirementDescription;
 enum class EMassFragmentAccess : uint8;
 enum class EMassFragmentPresence : uint8;
 #endif // WITH_MASSENTITY_DEBUG
+#include "MassDebugger.generated.h"
 
 namespace UE::Mass::Debug
 {
@@ -40,6 +42,16 @@ namespace UE::Mass::Debug
 	};
 } // namespace UE::Mass::Debug
 
+USTRUCT()
+struct MASSENTITY_API FMassGenericDebugEvent
+{
+	GENERATED_BODY()
+#if WITH_EDITORONLY_DATA
+	// note that it's not a uproperty since these events are only intended to be used instantly, never stored
+	const UObject* Context = nullptr;
+#endif // WITH_EDITORONLY_DATA
+};
+
 #if WITH_MASSENTITY_DEBUG
 
 namespace UE::Mass::Debug
@@ -54,6 +66,16 @@ namespace UE::Mass::Debug
 #define MASS_BREAK_IF_ENTITY_INDEX(EntityHandle, InIndex) { if (UE::Mass::Debug::bAllowBreakOnDebuggedEntity && EntityHandle.Index == InIndex) { PLATFORM_BREAK();} }
 #define MASS_SET_ENTITY_DEBUGGED(Manager, EntityHandle) { if (UE::Mass::Debug::bAllowProceduralDebuggedEntitySelection) {FMassDebugger::SelectEntity(Manager, EntityHandle); }}
 
+enum class EMassDebugMessageSeverity : uint8
+{
+	Error,
+	Warning,
+	Info,
+	// the following two need to remain last
+	Default,
+	MAX = Default
+};
+
 namespace UE::Mass::Debug
 {
 	struct MASSENTITY_API FQueryRequirementsView
@@ -84,13 +106,26 @@ namespace UE::Mass::Debug
 	MASSENTITY_API extern bool GetDebugEntitiesRange(int32& OutBegin, int32& OutEnd);
 	MASSENTITY_API extern bool IsDebuggingEntity(FMassEntityHandle Entity, FColor* OutEntityColor = nullptr);
 	MASSENTITY_API extern FColor GetEntityDebugColor(FMassEntityHandle Entity);
-} // namespace UE::Mass::Debug
 
+	inline EMessageSeverity::Type MassSeverityToMessageSeverity(EMessageSeverity::Type OriginalSeverity, EMassDebugMessageSeverity MassSeverity)
+	{
+		static constexpr EMessageSeverity::Type ConversionMap[int(EMassDebugMessageSeverity::MAX)] =
+		{
+			/*EMassDebugMessageSeverity::Error=*/EMessageSeverity::Error,
+			/*EMassDebugMessageSeverity::Warning=*/EMessageSeverity::Warning,
+			/*EMassDebugMessageSeverity::Info=*/EMessageSeverity::Info
+		};
+		return MassSeverity == EMassDebugMessageSeverity::Default 
+			? OriginalSeverity
+			: ConversionMap[int(MassSeverity)];
+	}
+} // namespace UE::Mass::Debug
 
 struct MASSENTITY_API FMassDebugger
 {
 	DECLARE_MULTICAST_DELEGATE_TwoParams(FOnEntitySelected, const FMassEntityManager&, const FMassEntityHandle);
 	DECLARE_MULTICAST_DELEGATE_OneParam(FOnMassEntityManagerEvent, const FMassEntityManager&);
+	DECLARE_MULTICAST_DELEGATE_ThreeParams(FOnDebugEvent, const FName /*EventName*/, FConstStructView /*Payload*/, const EMassDebugMessageSeverity /*SeverityOverride*/);
 
 	struct FEnvironment
 	{
@@ -138,6 +173,20 @@ struct MASSENTITY_API FMassDebugger
 	static FOnMassEntityManagerEvent OnEntityManagerInitialized;
 	static FOnMassEntityManagerEvent OnEntityManagerDeinitialized;
 
+	static FOnDebugEvent OnDebugEvent;
+	
+	static void DebugEvent(const FName EventName, FConstStructView Payload, const EMassDebugMessageSeverity SeverityOverride = EMassDebugMessageSeverity::Default)
+	{
+		OnDebugEvent.Broadcast(EventName, Payload, SeverityOverride);
+	}
+
+	template<typename TMessage, typename... TArgs>
+	static void DebugEvent(TArgs&&... InArgs)
+	{
+		DebugEvent(TMessage::StaticStruct()->GetFName()
+			, FInstancedStruct::Make<TMessage>(Forward<TArgs>(InArgs)...));
+	}
+
 	static void RegisterEntityManager(FMassEntityManager& EntityManager);
 	static void UnregisterEntityManager(FMassEntityManager& EntityManager);
 	static TConstArrayView<FEnvironment> GetEnvironments() { return ActiveEnvironments; }
@@ -157,13 +206,14 @@ private:
 struct FMassArchetypeHandle;
 struct FMassFragmentRequirements;
 struct FMassFragmentRequirementDescription;
+struct FMassArchetypeCompositionDescriptor;
 
 struct MASSENTITY_API FMassDebugger
 {
-	static FString GetSingleRequirementDescription(const FMassFragmentRequirementDescription& Requirement) { return TEXT("[no debug information]"); }
-	static FString GetRequirementsDescription(const FMassFragmentRequirements& Requirements) { return TEXT("[no debug information]"); }
-	static FString GetArchetypeRequirementCompatibilityDescription(const FMassFragmentRequirements& Requirements, const FMassArchetypeHandle& ArchetypeHandle) { return TEXT("[no debug information]"); }
-	static FString GetArchetypeRequirementCompatibilityDescription(const FMassFragmentRequirements& Requirements, const FMassArchetypeCompositionDescriptor& ArchetypeComposition) { return TEXT("[no debug information]"); }
+	static FString GetSingleRequirementDescription(const FMassFragmentRequirementDescription&) { return TEXT("[no debug information]"); }
+	static FString GetRequirementsDescription(const FMassFragmentRequirements&) { return TEXT("[no debug information]"); }
+	static FString GetArchetypeRequirementCompatibilityDescription(const FMassFragmentRequirements&, const FMassArchetypeHandle&) { return TEXT("[no debug information]"); }
+	static FString GetArchetypeRequirementCompatibilityDescription(const FMassFragmentRequirements&, const FMassArchetypeCompositionDescriptor&) { return TEXT("[no debug information]"); }
 };
 
 #define MASS_IF_ENTITY_DEBUGGED(a, b) false
