[Main CIS] Issue 533091: Mac CIS fix

diff --git a/Engine/Source/Runtime/MassEntity/Public/MassDebugger.h b/Engine/Source/Runtime/MassEntity/Public/MassDebugger.h
--- a/Engine/Source/Runtime/MassEntity/Public/MassDebugger.h
+++ b/Engine/Source/Runtime/MassEntity/Public/MassDebugger.h
@@ -46,6 +46,13 @@ USTRUCT()
 struct MASSENTITY_API FMassGenericDebugEvent
 {
 	GENERATED_BODY()
+	explicit FMassGenericDebugEvent(const UObject* InContext = nullptr)
+#if WITH_EDITORONLY_DATA
+		: Context(InContext)
+#endif // WITH_EDITORONLY_DATA
+	{
+	}
+
 #if WITH_EDITORONLY_DATA
 	// note that it's not a uproperty since these events are only intended to be used instantly, never stored
 	const UObject* Context = nullptr;
