Mass Concurrent Reserve Everywhere
* Enable Mass Concurrent Reserve feature for general runtime
* CVars added to revert back to previous implementation and to control memory usage of entity lookup datastructure
- "Mass.ConcurrentReserve.Enable"
- "Mass.ConcurrentReserve.MaxEntityCount"
- "Mass.ConcurrentReserve.EntitiesPerPage"
* Allows deprecation/removal/cleanup of single threaded implementation

diff --git a/Engine/Source/Runtime/MassEntity/Private/MassEntitySubsystem.cpp b/Engine/Source/Runtime/MassEntity/Private/MassEntitySubsystem.cpp
--- a/Engine/Source/Runtime/MassEntity/Private/MassEntitySubsystem.cpp
+++ b/Engine/Source/Runtime/MassEntity/Private/MassEntitySubsystem.cpp
@@ -5,8 +5,36 @@
 
 #include UE_INLINE_GENERATED_CPP_BY_NAME(MassEntitySubsystem)
 
+namespace UE::Mass::Private
+{
+	static bool bEnableMassConcurrentReserveRuntime = true;
+	static int32 ConcurrentReserveMaxEntityCount = 1 << 27;
+	static int32 ConcurrentReserveMaxEntitiesPerPage = 1 << 16;
 
-
+	namespace
+	{
+		FAutoConsoleVariableRef CVars[] = {
+			{
+				TEXT("Mass.ConcurrentReserve.Enable"),
+				bEnableMassConcurrentReserveRuntime,
+				TEXT("Enable Mass's concurrent reserve feature in runtime"),
+				ECVF_Default
+			},
+			{
+				TEXT("Mass.ConcurrentReserve.MaxEntityCount"),
+				ConcurrentReserveMaxEntityCount,
+				TEXT("Set maximum number of permissible entities.  Must be power of 2."),
+				ECVF_Default
+			},
+			{
+				TEXT("Mass.ConcurrentReserve.EntitiesPerPage"),
+				ConcurrentReserveMaxEntitiesPerPage,
+				TEXT("Set number of entities per page. Must be power of 2. Larger reduces fixed memory overhead of FEntityData page lookup but requires bigger contiguous memory blocks per page"),
+				ECVF_Default
+			}
+		};
+	}
+}
 //////////////////////////////////////////////////////////////////////
 // UMassEntitySubsystem
 
@@ -25,7 +53,24 @@ void UMassEntitySubsystem::GetResourceSizeEx(FResourceSizeEx& CumulativeResource
 void UMassEntitySubsystem::Initialize(FSubsystemCollectionBase& Collection)
 {
 	Super::Initialize(Collection);
-	EntityManager->Initialize();
+
+	FMassEntityManagerStorageInitParams InitializationParams;
+	
+	if (UE::Mass::Private::bEnableMassConcurrentReserveRuntime)
+	{
+		InitializationParams.Emplace<FMassEntityManager_InitParams_Concurrent>(
+			FMassEntityManager_InitParams_Concurrent
+			{
+				.MaxEntityCount = static_cast<uint32>(UE::Mass::Private::ConcurrentReserveMaxEntityCount),
+				.MaxEntitiesPerPage = static_cast<uint32>(UE::Mass::Private::ConcurrentReserveMaxEntitiesPerPage)
+			});
+	}
+	else
+	{
+		InitializationParams.Emplace<FMassEntityManager_InitParams_SingleThreaded>();
+	}
+	
+	EntityManager->Initialize(InitializationParams);
 	HandleLateCreation();
 }
 
diff --git a/Engine/Source/Runtime/MassEntity/Public/MassEntityManager.h b/Engine/Source/Runtime/MassEntity/Public/MassEntityManager.h
--- a/Engine/Source/Runtime/MassEntity/Public/MassEntityManager.h
+++ b/Engine/Source/Runtime/MassEntity/Public/MassEntityManager.h
@@ -32,13 +32,7 @@ namespace UE::Mass::Private
 	struct FEntityStorageInitializer;
 }
 
-// Only Editor builds enable Mass to use the concurrent entity reservation implementation
-// Non-editor builds will devirtualize to a single threaded implementation only
-#if UE_EDITOR
 #define MASS_CONCURRENT_RESERVE 1
-#else
-#define MASS_CONCURRENT_RESERVE 0
-#endif
 
 /** 
  * The type responsible for hosting Entities managing Archetypes.
