[Mass] Fix for observers getting notified before the fragments values are set on entities created through 'CreateEntity' and 'BuildEntity'

diff --git a/Engine/Source/Developer/MassEntityTestSuite/Private/MassObserverTest.cpp b/Engine/Source/Developer/MassEntityTestSuite/Private/MassObserverTest.cpp
--- a/Engine/Source/Developer/MassEntityTestSuite/Private/MassObserverTest.cpp
+++ b/Engine/Source/Developer/MassEntityTestSuite/Private/MassObserverTest.cpp
@@ -452,6 +452,50 @@ struct FFragmentObserverTest_MultipleArchetypeDestroy : FFragmentTestBase
 };
 IMPLEMENT_AI_INSTANT_TEST(FFragmentObserverTest_MultipleArchetypeDestroy, "System.Mass.Observer.Fragment.MultipleArchetypesDestroy");
 
+struct FFragmentObserverTest_EntityCreation_Individual : FFragmentTestBase
+{
+	FFragmentObserverTest_EntityCreation_Individual() { OperationObserved = EMassObservedOperation::Add; }
+
+	virtual bool InstantTest() override
+	{
+		constexpr float TestValue = 123.456f;
+		float ValueOnNotification = 0.f;
+
+		ObserverProcessor->ForEachEntityChunkExecutionFunction = [&ValueOnNotification](FMassExecutionContext& Context)
+			{
+				const TConstArrayView<FFragmentStruct> Fragments = Context.GetFragmentView<FFragmentStruct>();
+				for (int32 EntityIndex = 0; EntityIndex < Context.GetNumEntities(); EntityIndex++)
+				{
+					ValueOnNotification = Fragments[EntityIndex].Value;
+				};
+			};
+
+		FMassObserverManager& ObserverManager = EntityManager->GetObserverManager();
+		ObserverManager.AddObserverInstance(*FFragmentStruct::StaticStruct(), OperationObserved, *ObserverProcessor);
+
+		TArray<FInstancedStruct> FragmentInstanceList = { FInstancedStruct::Make(FFragmentStruct(TestValue)) };
+
+		// BuildEntity
+		{
+			const FMassEntityHandle Entity= EntityManager->ReserveEntity();
+			EntityManager->BuildEntity(Entity, FragmentInstanceList);	
+			AITEST_EQUAL(TEXT("The fragment observer notified by BuildEntity is expected to be able to fetch the initial value"), ValueOnNotification, TestValue);
+			EntityManager->DestroyEntity(Entity);
+		}
+
+		// CreateEntity
+		{
+			ValueOnNotification = 0.f;
+			const FMassEntityHandle Entity = EntityManager->CreateEntity(FragmentInstanceList);
+			AITEST_EQUAL(TEXT("The fragment observer notified by CreateEntity is expected to be able to fetch the initial value"), ValueOnNotification, TestValue);
+			EntityManager->DestroyEntity(Entity);
+		}
+
+		return true;
+	}
+};
+IMPLEMENT_AI_INSTANT_TEST(FFragmentObserverTest_EntityCreation_Individual, "System.Mass.Observer.Create.FragmentSingleEntity");
+
 struct FFragmentObserverTest_EntityCreation_Individuals : FFragmentTestBase
 {
 	FFragmentObserverTest_EntityCreation_Individuals() { OperationObserved = EMassObservedOperation::Add; }
diff --git a/Engine/Source/Runtime/MassEntity/Private/MassEntityManager.cpp b/Engine/Source/Runtime/MassEntity/Private/MassEntityManager.cpp
--- a/Engine/Source/Runtime/MassEntity/Private/MassEntityManager.cpp
+++ b/Engine/Source/Runtime/MassEntity/Private/MassEntityManager.cpp
@@ -109,8 +109,8 @@ void FMassEntityManager::FEntityCreationContext::AppendEntities(const TConstArra
 			CollectionCreationDuplicatesHandling = FMassArchetypeEntityCollection::FoldDuplicates;
 			MarkDirty();
 		}
-		// else, if there are no entities the resulting state will bE "dirty" by design 
-		ensureMsgf(EntityCollections.IsEmpty(), TEXT("Having a non-empty array of entity collecitons is unexpected at this point!"));
+		// else, if there are no entities the resulting state will be "dirty" by design
+		ensureMsgf(EntityCollections.IsEmpty(), TEXT("Having a non-empty array of entity collections is unexpected at this point!"));
 
 		CreatedEntities.Append(EntitiesToAppend);
 		ensure(IsDirty());
@@ -712,6 +712,11 @@ FMassEntityHandle FMassEntityManager::CreateEntity(TConstArrayView<FInstancedStr
 	check(ArchetypeHandle.IsValid());
 
 	const FMassEntityHandle Entity = ReserveEntity();
+
+	// Using a creation context to prevent InternalBuildEntity from notifying observers before we set fragments data
+	const TSharedRef<FEntityCreationContext> CreationContext = GetOrMakeCreationContext();
+	CreationContext->AppendEntities({Entity});
+
 	InternalBuildEntity(Entity, ArchetypeHandle, SharedFragmentValues);
 
 	FMassArchetypeData* CurrentArchetype = GetEntityStorageInterface().GetArchetype(Entity.Index);
@@ -764,6 +769,10 @@ void FMassEntityManager::BuildEntity(FMassEntityHandle Entity, TConstArrayView<F
 	const FMassArchetypeHandle& ArchetypeHandle = CreateArchetype(Composition);
 	check(ArchetypeHandle.IsValid());
 
+	// Using a creation context to prevent InternalBuildEntity from notifying observers before we set fragments data
+	const TSharedRef<FEntityCreationContext> CreationContext = GetOrMakeCreationContext();
+	CreationContext->AppendEntities({Entity});
+
 	InternalBuildEntity(Entity, ArchetypeHandle, SharedFragmentValues);
 
 	FMassArchetypeData* CurrentArchetype = GetEntityStorageInterface().GetArchetype(Entity.Index);
@@ -825,7 +834,7 @@ TSharedRef<FMassEntityManager::FEntityCreationContext> FMassEntityManager::Batch
 	UE::Mass::Private::ConvertArchetypelessSubchunksIntoEntityHandles(EncodedEntitiesWithPayload.GetEntityCollection().GetRanges(), EntityHandles);
 
 	// since the handles encoded via FMassArchetypeEntityCollectionWithPayload miss the SerialNumber we need to update it
-	// before passing over the the new archetype. Thankfully we need to iterate over all the entity handles anyway
+	// before passing over the new archetype. Thankfully we need to iterate over all the entity handles anyway
 	// to update the manager's information on these entities (stored in FMassEntityManager::Entities)
 	for (FMassEntityHandle& Entity : EntityHandles)
 	{
