Reduce fixed memory overhead from Mass with concurrent reserve enabled but Mass is not used
* Lazy initialize the first page of FEntityData and, by extension, the free index list

diff --git a/Engine/Source/Runtime/MassEntity/Private/MassEntityManager.cpp b/Engine/Source/Runtime/MassEntity/Private/MassEntityManager.cpp
--- a/Engine/Source/Runtime/MassEntity/Private/MassEntityManager.cpp
+++ b/Engine/Source/Runtime/MassEntity/Private/MassEntityManager.cpp
@@ -1,6 +1,7 @@
 // Copyright Epic Games, Inc. All Rights Reserved.
 
 #include "MassEntityManager.h"
+#include "MassEntityManagerConstants.h"
 #include "MassArchetypeData.h"
 #include "MassCommandBuffer.h"
 #include "MassEntityManagerStorage.h"
@@ -39,9 +40,6 @@ namespace UE::Mass::Private
 			}
 		}
 	}
-
-	// Index 0 is a sentinel for an Empty/Unset EntityHandle
-	constexpr int32 InvalidEntityIndex = 0;
 }
 
 //-----------------------------------------------------------------------------
@@ -293,6 +291,7 @@ namespace UE::Mass::Private
 		void operator()(const FMassEntityManager_InitParams_SingleThreaded& Params)
 		{
 			EntityStorage->Emplace<UE::Mass::FSingleThreadedEntityStorage>();
+			EntityStorage->Get<FSingleThreadedEntityStorage>().Initialize(Params);
 		}
 		void operator()(const FMassEntityManager_InitParams_Concurrent& Params)
 		{
@@ -319,10 +318,6 @@ void FMassEntityManager::Initialize(const FMassEntityManagerStorageInitParams& I
 
 	Visit(UE::Mass::Private::FEntityStorageInitializer{&EntityStorage}, InitializationParams);
 
-	// Index 0 is reserved so we can treat that index as an invalid entity handle
-	const FMassEntityHandle SentinelEntity = GetEntityStorageInterface().AcquireOne();
-	check(SentinelEntity.Index == UE::Mass::Private::InvalidEntityIndex);
-
 	for (TSharedPtr<FMassCommandBuffer>& CommandBuffer : DeferredCommandBuffers)
 	{
 		CommandBuffer = MakeShareable(new FMassCommandBuffer());
diff --git a/Engine/Source/Runtime/MassEntity/Private/MassEntityManagerConstants.h b/Engine/Source/Runtime/MassEntity/Private/MassEntityManagerConstants.h
new file mode 100644
--- /dev/null
+++ b/Engine/Source/Runtime/MassEntity/Private/MassEntityManagerConstants.h
@@ -0,0 +1,11 @@
+﻿// Copyright Epic Games, Inc. All Rights Reserved.
+
+#pragma once
+
+#include "HAL/Platform.h"
+
+namespace UE::Mass::Private
+{
+	// Index 0 is a sentinel for an Empty/Unset EntityHandle
+	constexpr int32 InvalidEntityIndex = 0;
+};
\ No newline at end of file
diff --git a/Engine/Source/Runtime/MassEntity/Private/MassEntityManagerStorage.cpp b/Engine/Source/Runtime/MassEntity/Private/MassEntityManagerStorage.cpp
--- a/Engine/Source/Runtime/MassEntity/Private/MassEntityManagerStorage.cpp
+++ b/Engine/Source/Runtime/MassEntity/Private/MassEntityManagerStorage.cpp
@@ -2,6 +2,7 @@
 
 #include "MassEntityManagerStorage.h"
 
+#include "MassEntityManagerConstants.h"
 #include "MassEntityTypes.h"
 #include "Templates/SharedPointer.h"
 
@@ -30,7 +31,14 @@ namespace UE::Mass
 	//-----------------------------------------------------------------------------
 	// FSingleThreadedEntityStorage
 	//-----------------------------------------------------------------------------
-	
+
+	void FSingleThreadedEntityStorage::Initialize(const FMassEntityManager_InitParams_SingleThreaded&)
+	{
+		// Index 0 is reserved so we can treat that index as an invalid entity handle
+		const FMassEntityHandle SentinelEntity = AcquireOne();
+		check(SentinelEntity.Index == UE::Mass::Private::InvalidEntityIndex);
+	}
+
 	FMassArchetypeData* FSingleThreadedEntityStorage::GetArchetype(int32 Index)
 	{
 		return Entities[Index].CurrentArchetype.Get();
@@ -361,15 +369,28 @@ namespace UE::Mass
 		EntityPages[PageCount] = Page;
 		++PageCount;
 
-		const int32 NewEntityIndexStart = NewPageIndex * MaxEntitiesPerPage;
+		// Somewhat tricksy thing here to be aware of
+		// MassEntityManager expects the very first allocated entity to be at index 0
+		static_assert(UE::Mass::Private::InvalidEntityIndex == 0, "Free Entity list algorithm depends on InvalidEntityIndex being 0");
+		int32 NewEntityIndexStart;
+		if (LIKELY(NewPageIndex != 0))
+		{
+			NewEntityIndexStart = NewPageIndex * MaxEntitiesPerPage;
+		}
+		else
+		{
+			NewEntityIndexStart = 1;
+			// Allocate the 0th entity. It will always be the sentinel entity that InvalidEntityIndex points to.
+			FEntityData* SentinelEntity = new (Page + UE::Mass::Private::InvalidEntityIndex) FEntityData();
+			SentinelEntity->bIsAllocated = 1;
+			++SentinelEntity->GenerationId;
+		}
+		
 		const int32 NewEntityIndexEnd = (NewPageIndex + 1) * MaxEntitiesPerPage;
 
-		EntityFreeIndexList.Reserve(MaxEntitiesPerPage);
+		EntityFreeIndexList.Reserve(NewEntityIndexEnd - NewEntityIndexStart);
 
-		// Somewhat tricksy thing here to be aware of
-		// MassEntityManager expects the very first allocated entity to be at index 0
-		// However, EntityFreeIndexList.Pop() will return the last one added to the list
-		// Therefore, populate the free list backwards
+		// Push free entities indices onto the stack backwards so new entities pop off in order
 		for (int32 NewEntityIndex = NewEntityIndexEnd - 1; NewEntityIndex >= NewEntityIndexStart; --NewEntityIndex)
 		{
 			// Setup the free list
diff --git a/Engine/Source/Runtime/MassEntity/Public/MassEntityManagerStorage.h b/Engine/Source/Runtime/MassEntity/Public/MassEntityManagerStorage.h
--- a/Engine/Source/Runtime/MassEntity/Public/MassEntityManagerStorage.h
+++ b/Engine/Source/Runtime/MassEntity/Public/MassEntityManagerStorage.h
@@ -109,6 +109,7 @@ namespace UE::Mass
 	class FSingleThreadedEntityStorage final : public IEntityStorageInterface
 	{
 	public:
+		void Initialize(const FMassEntityManager_InitParams_SingleThreaded&);
 		virtual FMassArchetypeData* GetArchetype(int32 Index) override;
 		virtual const FMassArchetypeData* GetArchetype(int32 Index) const override;
 		virtual TSharedPtr<FMassArchetypeData>& GetArchetypeAsShared(int32 Index) override;
