[Mass] Pulled in external Mass changes.
* FMassEntityView::HasTag(const UScriptStruct&) made public
* FMassArchetypeSharedFragmentValues.ReplaceSharedFragments added, allows replacing existing elements of FMassArchetypeSharedFragmentValues
* UMassActorSpawnerSubsystem.DestroyActor made virtual
* FMassArchetypeData.GetOrAddChunk added and existing code updated to reuse it
* FMassArchetypeData.SetSharedFragmentsData added, allows changing shared fragment values of a single entity, moving it to the appropriate chunk in the process

diff --git a/Engine/Plugins/Runtime/MassGameplay/Source/MassActors/Public/MassActorSpawnerSubsystem.h b/Engine/Plugins/Runtime/MassGameplay/Source/MassActors/Public/MassActorSpawnerSubsystem.h
--- a/Engine/Plugins/Runtime/MassGameplay/Source/MassActors/Public/MassActorSpawnerSubsystem.h
+++ b/Engine/Plugins/Runtime/MassGameplay/Source/MassActors/Public/MassActorSpawnerSubsystem.h
@@ -206,7 +206,7 @@ public:
 	 * @param Actor to destroy
 	 * @param bImmediate to do the destruction immediately, otherwise will be queued up for later
 	 */
-	void DestroyActor(AActor* Actor, bool bImmediate = false);
+	virtual void DestroyActor(AActor* Actor, bool bImmediate = false);
 	
 	void EnableActorPooling();
 	void DisableActorPooling();
diff --git a/Engine/Source/Runtime/MassEntity/Internal/MassArchetypeData.h b/Engine/Source/Runtime/MassEntity/Internal/MassArchetypeData.h
--- a/Engine/Source/Runtime/MassEntity/Internal/MassArchetypeData.h
+++ b/Engine/Source/Runtime/MassEntity/Internal/MassArchetypeData.h
@@ -400,6 +400,17 @@ protected:
 	void BindConstSharedFragmentRequirements(FMassExecutionContext& RunContext, const FMassArchetypeSharedFragmentValues& SharedFragmentValues, const FMassFragmentIndicesMapping& ChunkFragmentsMapping);
 	void BindSharedFragmentRequirements(FMassExecutionContext& RunContext, FMassArchetypeSharedFragmentValues& SharedFragmentValues, const FMassFragmentIndicesMapping& ChunkFragmentsMapping);
 
+	/**
+	 * The function first creates new FMassArchetypeSharedFragmentValues instance combining existing values
+	 * and the contents of SharedFragmentValueOverrides. Then that is used to find the target chunk for Entity,
+	 * and if one cannot be found a new one will be created. 
+	 * @param SharedFragmentValueOverrides is expected to contain only instances of types already
+	 *    present in given archetypes FMassArchetypeSharedFragmentValues
+	 */
+	void SetSharedFragmentsData(const FMassEntityHandle Entity, TConstArrayView<FSharedStruct> SharedFragmentValueOverrides);
+
+	FMassArchetypeChunk& GetOrAddChunk(const FMassArchetypeSharedFragmentValues& SharedFragmentValues, int32& OutAbsoluteIndex, int32& OutIndexWithinChunk);
+	
 private:
 	int32 AddEntityInternal(FMassEntityHandle Entity, const FMassArchetypeSharedFragmentValues& InSharedFragmentValues);
 	void RemoveEntityInternal(const int32 AbsoluteIndex);
diff --git a/Engine/Source/Runtime/MassEntity/Private/MassArchetypeData.cpp b/Engine/Source/Runtime/MassEntity/Private/MassArchetypeData.cpp
--- a/Engine/Source/Runtime/MassEntity/Private/MassArchetypeData.cpp
+++ b/Engine/Source/Runtime/MassEntity/Private/MassArchetypeData.cpp
@@ -187,6 +187,22 @@ int32 FMassArchetypeData::AddEntityInternal(FMassEntityHandle Entity, const FMas
 
 	int32 IndexWithinChunk = 0;
 	int32 AbsoluteIndex = 0;
+
+	FMassArchetypeChunk& DestinationChunk = GetOrAddChunk(SharedFragmentValues, AbsoluteIndex, IndexWithinChunk);
+	DestinationChunk.AddInstance();
+
+	// Add to the table and map
+	EntityMap.Add(Entity.Index, AbsoluteIndex);
+	DestinationChunk.GetEntityArrayElementRef(EntityListOffsetWithinChunk, IndexWithinChunk) = Entity;
+
+	return AbsoluteIndex;
+}
+
+FMassArchetypeChunk& FMassArchetypeData::GetOrAddChunk(const FMassArchetypeSharedFragmentValues& SharedFragmentValues, int32& OutAbsoluteIndex, int32& OutIndexWithinChunk)
+{
+	OutAbsoluteIndex = 0;
+	OutIndexWithinChunk = 0;
+
 	int32 ChunkIndex = 0;
 	int32 EmptyChunkIndex = INDEX_NONE;
 	int32 EmptyAbsoluteIndex = INDEX_NONE;
@@ -202,20 +218,18 @@ int32 FMassArchetypeData::AddEntityInternal(FMassEntityHandle Entity, const FMas
 			if (EmptyChunkIndex == INDEX_NONE)
 			{
 				EmptyChunkIndex = ChunkIndex;
-				EmptyAbsoluteIndex = AbsoluteIndex;
+				EmptyAbsoluteIndex = OutAbsoluteIndex;
 			}
 		}
 		else if (Chunk.GetNumInstances() < NumEntitiesPerChunk && Chunk.GetSharedFragmentValues().IsEquivalent(SharedFragmentValues))
 		{
-			IndexWithinChunk = Chunk.GetNumInstances();
-			AbsoluteIndex += IndexWithinChunk;
-
-			Chunk.AddInstance();
+			OutIndexWithinChunk = Chunk.GetNumInstances();
+			OutAbsoluteIndex += OutIndexWithinChunk;
 
 			DestinationChunk = &Chunk;
 			break;
 		}
-		AbsoluteIndex += NumEntitiesPerChunk;
+		OutAbsoluteIndex += NumEntitiesPerChunk;
 		++ChunkIndex;
 	}
 
@@ -226,22 +240,16 @@ int32 FMassArchetypeData::AddEntityInternal(FMassEntityHandle Entity, const FMas
 		{
 			DestinationChunk = &Chunks[EmptyChunkIndex];
 			DestinationChunk->Recycle(ChunkFragmentsTemplate, SharedFragmentValues);
-			AbsoluteIndex = EmptyAbsoluteIndex;
+			OutAbsoluteIndex = EmptyAbsoluteIndex;
 		}
 		else
 		{
 			DestinationChunk = &Chunks.Emplace_GetRef(GetChunkAllocSize(), ChunkFragmentsTemplate, SharedFragmentValues);
 		}
-
-		check(DestinationChunk);
-		DestinationChunk->AddInstance();
 	}
 
-	// Add to the table and map
-	EntityMap.Add(Entity.Index, AbsoluteIndex);
-	DestinationChunk->GetEntityArrayElementRef(EntityListOffsetWithinChunk, IndexWithinChunk) = Entity;
-
-	return AbsoluteIndex;
+	check(DestinationChunk);
+	return *DestinationChunk;
 }
 
 void FMassArchetypeData::RemoveEntity(FMassEntityHandle Entity)
@@ -368,6 +376,40 @@ void* FMassArchetypeData::GetFragmentDataForEntity(const UScriptStruct* Fragment
 	return nullptr;
 }
 
+void FMassArchetypeData::SetSharedFragmentsData(const FMassEntityHandle Entity, TConstArrayView<FSharedStruct> SharedFragmentValueOverrides)
+{
+	// Gets the current chunk where the entity is located
+	const int32 OldAbsoluteIndex = EntityMap.FindChecked(Entity.Index);
+	const int32 OldChunkIndex = OldAbsoluteIndex / NumEntitiesPerChunk;
+	const int32 OldIndexWithinChunk = OldAbsoluteIndex % NumEntitiesPerChunk;
+	const FMassArchetypeChunk& OldChunk = Chunks[OldChunkIndex];
+
+	// Gets or adds a new chunk that will hold the new entity with the new shared values
+	FMassArchetypeSharedFragmentValues NewSharedFragmentValues(OldChunk.GetSharedFragmentValues());
+	NewSharedFragmentValues.ReplaceSharedFragments(SharedFragmentValueOverrides);
+	NewSharedFragmentValues.Sort();
+
+	int32 NewAbsoluteIndex = 0;
+	int32 NewIndexWithinChunk = 0;
+	FMassArchetypeChunk& NewChunk = GetOrAddChunk(NewSharedFragmentValues, NewAbsoluteIndex, NewIndexWithinChunk);
+
+	if (ensureMsgf(&NewChunk != &OldChunk, TEXT("Found target chunk is the same as the source chunk. Probably "
+		"caused by setting shared fragment values resulted in no change, meaning the target values equal the source values")))
+	{
+		NewChunk.AddInstance();
+
+		// Update the new entity in the table and map
+		EntityMap[Entity.Index] = NewAbsoluteIndex;
+		NewChunk.GetEntityArrayElementRef(EntityListOffsetWithinChunk, NewIndexWithinChunk) = Entity;
+		
+		// Move the current entity fragments into the new chunk
+		MoveFragmentsToNewLocationInternal({ OldChunk.GetRawMemory(), OldIndexWithinChunk }, { NewChunk.GetRawMemory(), NewIndexWithinChunk }, 1);
+
+		// Clean up the old chunk
+		RemoveEntityInternal(OldAbsoluteIndex);
+	}
+}
+
 void FMassArchetypeData::SetFragmentsData(const FMassEntityHandle Entity, TArrayView<const FInstancedStruct> FragmentInstances)
 {
 	FMassRawEntityInChunkData InternalIndex = MakeEntityHandle(Entity);
diff --git a/Engine/Source/Runtime/MassEntity/Private/MassEntityTypes.cpp b/Engine/Source/Runtime/MassEntity/Private/MassEntityTypes.cpp
--- a/Engine/Source/Runtime/MassEntity/Private/MassEntityTypes.cpp
+++ b/Engine/Source/Runtime/MassEntity/Private/MassEntityTypes.cpp
@@ -54,6 +54,28 @@ FSharedStruct FMassArchetypeSharedFragmentValues::AddSharedFragment(const FShare
 	return StructInstance;
 }
 
+void FMassArchetypeSharedFragmentValues::ReplaceSharedFragments(TConstArrayView<FSharedStruct> Fragments)
+{
+	DirtyHashCache();
+	for (const FSharedStruct& NewFragment : Fragments)
+	{
+		const UScriptStruct* NewFragScriptStruct = NewFragment.GetScriptStruct();
+		check(NewFragScriptStruct);
+
+		bool bEntryFound = false;
+		for (FSharedStruct& MyFragment : SharedFragments)
+		{
+			if (MyFragment.GetScriptStruct() == NewFragScriptStruct)
+			{
+				MyFragment = NewFragment;
+				bEntryFound = true;
+				break;
+			}
+		}
+		ensureMsgf(bEntryFound, TEXT("Existing fragment of type %s could not be found"), *GetNameSafe(NewFragScriptStruct));
+	}
+}
+
 uint32 FMassArchetypeSharedFragmentValues::CalculateHash() const
 {
 	if (!testableEnsureMsgf(bSorted, TEXT("Expecting the containers to be sorted for the hash caluclation to be consistent")))
diff --git a/Engine/Source/Runtime/MassEntity/Public/MassEntityTypes.h b/Engine/Source/Runtime/MassEntity/Public/MassEntityTypes.h
--- a/Engine/Source/Runtime/MassEntity/Public/MassEntityTypes.h
+++ b/Engine/Source/Runtime/MassEntity/Public/MassEntityTypes.h
@@ -440,6 +440,14 @@ struct MASSENTITY_API FMassArchetypeSharedFragmentValues
 	 */
 	FSharedStruct AddSharedFragment(const FSharedStruct& Fragment);
 
+	/**
+	 * Finds instances of fragment types given by Fragments and replaces their values with contents of respective
+	 * element of Fragments.
+	 * Note that it's callers responsibility to ensure every fragment type in Fragments already has an instance in
+	 * this FMassArchetypeSharedFragmentValues instance. Failing that assumption will result in ensure failure. 
+	 */
+	void ReplaceSharedFragments(TConstArrayView<FSharedStruct> Fragments);
+
 	/** 
 	 * Appends contents of Other to `this` instance. All common fragments will get overridden with values in Other.
 	 * Note that changing a fragments "role" (being const or non-const) is not supported and the function will fail an
