[Mass] Updated FMassArchetypeData::CompactEntities to reuse MoveFragmentsToNewLocationInternal rather than duplicate code.

diff --git a/Engine/Source/Runtime/MassEntity/Private/MassArchetypeData.cpp b/Engine/Source/Runtime/MassEntity/Private/MassArchetypeData.cpp
--- a/Engine/Source/Runtime/MassEntity/Private/MassArchetypeData.cpp
+++ b/Engine/Source/Runtime/MassEntity/Private/MassArchetypeData.cpp
@@ -617,13 +617,8 @@ int32 FMassArchetypeData::CompactEntities(const double TimeAllowed)
 			const int32 ToIndex = ChunkToFill->GetNumInstances();
 			check(NumberOfEntitiesToMove > 0);
 
-			for (const FMassArchetypeFragmentConfig& FragmentConfig : FragmentConfigs)
-			{
-				void* FromFragmentPtr = FragmentConfig.GetFragmentData(ChunkToEmpty->GetRawMemory(), FromIndex);
-				void* ToFragmentPtr = FragmentConfig.GetFragmentData(ChunkToFill->GetRawMemory(), ToIndex);
-				// Move all entries
-				FMemory::Memcpy(ToFragmentPtr, FromFragmentPtr, FragmentConfig.FragmentType->GetStructureSize() * NumberOfEntitiesToMove);
-			}
+			MoveFragmentsToNewLocationInternal({ChunkToFill->GetRawMemory(), ToIndex}, {ChunkToEmpty->GetRawMemory(), FromIndex}
+				, NumberOfEntitiesToMove);
 
 			FMassEntityHandle* FromEntity = &ChunkToEmpty->GetEntityArrayElementRef(EntityListOffsetWithinChunk, FromIndex);
 			FMassEntityHandle* ToEntity = &ChunkToFill->GetEntityArrayElementRef(EntityListOffsetWithinChunk, ToIndex);
