[MassEntityManager] refactored GetOrCreateConstSharedFragment and GetOrCreateSharedFragment

diff --git a/Engine/Plugins/AI/MassCrowd/Source/MassCrowd/Private/MassCrowdServerRepresentationTrait.cpp b/Engine/Plugins/AI/MassCrowd/Source/MassCrowd/Private/MassCrowdServerRepresentationTrait.cpp
--- a/Engine/Plugins/AI/MassCrowd/Source/MassCrowd/Private/MassCrowdServerRepresentationTrait.cpp
+++ b/Engine/Plugins/AI/MassCrowd/Source/MassCrowd/Private/MassCrowdServerRepresentationTrait.cpp
@@ -42,8 +42,7 @@ void UMassCrowdServerRepresentationTrait::BuildTemplate(FMassEntityTemplateBuild
 
 	FMassRepresentationSubsystemSharedFragment SubsystemSharedFragment;
 	SubsystemSharedFragment.RepresentationSubsystem = RepresentationSubsystem;
-	uint32 SubsystemHash = UE::StructUtils::GetStructCrc32(FConstStructView::Make(SubsystemSharedFragment));
-	FSharedStruct SubsystemFragment = EntityManager.GetOrCreateSharedFragmentByHash<FMassRepresentationSubsystemSharedFragment>(SubsystemHash, SubsystemSharedFragment);
+	FSharedStruct SubsystemFragment = EntityManager.GetOrCreateSharedFragment<FMassRepresentationSubsystemSharedFragment>(SubsystemSharedFragment);
 	BuildContext.AddSharedFragment(SubsystemFragment);
 
 	FConstSharedStruct ParamsFragment = EntityManager.GetOrCreateConstSharedFragment(Params);
diff --git a/Engine/Plugins/Runtime/InstancedActors/Source/InstancedActors/Private/InstancedActorsData.cpp b/Engine/Plugins/Runtime/InstancedActors/Source/InstancedActors/Private/InstancedActorsData.cpp
--- a/Engine/Plugins/Runtime/InstancedActors/Source/InstancedActors/Private/InstancedActorsData.cpp
+++ b/Engine/Plugins/Runtime/InstancedActors/Source/InstancedActors/Private/InstancedActorsData.cpp
@@ -477,8 +477,7 @@ void UInstancedActorsData::DespawnEntities()
 
 	FInstancedActorsDataSharedFragment ManagerFragment;
 	ManagerFragment.InstanceData = this;
-	const uint32 FragmentHash = UE::StructUtils::GetStructCrc32(FConstStructView::Make(ManagerFragment));
-	FSharedStruct SharedFragmentInstance = MassEntityManager.GetOrCreateSharedFragmentByHash<FInstancedActorsDataSharedFragment>(FragmentHash, ManagerFragment);
+	FSharedStruct SharedFragmentInstance = MassEntityManager.GetOrCreateSharedFragment<FInstancedActorsDataSharedFragment>(ManagerFragment);
 
 	FInstancedActorsDataSharedFragment* AsSharedManagerFragment = SharedFragmentInstance.GetPtr<FInstancedActorsDataSharedFragment>();
 	if (ensure(AsSharedManagerFragment) && AsSharedManagerFragment->InstanceData == this)
diff --git a/Engine/Plugins/Runtime/InstancedActors/Source/InstancedActors/Private/InstancedActorsVisualizationTrait.cpp b/Engine/Plugins/Runtime/InstancedActors/Source/InstancedActors/Private/InstancedActorsVisualizationTrait.cpp
--- a/Engine/Plugins/Runtime/InstancedActors/Source/InstancedActors/Private/InstancedActorsVisualizationTrait.cpp
+++ b/Engine/Plugins/Runtime/InstancedActors/Source/InstancedActors/Private/InstancedActorsVisualizationTrait.cpp
@@ -87,9 +87,7 @@ void UInstancedActorsVisualizationTrait::BuildTemplate(FMassEntityTemplateBuildC
 
 	FInstancedActorsDataSharedFragment ManagerSharedFragment;
 	ManagerSharedFragment.InstanceData = InstanceData;
-	const uint32 SubsystemHash = UE::StructUtils::GetStructCrc32(FConstStructView::Make(ManagerSharedFragment));
-	FSharedStruct SubsystemFragment = EntityManager.GetOrCreateSharedFragmentByHash<FInstancedActorsDataSharedFragment>(SubsystemHash, ManagerSharedFragment);
-
+	FSharedStruct SubsystemFragment = EntityManager.GetOrCreateSharedFragment<FInstancedActorsDataSharedFragment>(ManagerSharedFragment);
 
 	FInstancedActorsDataSharedFragment* AsShared = SubsystemFragment.GetPtr<FInstancedActorsDataSharedFragment>();
 	if (ensure(AsShared))
diff --git a/Engine/Plugins/Runtime/MassGameplay/Source/MassGameplayTestSuite/Private/MassEntityTemplateTest.cpp b/Engine/Plugins/Runtime/MassGameplay/Source/MassGameplayTestSuite/Private/MassEntityTemplateTest.cpp
--- a/Engine/Plugins/Runtime/MassGameplay/Source/MassGameplayTestSuite/Private/MassEntityTemplateTest.cpp
+++ b/Engine/Plugins/Runtime/MassGameplay/Source/MassGameplayTestSuite/Private/MassEntityTemplateTest.cpp
@@ -64,11 +64,11 @@ struct FEntityTemplate_Composition : FEntityTemplateBase
 		TemplateData.AddFragment<FTestFragment_Float>();
 		TemplateData.AddChunkFragment<FTestChunkFragment_Int>();
 		{
-			FSharedStruct SharedFragment = EntityManager->GetOrCreateSharedFragmentByHash<FTestSharedFragment_Int>(0);
+			FSharedStruct SharedFragment = EntityManager->GetOrCreateSharedFragment<FTestSharedFragment_Int>(0);
 			TemplateData.AddSharedFragment(SharedFragment);
 		}
 		{
-			FSharedStruct SharedFragment = EntityManager->GetOrCreateSharedFragmentByHash<FTestSharedFragment_Float>(1);
+			FSharedStruct SharedFragment = EntityManager->GetOrCreateSharedFragment<FTestSharedFragment_Float>(1);
 			TemplateData.AddSharedFragment(SharedFragment);
 		}
 		TemplateData.Sort();
diff --git a/Engine/Plugins/Runtime/MassGameplay/Source/MassLOD/Private/MassLODTrait.cpp b/Engine/Plugins/Runtime/MassGameplay/Source/MassLOD/Private/MassLODTrait.cpp
--- a/Engine/Plugins/Runtime/MassGameplay/Source/MassLOD/Private/MassLODTrait.cpp
+++ b/Engine/Plugins/Runtime/MassGameplay/Source/MassLOD/Private/MassLODTrait.cpp
@@ -50,8 +50,7 @@ void UMassSimulationLODTrait::BuildTemplate(FMassEntityTemplateBuildContext& Bui
 	FConstSharedStruct ParamsFragment = EntityManager.GetOrCreateConstSharedFragment(Params);
 	BuildContext.AddConstSharedFragment(ParamsFragment);
 
-	uint32 ParamsHash = UE::StructUtils::GetStructCrc32(FConstStructView::Make(Params));
-	FSharedStruct SharedFragment = EntityManager.GetOrCreateSharedFragmentByHash<FMassSimulationLODSharedFragment>(ParamsHash, Params);
+	FSharedStruct SharedFragment = EntityManager.GetOrCreateSharedFragment<FMassSimulationLODSharedFragment>(FConstStructView::Make(Params), Params);
 	BuildContext.AddSharedFragment(SharedFragment);
 
 	// Variable ticking from simulation LOD
@@ -63,8 +62,7 @@ void UMassSimulationLODTrait::BuildTemplate(FMassEntityTemplateBuildContext& Bui
 		FConstSharedStruct VariableTickParamsFragment = EntityManager.GetOrCreateConstSharedFragment(VariableTickParams);
 		BuildContext.AddConstSharedFragment(VariableTickParamsFragment);
 
-		uint32 VariableTickParamsHash = UE::StructUtils::GetStructCrc32(FConstStructView::Make(VariableTickParams));
-		FSharedStruct VariableTickSharedFragment = EntityManager.GetOrCreateSharedFragmentByHash<FMassSimulationVariableTickSharedFragment>(VariableTickParamsHash, VariableTickParams);
+		FSharedStruct VariableTickSharedFragment = EntityManager.GetOrCreateSharedFragment<FMassSimulationVariableTickSharedFragment>(FConstStructView::Make(VariableTickParams), VariableTickParams);
 		BuildContext.AddSharedFragment(VariableTickSharedFragment);
 	}
 }
diff --git a/Engine/Plugins/Runtime/MassGameplay/Source/MassLOD/Public/MassSimulationLOD.h b/Engine/Plugins/Runtime/MassGameplay/Source/MassLOD/Public/MassSimulationLOD.h
--- a/Engine/Plugins/Runtime/MassGameplay/Source/MassLOD/Public/MassSimulationLOD.h
+++ b/Engine/Plugins/Runtime/MassGameplay/Source/MassLOD/Public/MassSimulationLOD.h
@@ -159,15 +159,6 @@ struct MASSLOD_API FMassSimulationVariableTickSharedFragment : public FMassShare
 		}
 		return true;
 	}
-
-	/**
-	 * Custom hash function for GetOrCreateConstSharedFragment to use. This hashes the parameters of LODTickRateController
-	 * which are the values indicating the uniqueness of a given shared fragment
-	 */
-	friend FORCEINLINE uint32 GetTypeHash(const FMassSimulationVariableTickSharedFragment& SharedFragmentInstance)
-	{
-		return GetTypeHash(SharedFragmentInstance.LODTickRateController);
-	}
 };
 
 UCLASS(meta = (DisplayName = "Simulation LOD"))
diff --git a/Engine/Plugins/Runtime/MassGameplay/Source/MassReplication/Private/MassReplicationTrait.cpp b/Engine/Plugins/Runtime/MassGameplay/Source/MassReplication/Private/MassReplicationTrait.cpp
--- a/Engine/Plugins/Runtime/MassGameplay/Source/MassReplication/Private/MassReplicationTrait.cpp
+++ b/Engine/Plugins/Runtime/MassGameplay/Source/MassReplication/Private/MassReplicationTrait.cpp
@@ -33,17 +33,16 @@ void UMassReplicationTrait::BuildTemplate(FMassEntityTemplateBuildContext& Build
 	FConstSharedStruct ParamsFragment = EntityManager.GetOrCreateConstSharedFragment(Params);
 	BuildContext.AddConstSharedFragment(ParamsFragment);
 
-	uint32 ParamsHash = UE::StructUtils::GetStructCrc32(FConstStructView::Make(Params));
 	if (LIKELY(!BuildContext.IsInspectingData()))
 	{
 		check(ReplicationSubsystem);
-		FSharedStruct SharedFragment = EntityManager.GetOrCreateSharedFragmentByHash<FMassReplicationSharedFragment>(ParamsHash, *ReplicationSubsystem, Params);
+		FSharedStruct SharedFragment = EntityManager.GetOrCreateSharedFragment<FMassReplicationSharedFragment>(FConstStructView::Make(Params), *ReplicationSubsystem, Params);
 		BuildContext.AddSharedFragment(SharedFragment);
 	}
 	else
 	{
 		// in the investigation mode we only care about the fragment type
-		FSharedStruct SharedFragment = EntityManager.GetOrCreateSharedFragmentByHash<FMassReplicationSharedFragment>(ParamsHash);
+		FSharedStruct SharedFragment = EntityManager.GetOrCreateSharedFragment<FMassReplicationSharedFragment>();
 		BuildContext.AddSharedFragment(SharedFragment);
 	}
 }
\ No newline at end of file
diff --git a/Engine/Plugins/Runtime/MassGameplay/Source/MassRepresentation/Private/MassDistanceVisualizationTrait.cpp b/Engine/Plugins/Runtime/MassGameplay/Source/MassRepresentation/Private/MassDistanceVisualizationTrait.cpp
--- a/Engine/Plugins/Runtime/MassGameplay/Source/MassRepresentation/Private/MassDistanceVisualizationTrait.cpp
+++ b/Engine/Plugins/Runtime/MassGameplay/Source/MassRepresentation/Private/MassDistanceVisualizationTrait.cpp
@@ -59,8 +59,7 @@ void UMassDistanceVisualizationTrait::BuildTemplate(FMassEntityTemplateBuildCont
 
 	FMassRepresentationSubsystemSharedFragment SubsystemSharedFragment;
 	SubsystemSharedFragment.RepresentationSubsystem = RepresentationSubsystem;
-	uint32 SubsystemHash = UE::StructUtils::GetStructCrc32(FConstStructView::Make(SubsystemSharedFragment));
-	FSharedStruct SubsystemFragment = EntityManager.GetOrCreateSharedFragmentByHash<FMassRepresentationSubsystemSharedFragment>(SubsystemHash, SubsystemSharedFragment);
+	FSharedStruct SubsystemFragment = EntityManager.GetOrCreateSharedFragment<FMassRepresentationSubsystemSharedFragment>(SubsystemSharedFragment);
 	BuildContext.AddSharedFragment(SubsystemFragment);
 
 	if (!Params.RepresentationActorManagementClass && !BuildContext.IsInspectingData())
@@ -85,8 +84,7 @@ void UMassDistanceVisualizationTrait::BuildTemplate(FMassEntityTemplateBuildCont
 	FConstSharedStruct LODParamsFragment = EntityManager.GetOrCreateConstSharedFragment(LODParams);
 	BuildContext.AddConstSharedFragment(LODParamsFragment);
 
-	uint32 LODParamsHash = UE::StructUtils::GetStructCrc32(FConstStructView::Make(LODParams));
-	FSharedStruct LODSharedFragment = EntityManager.GetOrCreateSharedFragmentByHash<FMassDistanceLODSharedFragment>(LODParamsHash, LODParams);
+	FSharedStruct LODSharedFragment = EntityManager.GetOrCreateSharedFragment<FMassDistanceLODSharedFragment>(FConstStructView::Make(LODParams), LODParams);
 	BuildContext.AddSharedFragment(LODSharedFragment);
 
 	BuildContext.AddFragment<FMassRepresentationLODFragment>();
diff --git a/Engine/Plugins/Runtime/MassGameplay/Source/MassRepresentation/Private/MassVisualizationTrait.cpp b/Engine/Plugins/Runtime/MassGameplay/Source/MassRepresentation/Private/MassVisualizationTrait.cpp
--- a/Engine/Plugins/Runtime/MassGameplay/Source/MassRepresentation/Private/MassVisualizationTrait.cpp
+++ b/Engine/Plugins/Runtime/MassGameplay/Source/MassRepresentation/Private/MassVisualizationTrait.cpp
@@ -78,8 +78,7 @@ void UMassVisualizationTrait::BuildTemplate(FMassEntityTemplateBuildContext& Bui
 
 	FMassRepresentationSubsystemSharedFragment SubsystemSharedFragment;
 	SubsystemSharedFragment.RepresentationSubsystem = RepresentationSubsystem;
-	uint32 SubsystemHash = UE::StructUtils::GetStructCrc32(FConstStructView::Make(SubsystemSharedFragment));
-	FSharedStruct SubsystemFragment = EntityManager.GetOrCreateSharedFragmentByHash<FMassRepresentationSubsystemSharedFragment>(SubsystemHash, SubsystemSharedFragment);
+	FSharedStruct SubsystemFragment = EntityManager.GetOrCreateSharedFragment<FMassRepresentationSubsystemSharedFragment>(SubsystemSharedFragment);
 	BuildContext.AddSharedFragment(SubsystemFragment);
 
 	if (!Params.RepresentationActorManagementClass)
@@ -125,8 +124,7 @@ void UMassVisualizationTrait::BuildTemplate(FMassEntityTemplateBuildContext& Bui
 	FConstSharedStruct LODParamsFragment = EntityManager.GetOrCreateConstSharedFragment(LODParams);
 	BuildContext.AddConstSharedFragment(LODParamsFragment);
 
-	uint32 LODParamsHash = UE::StructUtils::GetStructCrc32(FConstStructView::Make(LODParams));
-	FSharedStruct LODSharedFragment = EntityManager.GetOrCreateSharedFragmentByHash<FMassVisualizationLODSharedFragment>(LODParamsHash, LODParams);
+	FSharedStruct LODSharedFragment = EntityManager.GetOrCreateSharedFragment<FMassVisualizationLODSharedFragment>(FConstStructView::Make(LODParams), LODParams);
 	BuildContext.AddSharedFragment(LODSharedFragment);
 
 	BuildContext.AddFragment<FMassRepresentationLODFragment>();
diff --git a/Engine/Source/Developer/MassEntityTestSuite/Private/MassSharedFragmentsTest.cpp b/Engine/Source/Developer/MassEntityTestSuite/Private/MassSharedFragmentsTest.cpp
--- a/Engine/Source/Developer/MassEntityTestSuite/Private/MassSharedFragmentsTest.cpp
+++ b/Engine/Source/Developer/MassEntityTestSuite/Private/MassSharedFragmentsTest.cpp
@@ -636,6 +636,133 @@ struct FSharedFragmentValues_TypeEquivalency : FEntityTestBase
 };
 IMPLEMENT_AI_INSTANT_TEST(FSharedFragmentValues_TypeEquivalency, "System.Mass.SharedFragments.TypeEquivalency");
 
+struct FSharedFragmentValues_GetOrCreateWithArgs : FEntityTestBase
+{
+	virtual bool InstantTest() override
+	{
+		constexpr int32 ConstIntValueOne = 1;
+		constexpr int32 ConstIntValueTwo = 2;
+
+		const FSharedStruct SharedFragment1 = EntityManager->GetOrCreateSharedFragment<FTestSharedFragment_Int>(/*Args*/ConstIntValueOne);
+		const FSharedStruct SharedFragment2 = EntityManager->GetOrCreateSharedFragment<FTestSharedFragment_Int>(/*Args*/ConstIntValueOne);
+		const FSharedStruct SharedFragment3 = EntityManager->GetOrCreateSharedFragment<FTestSharedFragment_Int>(/*Args*/ConstIntValueTwo);
+
+		AITEST_EQUAL("Shared fragments created for same struct type using same constructor value should share memory", SharedFragment1, SharedFragment2);
+		AITEST_EQUAL("Value in shared struct should be the same as the argument provided to GetOrCreateSharedFragment", SharedFragment1.Get<FTestSharedFragment_Int>().Value, ConstIntValueOne);
+
+		AITEST_NOT_EQUAL("Shared fragments created for same struct type using different constructor values should not share memory", SharedFragment1, SharedFragment3);
+		AITEST_EQUAL("Value in shared struct should be the same as the argument provided to GetOrCreateSharedFragment", SharedFragment3.Get<FTestSharedFragment_Int>().Value, ConstIntValueTwo);
+
+		return true;
+	}
+};
+IMPLEMENT_AI_INSTANT_TEST(FSharedFragmentValues_GetOrCreateWithArgs, "System.Mass.SharedFragments.GetOrCreate.WithArgs");
+
+struct FSharedFragmentValues_GetOrCreateWithStruct : FEntityTestBase
+{
+	virtual bool InstantTest() override
+	{
+		constexpr int32 ConstIntValueOne = 1;
+		constexpr int32 ConstIntValueTwo = 2;
+
+		const FTestSharedFragment_Int TestSharedFragment_Int1(ConstIntValueOne);
+		const FTestSharedFragment_Int TestSharedFragment_Int2(ConstIntValueOne);
+		const FTestSharedFragment_Int TestSharedFragment_Int3(ConstIntValueTwo);
+		const FSharedStruct SharedFragment1 = EntityManager->GetOrCreateSharedFragment(TestSharedFragment_Int1);
+		const FSharedStruct SharedFragment2 = EntityManager->GetOrCreateSharedFragment(TestSharedFragment_Int2);
+		const FSharedStruct SharedFragment3 = EntityManager->GetOrCreateSharedFragment(TestSharedFragment_Int3);
+
+		AITEST_EQUAL("Shared fragments created for same struct type using same constructor value should share memory", SharedFragment1, SharedFragment2);
+		AITEST_EQUAL("Value in shared struct should be the same as the argument provided to GetOrCreateSharedFragment", SharedFragment1.Get<FTestSharedFragment_Int>().Value, ConstIntValueOne);
+
+		AITEST_NOT_EQUAL("Shared fragments created for same struct type using different constructor values should not share memory", SharedFragment1, SharedFragment3);
+		AITEST_EQUAL("Value in shared struct should be the same as the argument provided to GetOrCreateSharedFragment", SharedFragment3.Get<FTestSharedFragment_Int>().Value, ConstIntValueTwo);
+
+
+		return true;
+	}
+};
+IMPLEMENT_AI_INSTANT_TEST(FSharedFragmentValues_GetOrCreateWithStruct, "System.Mass.SharedFragments.GetOrCreate.WithStruct");
+
+struct FSharedFragmentValues_GetOrCreateNoArgs : FEntityTestBase
+{
+	virtual bool InstantTest() override
+	{
+		const FSharedStruct SharedFragment1 = EntityManager->GetOrCreateSharedFragment<FTestSharedFragment_Int>();
+		const FSharedStruct SharedFragment2 = EntityManager->GetOrCreateSharedFragment<FTestSharedFragment_Int>();
+		const FSharedStruct SharedFragment3 = EntityManager->GetOrCreateSharedFragment<FTestSharedFragment_Int>();
+
+		AITEST_EQUAL("Shared fragments created for same struct type using default constructor should share memory", SharedFragment1, SharedFragment2);
+		AITEST_EQUAL("Shared fragments created for same struct type using default constructor should share memory", SharedFragment1, SharedFragment3);
+
+		return true;
+	}
+};
+IMPLEMENT_AI_INSTANT_TEST(FSharedFragmentValues_GetOrCreateNoArgs, "System.Mass.SharedFragments.GetOrCreate.NoArgs");
+
+struct FSharedFragmentValues_GetOrCreateConstNoArgs : FEntityTestBase
+{
+	virtual bool InstantTest() override
+	{
+		const FConstSharedStruct SharedFragment1 = EntityManager->GetOrCreateConstSharedFragment<FTestConstSharedFragment_Int>();
+		const FConstSharedStruct SharedFragment2 = EntityManager->GetOrCreateConstSharedFragment<FTestConstSharedFragment_Int>();
+		const FConstSharedStruct SharedFragment3 = EntityManager->GetOrCreateConstSharedFragment<FTestConstSharedFragment_Int>();
+
+		AITEST_EQUAL("Shared fragments created for same struct type using default constructor should share memory", SharedFragment1, SharedFragment2);
+		AITEST_EQUAL("Shared fragments created for same struct type using default constructor should share memory", SharedFragment1, SharedFragment3);
+
+		return true;
+	}
+};
+IMPLEMENT_AI_INSTANT_TEST(FSharedFragmentValues_GetOrCreateConstNoArgs, "System.Mass.SharedFragments.GetOrCreate.ConstNoArgs");
+
+struct FSharedFragmentValues_GetOrCreateConstWithArgs : FEntityTestBase
+{
+	virtual bool InstantTest() override
+	{
+		constexpr int32 ConstIntValueOne = 1;
+		constexpr int32 ConstIntValueTwo = 2;
+
+		const FConstSharedStruct SharedFragment1 = EntityManager->GetOrCreateConstSharedFragment<FTestConstSharedFragment_Int>(/*Args*/ConstIntValueOne);
+		const FConstSharedStruct SharedFragment2 = EntityManager->GetOrCreateConstSharedFragment<FTestConstSharedFragment_Int>(/*Args*/ConstIntValueOne);
+		const FConstSharedStruct SharedFragment3 = EntityManager->GetOrCreateConstSharedFragment<FTestConstSharedFragment_Int>(/*Args*/ConstIntValueTwo);
+
+		AITEST_EQUAL("Shared fragments created for same struct type using same constructor value should share memory", SharedFragment1, SharedFragment2);
+		AITEST_EQUAL("Value in shared struct should be the same as the argument provided to GetOrCreateConstSharedFragment", SharedFragment1.Get<const FTestConstSharedFragment_Int>().Value, ConstIntValueOne);
+
+		AITEST_NOT_EQUAL("Shared fragments created for same struct type using different constructor values should not share memory", SharedFragment1, SharedFragment3);
+		AITEST_EQUAL("Value in shared struct should be the same as the argument provided to GetOrCreateConstSharedFragment", SharedFragment3.Get<const FTestConstSharedFragment_Int>().Value, ConstIntValueTwo);
+
+		return true;
+	}
+};
+IMPLEMENT_AI_INSTANT_TEST(FSharedFragmentValues_GetOrCreateConstWithArgs, "System.Mass.SharedFragments.GetOrCreate.ConstWithArgs");
+
+struct FSharedFragmentValues_GetOrCreateConstWithStruct : FEntityTestBase
+{
+	virtual bool InstantTest() override
+	{
+		constexpr int32 ConstIntValueOne = 1;
+		constexpr int32 ConstIntValueTwo = 2;
+
+		const FTestConstSharedFragment_Int TestSharedFragment_Int1(ConstIntValueOne);
+		const FTestConstSharedFragment_Int TestSharedFragment_Int2(ConstIntValueOne);
+		const FTestConstSharedFragment_Int TestSharedFragment_Int3(ConstIntValueTwo);
+		const FConstSharedStruct SharedFragment1 = EntityManager->GetOrCreateConstSharedFragment(TestSharedFragment_Int1);
+		const FConstSharedStruct SharedFragment2 = EntityManager->GetOrCreateConstSharedFragment(TestSharedFragment_Int2);
+		const FConstSharedStruct SharedFragment3 = EntityManager->GetOrCreateConstSharedFragment(TestSharedFragment_Int3);
+
+		AITEST_EQUAL("Shared fragments created for same struct type using same constructor value should share memory", SharedFragment1, SharedFragment2);
+		AITEST_EQUAL("Value in shared struct should be the same as the argument provided to GetOrCreateConstSharedFragment", SharedFragment1.Get<const FTestConstSharedFragment_Int>().Value, ConstIntValueOne);
+
+		AITEST_NOT_EQUAL("Shared fragments created for same struct type using different constructor values should not share memory", SharedFragment1, SharedFragment3);
+		AITEST_EQUAL("Value in shared struct should be the same as the argument provided to GetOrCreateConstSharedFragment", SharedFragment3.Get<const FTestConstSharedFragment_Int>().Value, ConstIntValueTwo);
+
+		return true;
+	}
+};
+IMPLEMENT_AI_INSTANT_TEST(FSharedFragmentValues_GetOrCreateConstWithStruct, "System.Mass.SharedFragments.GetOrCreate.ConstWithStruct");
+
 } // FMassEntityTest
 
 UE_ENABLE_OPTIMIZATION_SHIP
diff --git a/Engine/Source/Runtime/MassEntity/Public/MassEntityManager.h b/Engine/Source/Runtime/MassEntity/Public/MassEntityManager.h
--- a/Engine/Source/Runtime/MassEntity/Public/MassEntityManager.h
+++ b/Engine/Source/Runtime/MassEntity/Public/MassEntityManager.h
@@ -549,34 +549,85 @@ public:
 	/** 
 	 * Depending on the current state of Manager's command buffer the function will either move all the commands out of 
 	 * InOutCommandBuffer into the main command buffer or append it to the list of command buffers waiting to be flushed.
-	 * @note as a consequence of the call InOutCommandBuffer can get its contents emptied due some of the undelying code using Move semantics
+	 * @note as a consequence of the call InOutCommandBuffer can get its contents emptied due some of the underlying code using Move semantics
 	 */
 	void AppendCommands(TSharedPtr<FMassCommandBuffer>& InOutCommandBuffer);
 
 	template<typename T>
+	UE_DEPRECATED(5.5, "This method will no longer be exposed. Use GetOrCreateConstSharedFragment instead.")
 	const FConstSharedStruct& GetOrCreateConstSharedFragmentByHash(const uint32 Hash, const T& Fragment)
 	{
 		static_assert(TIsDerivedFrom<T, FMassConstSharedFragment>::IsDerived, "Given struct doesn't represent a valid const shared fragment type. Make sure to inherit from FMassConstSharedFragment or one of its child-types.");
 		int32& Index = ConstSharedFragmentsMap.FindOrAddByHash(Hash, Hash, INDEX_NONE);
 		if (Index == INDEX_NONE)
 		{
-			Index = ConstSharedFragments.Add(FSharedStruct::Make(Fragment));
+			Index = ConstSharedFragments.Add(FConstSharedStruct::Make(Fragment));
 		}
 		return ConstSharedFragments[Index];
 	}
 
+private:
 	template<typename T>
-	const FConstSharedStruct& GetOrCreateConstSharedFragment(const T& Fragment)
+	const FSharedStruct& GetOrCreateSharedFragmentByHash(const uint32 Hash, const T& Fragment)
 	{
-		const uint32 Hash = UE::StructUtils::GetStructCrc32(FConstStructView::Make(Fragment));
-		return GetOrCreateConstSharedFragmentByHash(Hash, Fragment);
+		static_assert(TIsDerivedFrom<T, FMassSharedFragment>::IsDerived, "Given struct doesn't represent a valid shared fragment type. Make sure to inherit from FMassSharedFragment or one of its child-types.");
+		int32& Index = SharedFragmentsMap.FindOrAddByHash(Hash, Hash, INDEX_NONE);
+		if (Index == INDEX_NONE)
+		{
+			Index = SharedFragments.Add(FSharedStruct::Make(Fragment));
+			// note that even though we're copying the freshly created FSharedStruct instance it's perfectly fine since 
+			// FSharedStruct do guarantee there's not going to be data duplication (via a member shared pointer to hosted data)
+			TArray<FSharedStruct>& InstancesOfType = SharedFragmentsTypeMap.FindOrAdd(T::StaticStruct(), {});
+			InstancesOfType.Add(SharedFragments[Index]);
+		}
+
+		return SharedFragments[Index];
 	}
 
+	const FConstSharedStruct& GetOrCreateConstSharedFragmentByHash(const uint32 Hash, const UScriptStruct* InScriptStruct, const uint8* InStructMemory)
+	{
+		int32& Index = ConstSharedFragmentsMap.FindOrAddByHash(Hash, Hash, INDEX_NONE);
+		if (Index == INDEX_NONE)
+		{
+			Index = ConstSharedFragments.Add(FConstSharedStruct::Make(InScriptStruct, InStructMemory));
+		}
+		return ConstSharedFragments[Index];
+	}
+
+	const FSharedStruct& GetOrCreateSharedFragmentByHash(const uint32 Hash, const UScriptStruct* InScriptStruct, const uint8* InStructMemory)
+	{
+		int32& Index = SharedFragmentsMap.FindOrAddByHash(Hash, Hash, INDEX_NONE);
+		if (Index == INDEX_NONE)
+		{
+			Index = SharedFragments.Add(FSharedStruct::Make(InScriptStruct, InStructMemory));
+			// note that even though we're copying the freshly created FSharedStruct instance it's perfectly fine since 
+			// FSharedStruct do guarantee there's not going to be data duplication (via a member shared pointer to hosted data)
+			TArray<FSharedStruct>& InstancesOfType = SharedFragmentsTypeMap.FindOrAdd(InScriptStruct, {});
+			InstancesOfType.Add(SharedFragments[Index]);
+		}
+		return SharedFragments[Index];
+	}
+
+	template<typename T, typename... TArgs>
+	const FConstSharedStruct& GetOrCreateConstSharedFragmentByHash(const uint32 Hash, TArgs&&... InArgs)
+	{
+		static_assert(TIsDerivedFrom<T, FMassConstSharedFragment>::IsDerived, "Given struct doesn't represent a valid const shared fragment type. Make sure to inherit from FMassConstSharedFragment or one of its child-types.");
+		int32& Index = ConstSharedFragmentsMap.FindOrAddByHash(Hash, Hash, INDEX_NONE);
+		if (Index == INDEX_NONE)
+		{
+			Index = ConstSharedFragments.Add(FConstSharedStruct::Make<T>(Forward<TArgs>(InArgs)...));
+		}
+
+		return ConstSharedFragments[Index];
+	}
+
+public:
+
 	template<typename T, typename... TArgs>
+	UE_DEPRECATED(5.5, "This method will no longer be exposed. Use GetOrCreateSharedFragment instead.")
 	const FSharedStruct& GetOrCreateSharedFragmentByHash(const uint32 Hash, TArgs&&... InArgs)
 	{
 		static_assert(TIsDerivedFrom<T, FMassSharedFragment>::IsDerived, "Given struct doesn't represent a valid shared fragment type. Make sure to inherit from FMassSharedFragment or one of its child-types.");
-
 		int32& Index = SharedFragmentsMap.FindOrAddByHash(Hash, Hash, INDEX_NONE);
 		if (Index == INDEX_NONE)
 		{
@@ -590,6 +641,268 @@ public:
 		return SharedFragments[Index];
 	}
 
+	/**
+	 * Returns or creates a shared struct associated to a given shared fragment set of values
+	 * identified internally by a CRC.
+	 * Use this overload when an instance of the desired const shared fragment type is available and
+	 * that can be used directly to compute a CRC (i.e., UE::StructUtils::GetStructCrc32)
+	 *	e.g.,
+	 *	USTRUCT()
+	 *	struct FIntConstSharedFragment : public FMassConstSharedFragment
+	 *	{
+	 *		GENERATED_BODY()
+	 *
+	 *		UPROPERTY()
+	 *		int32 Value = 0;
+	 *	};
+	 *
+	 *	FIntConstSharedFragment Fragment;
+	 *	Fragment.Value = 123;
+	 *	const FConstSharedStruct SharedStruct = EntityManager.GetOrCreateConstSharedFragment(Fragment);
+	 *
+	 * @params Fragment Instance of the desired fragment type
+	 * @return FConstSharedStruct to the matching, or newly created shared fragment
+	 */
+	template<typename T>
+	const FConstSharedStruct& GetOrCreateConstSharedFragment(const T& Fragment)
+	{
+		static_assert(TIsDerivedFrom<T, FMassConstSharedFragment>::IsDerived,
+			"Given struct doesn't represent a valid const shared fragment type. Make sure to inherit from FMassConstSharedFragment or one of its child-types.");
+		const uint32 Hash = UE::StructUtils::GetStructCrc32(FConstStructView::Make(Fragment));
+PRAGMA_DISABLE_DEPRECATION_WARNINGS
+		return GetOrCreateConstSharedFragmentByHash(Hash, Fragment);
+PRAGMA_ENABLE_DEPRECATION_WARNINGS
+	}
+
+	/**
+	 * Returns or creates a shared struct associated to a given shared fragment set of values
+	 * identified internally by a CRC.
+	 * Use this overload when an instance of the desired shared fragment type is available and
+	 * that can be used directly to compute a CRC (i.e., UE::StructUtils::GetStructCrc32)
+	 *	e.g.,
+	 *	USTRUCT()
+	 *	struct FIntSharedFragment : public FMassSharedFragment
+	 *	{
+	 *		GENERATED_BODY()
+	 *
+	 *		UPROPERTY()
+	 *		int32 Value = 0;
+	 *	};
+	 *
+	 *	FIntSharedFragment Fragment;
+	 *	Fragment.Value = 123;
+	 *	const FSharedStruct SharedStruct = EntityManager.GetOrCreateSharedFragment(Fragment);
+	 *
+	 * @params Fragment Instance of the desired fragment type
+	 * @return FSharedStruct to the matching, or newly created shared fragment
+	 */
+	template<typename T>
+	const FSharedStruct& GetOrCreateSharedFragment(const T& Fragment)
+	{
+		static_assert(TIsDerivedFrom<T, FMassSharedFragment>::IsDerived,
+			"Given struct doesn't represent a valid shared fragment type. Make sure to inherit from FMassSharedFragment or one of its child-types.");
+		const uint32 Hash = UE::StructUtils::GetStructCrc32(FConstStructView::Make(Fragment));
+		return GetOrCreateSharedFragmentByHash(Hash, Fragment);
+	}
+
+	/**
+	 * Returns or creates a shared struct associated to a given shared fragment set of values
+	 * identified internally by a CRC.
+	 * Use this overload when values can be provided as constructor arguments for the desired const shared fragment type and
+	 * that can be used directly to compute a CRC (i.e., UE::StructUtils::GetStructCrc32)
+ 	 *	e.g.,
+	 *	USTRUCT()
+	 *	struct FIntConstSharedFragment : public FMassConstSharedFragment
+	 *	{
+	 *		GENERATED_BODY()
+	 *
+	 *		FIntConstSharedFragment(const int32 InValue) : Value(InValue) {}
+	 *
+	 *		UPROPERTY()
+	 *		int32 Value = 0;
+	 *	};
+	 *
+	 *	const FConstSharedStruct SharedStruct = EntityManager.GetOrCreateConstSharedFragment<FIntConstSharedFragment>(123);
+	 *
+	 * @params InArgs List of arguments provided to the constructor of the desired fragment type
+	 * @return FConstSharedStruct to the matching, or newly created shared fragment
+	 */
+	template<typename T, typename... TArgs>
+	const FConstSharedStruct& GetOrCreateConstSharedFragment(TArgs&&... InArgs)
+	{
+		static_assert(TIsDerivedFrom<T, FMassConstSharedFragment>::IsDerived,
+			"Given struct doesn't represent a valid const shared fragment type. Make sure to inherit from FMassConstSharedFragment or one of its child-types.");
+		T Struct(Forward<TArgs>(InArgs)...);
+		const uint32 Hash = UE::StructUtils::GetStructCrc32(FConstStructView::Make(Struct));
+PRAGMA_DISABLE_DEPRECATION_WARNINGS
+		return GetOrCreateConstSharedFragmentByHash(Hash, MoveTemp(Struct));
+PRAGMA_ENABLE_DEPRECATION_WARNINGS
+	}
+
+	/**
+	 * Returns or creates a shared struct associated to a given shared fragment set of values
+	 * identified internally by a CRC.
+	 * Use this overload when values can be provided as constructor arguments for the desired shared fragment type and
+	 * that can be used directly to compute a CRC (i.e., UE::StructUtils::GetStructCrc32)
+ 	 *	e.g.,
+	 *	USTRUCT()
+	 *	struct FIntSharedFragment : public FMassSharedFragment
+	 *	{
+	 *		GENERATED_BODY()
+	 *
+	 *		FIntSharedFragment(const int32 InValue) : Value(InValue) {}
+	 *
+	 *		UPROPERTY()
+	 *		int32 Value = 0;
+	 *	};
+	 *
+	 *	const FSharedStruct SharedStruct = EntityManager.GetOrCreateSharedFragment<FIntSharedFragment>(123);
+	 *
+	 * @params InArgs List of arguments provided to the constructor of the desired fragment type
+	 * @return FSharedStruct to the matching, or newly created shared fragment
+	 */
+	template<typename T, typename... TArgs>
+	const FSharedStruct& GetOrCreateSharedFragment(TArgs&&... InArgs)
+	{
+		static_assert(TIsDerivedFrom<T, FMassSharedFragment>::IsDerived,
+			"Given struct doesn't represent a valid shared fragment type. Make sure to inherit from FMassSharedFragment or one of its child-types.");
+		T Struct(Forward<TArgs>(InArgs)...);
+		const uint32 Hash = UE::StructUtils::GetStructCrc32(FConstStructView::Make(Struct));
+		return GetOrCreateSharedFragmentByHash(Hash, MoveTemp(Struct));
+	}
+
+	/**
+	 * Returns or creates a shared struct associated to a given shared fragment set of values
+	 * identified internally by a CRC.
+	 * Use this overload when the reflection data and the memory of an instance of the desired const shared fragment type
+	 * is available and that can be used directly to compute a CRC (i.e., UE::StructUtils::GetStructCrc32)
+	 * e.g.,
+	 * FSharedStruct SharedStruct = EntityManager.GetOrCreateConstSharedFragment(*StructView.GetScriptStruct(), StructView.GetMemory());
+	 *
+	 * @params InScriptStruct Reflection data structure associated to the desired fragment type
+	 * @params InStructMemory Actual data of the desired fragment type 
+	 * @return FConstSharedStruct to the matching, or newly created shared fragment
+	 */
+	const FConstSharedStruct& GetOrCreateConstSharedFragment(const UScriptStruct& InScriptStruct, const uint8* InStructMemory)
+	{
+		checkf(InScriptStruct.IsChildOf(TBaseStructure<FMassConstSharedFragment>::Get()),
+			TEXT("Given struct doesn't represent a valid const shared fragment type. Make sure to inherit from FMassConstSharedFragment or one of its child-types."));
+		const uint32 Hash = UE::StructUtils::GetStructCrc32(InScriptStruct, InStructMemory);
+		return GetOrCreateConstSharedFragmentByHash(Hash, &InScriptStruct, InStructMemory);
+	}
+
+	/**
+	 * Returns or creates a shared struct associated to a given shared fragment set of values
+	 * identified internally by a CRC.
+	 * Use this overload when the reflection data and the memory of an instance of the desired shared fragment type
+	 * is available and that can be used directly to compute a CRC (i.e., UE::StructUtils::GetStructCrc32)
+	 * e.g.,
+	 * FSharedStruct SharedStruct = EntityManager.GetOrCreateSharedFragment(*StructView.GetScriptStruct(), StructView.GetMemory());
+	 *
+	 * @params InScriptStruct Reflection data structure associated to the desired fragment type
+	 * @params InStructMemory Actual data of the desired fragment type 
+	 * @return FSharedStruct to the matching, or newly created shared fragment
+	 */
+	const FSharedStruct& GetOrCreateSharedFragment(const UScriptStruct& InScriptStruct, const uint8* InStructMemory)
+	{
+		checkf(InScriptStruct.IsChildOf(TBaseStructure<FMassSharedFragment>::Get()),
+			TEXT("Given struct doesn't represent a valid shared fragment type. Make sure to inherit from FMassSharedFragment or one of its child-types."));
+		const uint32 Hash = UE::StructUtils::GetStructCrc32(InScriptStruct, InStructMemory);
+		return GetOrCreateSharedFragmentByHash(Hash, &InScriptStruct, InStructMemory);
+	}
+
+	/**
+	 * Returns or creates a shared struct associated to a given shared fragment set of values
+	 * identified internally by a CRC.
+	 * Use this overload when a different struct should be used to compute a CRC (i.e., UE::StructUtils::GetStructCrc32)
+	 * and values can be provided as constructor arguments for the desired const shared fragment type
+	 *	e.g.,
+	 *
+	 *	USTRUCT()
+	 *	struct FIntConstSharedFragmentParams
+	 *	{
+	 *		GENERATED_BODY()
+	 *
+	 *		FIntConstSharedFragmentParams(const int32 InValue) : Value(InValue) {}
+	 *
+	 *		UPROPERTY()
+	 *		int32 Value = 0;
+	 *	};
+	 *
+	 *	USTRUCT()
+	 *	struct FIntConstSharedFragment : public FMassConstSharedFragment
+	 *	{
+	 *		GENERATED_BODY()
+	 *
+	 *		FIntConstSharedFragment(const FIntConstSharedFragmentParams& InParams) : Value(InParams.Value) {}
+	 *
+	 *		int32 Value = 0;
+	 *	};
+	 *
+	 *	FIntConstSharedFragmentParams Params(123);
+	 *	const FConstSharedStruct SharedStruct = EntityManager.GetOrCreateConstSharedFragment<FIntConstSharedFragment>(FConstStructView::Make(Params), Params);
+	 *
+	 * @params HashingHelperStruct Struct view passed to UE::StructUtils::GetStructCrc32 to compute the CRC
+	 * @params InArgs List of arguments provided to the constructor of the desired fragment type
+	 * @return FConstSharedStruct to the matching, or newly created shared fragment
+	 */
+	template<typename T, typename... TArgs>
+	const FConstSharedStruct& GetOrCreateConstSharedFragment(const FConstStructView HashingHelperStruct, TArgs&&... InArgs)
+	{
+		static_assert(TIsDerivedFrom<T, FMassConstSharedFragment>::IsDerived,
+			"Given struct doesn't represent a valid const shared fragment type. Make sure to inherit from FMassConstSharedFragment or one of its child-types.");
+		T Fragment(Forward<TArgs>(InArgs)...);
+		const uint32 Hash = UE::StructUtils::GetStructCrc32(HashingHelperStruct);
+PRAGMA_DISABLE_DEPRECATION_WARNINGS
+		return GetOrCreateConstSharedFragmentByHash(Hash, MoveTemp(Fragment));
+PRAGMA_ENABLE_DEPRECATION_WARNINGS
+	}
+
+	/**
+	 * Returns or creates a shared struct associated to a given shared fragment set of values
+	 * identified internally by a CRC.
+	 * Use this overload when a different struct should be used to compute a CRC (i.e., UE::StructUtils::GetStructCrc32)
+	 * and values can be provided as constructor arguments for the desired shared fragment type
+	 *	e.g.,
+	 *
+	 *	USTRUCT()
+	 *	struct FIntSharedFragmentParams
+	 *	{
+	 *		GENERATED_BODY()
+	 *
+	 *		FInSharedFragmentParams(const int32 InValue) : Value(InValue) {}
+	 *
+	 *		UPROPERTY()
+	 *		int32 Value = 0;
+	 *	};
+	 *
+	 *	USTRUCT()
+	 *	struct FIntSharedFragment : public FMassSharedFragment
+	 *	{
+	 *		GENERATED_BODY()
+	 *
+	 *		FIntSharedFragment(const FIntConstSharedFragmentParams& InParams) : Value(InParams.Value) {}
+	 *
+	 *		int32 Value = 0;
+	 *	};
+	 *
+	 *	FIntSharedFragmentParams Params(123);
+	 *	const FSharedStruct SharedStruct = EntityManager.GetOrCreateSharedFragment<FIntSharedFragment>(FConstStructView::Make(Params), Params);
+	 *
+	 * @params HashingHelperStruct Struct view passed to UE::StructUtils::GetStructCrc32 to compute the CRC
+	 * @params InArgs List of arguments provided to the constructor of the desired fragment type
+	 * @return FSharedStruct to the matching, or newly created shared fragment
+	 */
+	template<typename T, typename... TArgs>
+	const FSharedStruct& GetOrCreateSharedFragment(const FConstStructView HashingHelperStruct, TArgs&&... InArgs)
+	{
+		static_assert(TIsDerivedFrom<T, FMassSharedFragment>::IsDerived,
+			"Given struct doesn't represent a valid shared fragment type. Make sure to inherit from FMassSharedFragment or one of its child-types.");
+		T Fragment(Forward<TArgs>(InArgs)...);
+		const uint32 Hash = UE::StructUtils::GetStructCrc32(HashingHelperStruct);
+		return GetOrCreateSharedFragmentByHash(Hash, MoveTemp(Fragment));
+	}
+
 	template<typename T>
 	void ForEachSharedFragment(TFunctionRef< void(T& /*SharedFragment*/) > ExecuteFunction)
 	{
