[MassEntityManager] deprecation cleanup

diff --git a/Engine/Source/Runtime/MassEntity/Private/MassEntityManager.cpp b/Engine/Source/Runtime/MassEntity/Private/MassEntityManager.cpp
--- a/Engine/Source/Runtime/MassEntity/Private/MassEntityManager.cpp
+++ b/Engine/Source/Runtime/MassEntity/Private/MassEntityManager.cpp
@@ -5,7 +5,6 @@
 #include "MassArchetypeData.h"
 #include "MassCommandBuffer.h"
 #include "MassEntityManagerStorage.h"
-#include "HAL/IConsoleManager.h"
 #include "Engine/World.h"
 #include "UObject/UObjectIterator.h"
 #include "VisualLogger/VisualLogger.h"
@@ -367,7 +366,7 @@ void FMassEntityManager::Initialize(const FMassEntityManagerStorageInitParams& I
 		}
 	}
 #if WITH_MASSENTITY_DEBUG
-	RequirementAccessDetector.Initialize();	
+	RequirementAccessDetector.Initialize();
 	FMassDebugger::RegisterEntityManager(*this);
 #endif // WITH_MASSENTITY_DEBUG
 
@@ -2082,56 +2081,6 @@ FMassRequirementAccessDetector& FMassEntityManager::GetRequirementAccessDetector
 //-----------------------------------------------------------------------------
 // DEPRECATED
 //-----------------------------------------------------------------------------
-FMassArchetypeHandle FMassEntityManager::CreateArchetype(TConstArrayView<const UScriptStruct*> FragmentsAndTagsList, const FName ArchetypeDebugName)
-{
-	FMassArchetypeCreationParams Params;
-	Params.DebugName = ArchetypeDebugName;
-	return CreateArchetype(FragmentsAndTagsList, Params);
-}
-
-FMassArchetypeHandle FMassEntityManager::CreateArchetype(FMassArchetypeHandle SourceArchetype, TConstArrayView<const UScriptStruct*> FragmentsAndTagsList, const FName ArchetypeDebugName)
-{
-	FMassArchetypeCreationParams Params;
-	Params.DebugName = ArchetypeDebugName;
-	return CreateArchetype(SourceArchetype, FragmentsAndTagsList, Params);
-}
-
-FMassArchetypeHandle FMassEntityManager::CreateArchetype(const FMassArchetypeCompositionDescriptor& Composition, const FName ArchetypeDebugName)
-{
-	FMassArchetypeCreationParams Params;
-	Params.DebugName = ArchetypeDebugName;
-	return CreateArchetype(Composition, Params);
-}
-
-FMassArchetypeHandle FMassEntityManager::CreateArchetype(const TSharedPtr<FMassArchetypeData>& SourceArchetype, const FMassFragmentBitSet& InFragments, const FName ArchetypeDebugName)
-{
-	FMassArchetypeCreationParams Params;
-	Params.DebugName = ArchetypeDebugName;
-	return CreateArchetype(SourceArchetype, InFragments, Params);
-}
-
-FMassEntityHandle FMassEntityManager::CreateEntity(TConstArrayView<FInstancedStruct> FragmentInstanceList, const FMassArchetypeSharedFragmentValues& SharedFragmentValues, const FName ArchetypeDebugName)
-{
-	FMassArchetypeCreationParams Params;
-	Params.DebugName = ArchetypeDebugName;
-	return CreateEntity(FragmentInstanceList, SharedFragmentValues, Params);
-}
-
-void FMassEntityManager::BatchBuildEntities(const FMassArchetypeEntityCollectionWithPayload& EncodedEntitiesWithPayload, const FMassFragmentBitSet& FragmentsAffected
-	, const FMassArchetypeSharedFragmentValues& SharedFragmentValues, const FName ArchetypeDebugName)
-{
-	FMassArchetypeCreationParams Params;
-	Params.DebugName = ArchetypeDebugName;
-	BatchBuildEntities(EncodedEntitiesWithPayload, FragmentsAffected, SharedFragmentValues, Params);
-}
-
-void FMassEntityManager::BatchBuildEntities(const FMassArchetypeEntityCollectionWithPayload& EncodedEntitiesWithPayload, FMassArchetypeCompositionDescriptor&& Composition
-	, const FMassArchetypeSharedFragmentValues& SharedFragmentValues, const FName ArchetypeDebugName)
-{
-	FMassArchetypeCreationParams Params;
-	Params.DebugName = ArchetypeDebugName;
-	BatchBuildEntities(EncodedEntitiesWithPayload, MoveTemp(Composition), SharedFragmentValues, Params);
-}
 
 const FMassArchetypeEntityCollection& FMassEntityManager::FEntityCreationContext::GetEntityCollection() const
 {
diff --git a/Engine/Source/Runtime/MassEntity/Public/MassEntityManager.h b/Engine/Source/Runtime/MassEntity/Public/MassEntityManager.h
--- a/Engine/Source/Runtime/MassEntity/Public/MassEntityManager.h
+++ b/Engine/Source/Runtime/MassEntity/Public/MassEntityManager.h
@@ -1127,29 +1127,4 @@ private:
 	bool bFirstCommandFlush = true;
 
 	FDelegateHandle OnPostForkHandle;
-
-	//-----------------------------------------------------------------------------
-	// DEPRECATED
-	//-----------------------------------------------------------------------------
-public:
-	UE_DEPRECATED(5.3, "This Flavor of CreateArchetype is deprecated. Use the one with FMassArchetypeCreationParams parameter instead.")
-	FMassArchetypeHandle CreateArchetype(TConstArrayView<const UScriptStruct*> FragmentsAndTagsList, const FName ArchetypeDebugName);
-	UE_DEPRECATED(5.3, "This Flavor of CreateArchetype is deprecated. Use the one with FMassArchetypeCreationParams parameter instead.")
-	FMassArchetypeHandle CreateArchetype(FMassArchetypeHandle SourceArchetype, TConstArrayView<const UScriptStruct*> FragmentsAndTagsList, const FName ArchetypeDebugName);
-	UE_DEPRECATED(5.3, "This Flavor of CreateArchetype is deprecated. Use the one with FMassArchetypeCreationParams parameter instead.")
-	FMassArchetypeHandle CreateArchetype(const FMassArchetypeCompositionDescriptor& Composition, const FName ArchetypeDebugName);
-	UE_DEPRECATED(5.3, "This Flavor of CreateArchetype is deprecated. Use the one with FMassArchetypeCreationParams parameter instead.")
-	FMassArchetypeHandle CreateArchetype(const TSharedPtr<FMassArchetypeData>& SourceArchetype, const FMassFragmentBitSet& InFragments, const FName ArchetypeDebugName);
-	UE_DEPRECATED(5.3, "This Flavor of CreateEntity is deprecated. Use the one with FMassArchetypeCreationParams parameter instead.")
-	FMassEntityHandle CreateEntity(TConstArrayView<FInstancedStruct> FragmentInstanceList, const FMassArchetypeSharedFragmentValues& SharedFragmentValues, const FName ArchetypeDebugName);
-	UE_DEPRECATED(5.3, "This Flavor of BatchBuildEntities is deprecated. Use the one with FMassArchetypeCreationParams parameter instead.")
-	void BatchBuildEntities(const FMassArchetypeEntityCollectionWithPayload& EncodedEntitiesWithPayload, const FMassFragmentBitSet& FragmentsAffected
-		, const FMassArchetypeSharedFragmentValues& SharedFragmentValues, const FName ArchetypeDebugName);
-	UE_DEPRECATED(5.3, "This Flavor of BatchBuildEntities is deprecated. Use the one with FMassArchetypeCreationParams parameter instead.")
-	void BatchBuildEntities(const FMassArchetypeEntityCollectionWithPayload& EncodedEntitiesWithPayload, FMassArchetypeCompositionDescriptor&& Composition
-		, const FMassArchetypeSharedFragmentValues& SharedFragmentValues, const FName ArchetypeDebugName);
 };
-
-#if UE_ENABLE_INCLUDE_ORDER_DEPRECATED_IN_5_2
-#include "MassExecutionContext.h"
-#endif
