PR #12151: Mass GC reference cleanup to support gc.GarbageEliminationEnabled=False

diff --git a/Engine/Source/Runtime/MassEntity/Private/MassEntityManager.cpp b/Engine/Source/Runtime/MassEntity/Private/MassEntityManager.cpp
--- a/Engine/Source/Runtime/MassEntity/Private/MassEntityManager.cpp
+++ b/Engine/Source/Runtime/MassEntity/Private/MassEntityManager.cpp
@@ -402,6 +402,8 @@ void FMassEntityManager::Deinitialize()
 #endif // WITH_MASSENTITY_DEBUG
 
 		EntityStorage.Emplace<FEmptyVariantState>();
+
+		ObserverManager.DeInitialize();
 		
 		bInitialized = false;
 	}
diff --git a/Engine/Source/Runtime/MassEntity/Private/MassObserverManager.cpp b/Engine/Source/Runtime/MassEntity/Private/MassObserverManager.cpp
--- a/Engine/Source/Runtime/MassEntity/Private/MassObserverManager.cpp
+++ b/Engine/Source/Runtime/MassEntity/Private/MassObserverManager.cpp
@@ -82,6 +82,15 @@ void FMassObserverManager::Initialize()
 	}
 }
 
+void FMassObserverManager::DeInitialize()
+{
+	for (int32 i = 0; i < (int32)EMassObservedOperation::MAX; ++i)
+	{
+		(*FragmentObservers[i]).Empty();
+		(*TagObservers[i]).Empty();
+	}
+}
+
 bool FMassObserverManager::OnPostEntitiesCreated(const FMassArchetypeEntityCollection& EntityCollection)
 {
 	return OnPostEntitiesCreated(MakeArrayView(&EntityCollection , 1));
diff --git a/Engine/Source/Runtime/MassEntity/Private/MassProcessingPhaseManager.cpp b/Engine/Source/Runtime/MassEntity/Private/MassProcessingPhaseManager.cpp
--- a/Engine/Source/Runtime/MassEntity/Private/MassProcessingPhaseManager.cpp
+++ b/Engine/Source/Runtime/MassEntity/Private/MassProcessingPhaseManager.cpp
@@ -256,6 +256,8 @@ void FMassProcessingPhaseManager::Deinitialize()
 	{
 		Phase.PhaseProcessor = nullptr;
 	}
+
+	DynamicProcessors.Reset();
 }
 
 const FGraphEventRef& FMassProcessingPhaseManager::TriggerPhase(const EMassProcessingPhase Phase, const float DeltaTime, const FGraphEventRef& MyCompletionGraphEvent)
diff --git a/Engine/Source/Runtime/MassEntity/Public/MassObserverManager.h b/Engine/Source/Runtime/MassEntity/Public/MassObserverManager.h
--- a/Engine/Source/Runtime/MassEntity/Public/MassObserverManager.h
+++ b/Engine/Source/Runtime/MassEntity/Public/MassObserverManager.h
@@ -132,6 +132,7 @@ protected:
 	explicit FMassObserverManager(FMassEntityManager& Owner);
 
 	void Initialize();
+	void DeInitialize();
 	void HandleFragmentsImpl(FMassProcessingContext& ProcessingContext, const FMassArchetypeEntityCollection& EntityCollection
 		, TArrayView<const UScriptStruct*> ObservedTypes, FMassObserversMap& HandlersContainer);
 	void HandleSingleEntityImpl(const UScriptStruct& FragmentType, const FMassArchetypeEntityCollection& EntityCollection, FMassObserversMap& HandlersContainer);
