[Mass] Fixed an obvious mistake in FMassEntityManager::BuildEntity where const-shared-fragments were being added to archetype's composition as regular-shared-fragments.

diff --git a/Engine/Source/Runtime/MassEntity/Private/MassEntityManager.cpp b/Engine/Source/Runtime/MassEntity/Private/MassEntityManager.cpp
--- a/Engine/Source/Runtime/MassEntity/Private/MassEntityManager.cpp
+++ b/Engine/Source/Runtime/MassEntity/Private/MassEntityManager.cpp
@@ -755,7 +755,7 @@ void FMassEntityManager::BuildEntity(FMassEntityHandle Entity, TConstArrayView<F
 	FMassArchetypeCompositionDescriptor Composition(FragmentInstanceList, FMassTagBitSet(), FMassChunkFragmentBitSet(), FMassSharedFragmentBitSet(), FMassConstSharedFragmentBitSet());
 	for (const FConstSharedStruct& SharedFragment : SharedFragmentValues.GetConstSharedFragments())
 	{
-		Composition.SharedFragments.Add(*SharedFragment.GetScriptStruct());
+		Composition.ConstSharedFragments.Add(*SharedFragment.GetScriptStruct());
 	}
 	for (const FSharedStruct& SharedFragment : SharedFragmentValues.GetSharedFragments())
 	{
