[Mass] Added a flag to MassProcessor to mark processor instances as "dynamic", both for debugging purposes as well as special handling of dynamic processors by the Phase Manager.

Also:
* marking processors that were removed by the dependency solver as garbage, while WITH_MASSENTITY_DEBUG, which helps MassDebugger's processor view's readability
* minor touches to MassProcessor. Moved some bools around, converted them to bitfields, removed deprecated stuff.
    * this reduced the size of MassProcessor by 16 bytes.

diff --git a/Engine/Source/Runtime/MassEntity/Private/MassProcessingPhaseManager.cpp b/Engine/Source/Runtime/MassEntity/Private/MassProcessingPhaseManager.cpp
--- a/Engine/Source/Runtime/MassEntity/Private/MassProcessingPhaseManager.cpp
+++ b/Engine/Source/Runtime/MassEntity/Private/MassProcessingPhaseManager.cpp
@@ -175,9 +175,13 @@ void FMassPhaseProcessorConfigurationHelper::Configure(TArrayView<UMassProcessor
 	if (TmpPipeline.Num())
 	{
 		UE_VLOG_UELOG(&PhaseProcessor, LogMass, Verbose, TEXT("Discarding processors due to not having anything to do (no relevant Archetypes):"));
-		for (const UMassProcessor* Processor : TmpPipeline.GetProcessors())
+		for (UMassProcessor* Processor : TmpPipeline.GetProcessors())
 		{
 			UE_VLOG_UELOG(&PhaseProcessor, LogMass, Verbose, TEXT("\t%s"), *Processor->GetProcessorName());
+			if (Processor->IsDynamic() == false)
+			{
+				Processor->MarkAsGarbage();
+			}
 		}
 	}
 #endif // WITH_MASSENTITY_DEBUG
@@ -455,6 +459,7 @@ FString FMassProcessingPhaseManager::GetName() const
 void FMassProcessingPhaseManager::RegisterDynamicProcessor(UMassProcessor& Processor)
 {
 	DynamicProcessors.Add(&Processor);
+	Processor.MarkAsDynamic();
 	ProcessingGraphBuildStates[uint32(Processor.GetProcessingPhase())].bProcessorsNeedRebuild = true;
 }
 
diff --git a/Engine/Source/Runtime/MassEntity/Public/MassProcessor.h b/Engine/Source/Runtime/MassEntity/Public/MassProcessor.h
--- a/Engine/Source/Runtime/MassEntity/Public/MassProcessor.h
+++ b/Engine/Source/Runtime/MassEntity/Public/MassProcessor.h
@@ -63,9 +63,6 @@ public:
 	 */
 	bool ShouldAllowMultipleInstances() const { return bAllowMultipleInstances; }
 
-	UE_DEPRECATED(5.2, "This function is deprecated. Use ShouldAllowDuplicates instead.")
-	bool AllowDuplicates() const { return ShouldAllowMultipleInstances(); }
-
 	void DebugOutputDescription(FOutputDevice& Ar) const { DebugOutputDescription(Ar, 0); }
 	virtual void DebugOutputDescription(FOutputDevice& Ar, int32 Indent) const;
 	virtual FString GetProcessorName() const { return GetName(); }
@@ -96,6 +93,9 @@ public:
 	 *  this requirement will cause check failure and the query won't be registered. */
 	void RegisterQuery(FMassEntityQuery& Query);
 
+	void MarkAsDynamic() { bIsDynamic = true; }
+	bool IsDynamic() const { return bIsDynamic != 0; }
+
 	bool ShouldAutoAddToGlobalList() const { return bAutoRegisterWithProcessingPhases; }
 #if WITH_EDITOR
 	bool ShouldShowUpInSettings() const { return ShouldAutoAddToGlobalList() || bCanShowUpInSettings; }
@@ -120,9 +120,9 @@ protected:
 	virtual void Execute(FMassEntityManager& EntityManager, FMassExecutionContext& Context) PURE_VIRTUAL(UMassProcessor::Execute);
 
 protected:
-	/** Whether this processor should be executed on StandAlone or Server or Client */
-	UPROPERTY(EditAnywhere, Category = "Pipeline", meta = (Bitmask, BitmaskEnum = "/Script/MassEntity.EProcessorExecutionFlags"), config)
-	int32 ExecutionFlags;
+	/** Configures when this given processor can be executed in relation to other processors and processing groups, within its processing phase. */
+	UPROPERTY(EditDefaultsOnly, Category = Processor, config)
+	FMassProcessorExecutionOrder ExecutionOrder;
 
 	/** Processing phase this processor will be automatically run as part of. Needs to be set before the processor gets
 	 *  registered with MassProcessingPhaseManager, otherwise it will have no effect. This property is usually read via
@@ -130,31 +130,39 @@ protected:
 	UPROPERTY(EditDefaultsOnly, Category = Processor, config)
 	EMassProcessingPhase ProcessingPhase = EMassProcessingPhase::PrePhysics;
 
-	/** Configures when this given processor can be executed in relation to other processors and processing groups, within its processing phase. */
-	UPROPERTY(EditDefaultsOnly, Category = Processor, config)
-	FMassProcessorExecutionOrder ExecutionOrder;
+	/** Whether this processor should be executed on StandAlone or Server or Client */
+	UPROPERTY(EditAnywhere, Category = "Pipeline", meta = (Bitmask, BitmaskEnum = "/Script/MassEntity.EProcessorExecutionFlags"), config)
+	uint8 ExecutionFlags;
 
 	/** Configures whether this processor should be automatically included in the global list of processors executed every tick (see ProcessingPhase and ExecutionOrder). */
 	UPROPERTY(EditDefaultsOnly, Category = Processor, config)
-	bool bAutoRegisterWithProcessingPhases = true;
+	uint8 bAutoRegisterWithProcessingPhases : 1 = true;
 
 	/** Meant as a class property, make sure to set it in subclass' constructor. Controls whether there can be multiple
 	 *  instances of a given class in a single FMassRuntimePipeline and during dependency solving. */
-	bool bAllowMultipleInstances = false;
+	uint8 bAllowMultipleInstances : 1 = false;
+
+private:
+	/**
+	 * Gets set to true when an instance of the processor gets added to the phase processing as a "dynamic processor".
+	 * Once set it's never expected to be cleared out to `false` thus the private visibility of the member variable.
+	 * @see MarkAsDynamic()
+	 * @see IsDynamic()
+	 */
+	uint8 bIsDynamic : 1 = false;
 
-	UE_DEPRECATED(5.2, "This property is deprecated. Use bAllowMultipleInstances instead")
-	/** meant as a class property, make sure to set it in subclass' constructor. Controls whether there can be multiple 
-	 *  instances of a given class in a single FMassRuntimePipeline */
-	bool bAllowDuplicates = false;
+	/** Used to track whether Initialized has been called. */
+	uint8 bInitialized : 1 = false;
 
+protected:
 	UPROPERTY(EditDefaultsOnly, Category = Processor, config)
-	bool bRequiresGameThreadExecution = false;
+	uint8 bRequiresGameThreadExecution : 1 = false;
 
 #if WITH_EDITORONLY_DATA
 	/** Used to permanently remove a given processor class from PipeSetting's listing. Used primarily for test-time 
 	 *  processor classes, but can also be used by project-specific code to prune the processor list. */
 	UPROPERTY(config)
-	bool bCanShowUpInSettings = true;
+	uint8 bCanShowUpInSettings : 1 = true;
 #endif // WITH_EDITORONLY_DATA
 
 	friend UMassCompositeProcessor;
@@ -168,9 +176,6 @@ private:
 	 *  @note that it's safe to store pointers here since RegisterQuery does verify that a given registered query is 
 	 *  a member variable of a given processor */
 	TArray<FMassEntityQuery*> OwnedQueries;
-
-	/** Used to track whether Initialized has been called. */
-	bool bInitialized = false;
 };
 
 
