[Mass] Added a way to clear out specific tag requirement from FMassFragmentRequirements

diff --git a/Engine/Source/Runtime/MassEntity/Private/MassRequirements.cpp b/Engine/Source/Runtime/MassEntity/Private/MassRequirements.cpp
--- a/Engine/Source/Runtime/MassEntity/Private/MassRequirements.cpp
+++ b/Engine/Source/Runtime/MassEntity/Private/MassRequirements.cpp
@@ -86,6 +86,16 @@ FMassFragmentRequirements::FMassFragmentRequirements(TConstArrayView<const UScri
 	}
 }
 
+FMassFragmentRequirements& FMassFragmentRequirements::ClearTagRequirements(const FMassTagBitSet& TagsToRemoveBitSet)
+{
+	RequiredAllTags.Remove(TagsToRemoveBitSet);
+	RequiredAnyTags.Remove(TagsToRemoveBitSet);
+	RequiredNoneTags.Remove(TagsToRemoveBitSet);
+	RequiredOptionalTags.Remove(TagsToRemoveBitSet);
+
+	return *this;
+}
+
 void FMassFragmentRequirements::SortRequirements()
 {
 	// we're sorting the Requirements the same way ArchetypeData's FragmentConfig is sorted (see FMassArchetypeData::Initialize)
diff --git a/Engine/Source/Runtime/MassEntity/Public/MassRequirements.h b/Engine/Source/Runtime/MassEntity/Public/MassRequirements.h
--- a/Engine/Source/Runtime/MassEntity/Public/MassRequirements.h
+++ b/Engine/Source/Runtime/MassEntity/Public/MassRequirements.h
@@ -284,6 +284,9 @@ public:
 		return *this;
 	}
 
+	/** Clears given tags out of all collected requirements, including negative ones */
+	FMassFragmentRequirements& ClearTagRequirements(const FMassTagBitSet& TagsToRemoveBitSet);
+
 	template<typename T>
 	FMassFragmentRequirements& AddChunkRequirement(const EMassFragmentAccess AccessMode, const EMassFragmentPresence Presence = EMassFragmentPresence::All)
 	{
