[Mass] late-summer cleaning:
* moved UMassProcessor's inlines to the end of the header to keep the API look clean
* moved the virtuals defined in the header over to the cpp
* comments spellig fixes and improvements.

diff --git a/Engine/Source/Editor/MassEntityDebugger/Private/MassDebuggerModel.cpp b/Engine/Source/Editor/MassEntityDebugger/Private/MassDebuggerModel.cpp
--- a/Engine/Source/Editor/MassEntityDebugger/Private/MassDebuggerModel.cpp
+++ b/Engine/Source/Editor/MassEntityDebugger/Private/MassDebuggerModel.cpp
@@ -58,7 +58,7 @@ namespace UE::Mass::Debugger::Private
 			&& Processor.HasAnyFlags(RF_ClassDefaultObject) == false
 			&& Processor.GetWorld() == ContextWorld 
 			// checking ContextWorld is a cheaper way of supporting the declared behavior, since if there is a world then
-			// the processors are definitelly not CDO owned (by design). Is there is no world we need to check specifically.
+			// the processors are definitely not CDO owned (by design). Is there is no world we need to check specifically.
 			&& (ContextWorld != nullptr || Processor.GetOuter()->HasAnyFlags(RF_ClassDefaultObject) == false);
 	}
 } // namespace UE::Mass::Debugger::Private
diff --git a/Engine/Source/Runtime/MassEntity/Private/MassProcessingPhaseManager.cpp b/Engine/Source/Runtime/MassEntity/Private/MassProcessingPhaseManager.cpp
--- a/Engine/Source/Runtime/MassEntity/Private/MassProcessingPhaseManager.cpp
+++ b/Engine/Source/Runtime/MassEntity/Private/MassProcessingPhaseManager.cpp
@@ -178,6 +178,8 @@ void FMassPhaseProcessorConfigurationHelper::Configure(TArrayView<UMassProcessor
 		for (UMassProcessor* Processor : TmpPipeline.GetProcessors())
 		{
 			UE_VLOG_UELOG(&PhaseProcessor, LogMass, Verbose, TEXT("\t%s"), *Processor->GetProcessorName());
+			// we're ignoring dynamic processors since these instances were explicitly created by the user, and we don't
+			// want to destroy them - the user would need to re-add them every time we recalculate the dependency graph
 			if (Processor->IsDynamic() == false)
 			{
 				Processor->MarkAsGarbage();
diff --git a/Engine/Source/Runtime/MassEntity/Private/MassProcessor.cpp b/Engine/Source/Runtime/MassEntity/Private/MassProcessor.cpp
--- a/Engine/Source/Runtime/MassEntity/Private/MassProcessor.cpp
+++ b/Engine/Source/Runtime/MassEntity/Private/MassProcessor.cpp
@@ -122,6 +122,11 @@ UMassProcessor::UMassProcessor()
 {
 }
 
+void UMassProcessor::Initialize(UObject& Owner)
+{
+	bInitialized = true;
+}
+
 void UMassProcessor::SetShouldAutoRegisterWithGlobalList(const bool bAutoRegister)
 {	
 	if (ensureMsgf(HasAnyFlags(RF_ClassDefaultObject), TEXT("Setting bAutoRegisterWithProcessingPhases for non-CDOs has no effect")))
@@ -221,6 +226,21 @@ void UMassProcessor::CallExecute(FMassEntityManager& EntityManager, FMassExecuti
 	}
 }
 
+bool UMassProcessor::ShouldAllowQueryBasedPruning(const bool bRuntimeMode) const
+{
+	return bRuntimeMode;
+}
+
+EMassProcessingPhase UMassProcessor::GetProcessingPhase() const
+{
+	return ProcessingPhase;
+}
+
+void UMassProcessor::SetProcessingPhase(EMassProcessingPhase Phase)
+{
+	ProcessingPhase = Phase;
+}
+
 void UMassProcessor::ExportRequirements(FMassExecutionRequirements& OutRequirements) const
 {
 	for (FMassEntityQuery* Query : OwnedQueries)
@@ -263,6 +283,11 @@ FGraphEventRef UMassProcessor::DispatchProcessorTasks(const TSharedPtr<FMassEnti
 	return ReturnVal;
 }
 
+FString UMassProcessor::GetProcessorName() const
+{
+	return GetName();
+}
+
 void UMassProcessor::DebugOutputDescription(FOutputDevice& Ar, int32 Indent) const
 {
 #if WITH_MASSENTITY_DEBUG
@@ -465,6 +490,11 @@ void UMassCompositeProcessor::UpdateProcessorsCollection(TArrayView<FMassProcess
 	}
 }
 
+FString UMassCompositeProcessor::GetProcessorName() const
+{
+	return GroupName.ToString();
+}
+
 void UMassCompositeProcessor::DebugOutputDescription(FOutputDevice& Ar, int32 Indent) const
 {
 #if WITH_MASSENTITY_DEBUG
diff --git a/Engine/Source/Runtime/MassEntity/Public/MassProcessor.h b/Engine/Source/Runtime/MassEntity/Public/MassProcessor.h
--- a/Engine/Source/Runtime/MassEntity/Public/MassProcessor.h
+++ b/Engine/Source/Runtime/MassEntity/Public/MassProcessor.h
@@ -45,27 +45,28 @@ class MASSENTITY_API UMassProcessor : public UObject
 	GENERATED_BODY()
 public:
 	UMassProcessor();
-	UMassProcessor(const FObjectInitializer& ObjectInitializer);
+	explicit UMassProcessor(const FObjectInitializer& ObjectInitializer);
 
-	bool IsInitialized() const { return bInitialized; }
-	virtual void Initialize(UObject& Owner) { bInitialized = true; }
+	bool IsInitialized() const;
+	/** Called to initialize the processor. Override to perform custom steps. The super implementation has to be called. */
+	virtual void Initialize(UObject& Owner);
 	virtual FGraphEventRef DispatchProcessorTasks(const TSharedPtr<FMassEntityManager>& EntityManager, FMassExecutionContext& ExecutionContext, const FGraphEventArray& Prerequisites = FGraphEventArray());
 
-	EProcessorExecutionFlags GetExecutionFlags() const { return (EProcessorExecutionFlags)ExecutionFlags; }
+	EProcessorExecutionFlags GetExecutionFlags() const;
 
 	/** Whether this processor should execute according the CurrentExecutionFlags parameters */
-	bool ShouldExecute(const EProcessorExecutionFlags CurrentExecutionFlags) const { return (GetExecutionFlags() & CurrentExecutionFlags) != EProcessorExecutionFlags::None; }
+	bool ShouldExecute(const EProcessorExecutionFlags CurrentExecutionFlags) const;
 	void CallExecute(FMassEntityManager& EntityManager, FMassExecutionContext& Context);
 
 	/** 
 	 * Controls whether there can be multiple instances of a given class in a single FMassRuntimePipeline and during 
 	 * dependency solving. 
 	 */
-	bool ShouldAllowMultipleInstances() const { return bAllowMultipleInstances; }
+	bool ShouldAllowMultipleInstances() const;
 
-	void DebugOutputDescription(FOutputDevice& Ar) const { DebugOutputDescription(Ar, 0); }
+	void DebugOutputDescription(FOutputDevice& Ar) const;
 	virtual void DebugOutputDescription(FOutputDevice& Ar, int32 Indent) const;
-	virtual FString GetProcessorName() const { return GetName(); }
+	virtual FString GetProcessorName() const;
 	
 	//----------------------------------------------------------------------//
 	// Ordering functions 
@@ -75,31 +76,31 @@ public:
 	 *  building. This can also happen for special processors that don't register any queries - if that's the case override 
 	 *  this function to return an appropriate value
 	 *  @param bRuntimeMode indicates whether the pruning is being done for game runtime (true) or editor-time presentation (false) */
-	virtual bool ShouldAllowQueryBasedPruning(const bool bRuntimeMode = true) const { return bRuntimeMode; }
+	virtual bool ShouldAllowQueryBasedPruning(const bool bRuntimeMode = true) const;
 
-	virtual EMassProcessingPhase GetProcessingPhase() const { return ProcessingPhase; }
-	virtual void SetProcessingPhase(EMassProcessingPhase Phase) { ProcessingPhase = Phase; }
-	bool DoesRequireGameThreadExecution() const { return bRequiresGameThreadExecution; }
+	virtual EMassProcessingPhase GetProcessingPhase() const;
+	virtual void SetProcessingPhase(EMassProcessingPhase Phase);
+	bool DoesRequireGameThreadExecution() const;
 	
-	const FMassProcessorExecutionOrder& GetExecutionOrder() const { return ExecutionOrder; }
+	const FMassProcessorExecutionOrder& GetExecutionOrder() const;
 
-	/** By default fetches requirements declared entity queries registered via RegisterQuery. Processors can override 
+	/** By default,  fetches requirements declared entity queries registered via RegisterQuery. Processors can override 
 	 *	this function to supply additional requirements */
 	virtual void ExportRequirements(FMassExecutionRequirements& OutRequirements) const;
 
-	const FMassSubsystemRequirements& GetProcessorRequirements() const { return ProcessorRequirements; }
+	const FMassSubsystemRequirements& GetProcessorRequirements() const;
 
 	/** Adds Query to RegisteredQueries list. Query is required to be a member variable of this processor. Not meeting
 	 *  this requirement will cause check failure and the query won't be registered. */
 	void RegisterQuery(FMassEntityQuery& Query);
 
-	void MarkAsDynamic() { bIsDynamic = true; }
-	bool IsDynamic() const { return bIsDynamic != 0; }
+	void MarkAsDynamic();
+	bool IsDynamic() const;
 
-	bool ShouldAutoAddToGlobalList() const { return bAutoRegisterWithProcessingPhases; }
+	bool ShouldAutoAddToGlobalList() const;
 #if WITH_EDITOR
-	bool ShouldShowUpInSettings() const { return ShouldAutoAddToGlobalList() || bCanShowUpInSettings; }
-#endif // WITH_EDITORONLY_DATA
+	bool ShouldShowUpInSettings() const;
+#endif // WITH_EDITOR
 
 	/** Sets bAutoRegisterWithProcessingPhases. Setting it to true will result in this processor class being always 
 	 * instantiated to be automatically evaluated every frame. @see FMassProcessingPhaseManager
@@ -146,6 +147,9 @@ private:
 	/**
 	 * Gets set to true when an instance of the processor gets added to the phase processing as a "dynamic processor".
 	 * Once set it's never expected to be cleared out to `false` thus the private visibility of the member variable.
+	 * A "dynamic" processor is a one that has bAutoRegisterWithProcessingPhases == false, meaning it's not automatically
+	 * added to the processing graph. Additionally, making processors dynamic allows one to have multiple instances
+	 * of processors of the same class. 
 	 * @see MarkAsDynamic()
 	 * @see IsDynamic()
 	 */
@@ -206,7 +210,7 @@ public:
 	virtual void SetProcessingPhase(EMassProcessingPhase Phase) override;
 
 	void SetGroupName(FName NewName);
-	FName GetGroupName() const { return GroupName; }
+	FName GetGroupName() const;
 
 	virtual void SetProcessors(TArrayView<UMassProcessor*> InProcessorInstances, const TSharedPtr<FMassEntityManager>& EntityManager = nullptr);
 
@@ -231,9 +235,9 @@ public:
 
 	virtual FGraphEventRef DispatchProcessorTasks(const TSharedPtr<FMassEntityManager>& EntityManager, FMassExecutionContext& ExecutionContext, const FGraphEventArray& Prerequisites = FGraphEventArray()) override;
 
-	bool IsEmpty() const { return ChildPipeline.IsEmpty(); }
+	bool IsEmpty() const;
 
-	virtual FString GetProcessorName() const override { return GroupName.ToString(); }
+	virtual FString GetProcessorName() const override;
 
 protected:
 	virtual void ConfigureQueries() override;
@@ -284,3 +288,82 @@ public:
 	UE_DEPRECATED(5.3, "Populate is deprecated. Please use UpdateProcessorsCollection instead.")
 	void Populate(TConstArrayView<FMassProcessorOrderInfo> OrderedProcessors);
 };
+
+
+//-----------------------------------------------------------------------------
+// UMassProcessor inlines
+//-----------------------------------------------------------------------------
+inline bool UMassProcessor::IsInitialized() const
+{
+	return bInitialized;
+}
+
+inline EProcessorExecutionFlags UMassProcessor::GetExecutionFlags() const
+{
+	return static_cast<EProcessorExecutionFlags>(ExecutionFlags);
+}
+
+inline bool UMassProcessor::ShouldExecute(const EProcessorExecutionFlags CurrentExecutionFlags) const
+{
+	return (GetExecutionFlags() & CurrentExecutionFlags) != EProcessorExecutionFlags::None;
+}
+
+inline bool UMassProcessor::ShouldAllowMultipleInstances() const
+{
+	return bAllowMultipleInstances;
+}
+
+inline void UMassProcessor::DebugOutputDescription(FOutputDevice& Ar) const
+{
+	DebugOutputDescription(Ar, 0);
+}
+
+inline bool UMassProcessor::DoesRequireGameThreadExecution() const
+{
+	return bRequiresGameThreadExecution;
+}
+	
+inline const FMassProcessorExecutionOrder& UMassProcessor::GetExecutionOrder() const
+{
+	return ExecutionOrder;
+}
+
+inline const FMassSubsystemRequirements& UMassProcessor::GetProcessorRequirements() const
+{
+	return ProcessorRequirements;
+}
+
+inline void UMassProcessor::MarkAsDynamic()
+{
+	bIsDynamic = true;
+}
+
+inline bool UMassProcessor::IsDynamic() const
+{
+	return bIsDynamic != 0;
+}
+
+inline bool UMassProcessor::ShouldAutoAddToGlobalList() const
+{
+	return bAutoRegisterWithProcessingPhases;
+}
+
+#if WITH_EDITOR
+inline bool UMassProcessor::ShouldShowUpInSettings() const
+{
+	return ShouldAutoAddToGlobalList() || bCanShowUpInSettings;
+}
+#endif // WITH_EDITOR
+
+//-----------------------------------------------------------------------------
+// UMassCompositeProcessor inlines
+//-----------------------------------------------------------------------------
+inline FName UMassCompositeProcessor::GetGroupName() const
+{
+	return GroupName;
+}
+
+inline bool UMassCompositeProcessor::IsEmpty() const
+{
+	return ChildPipeline.IsEmpty();
+}
