[AutoRTFM] Use transactionally safe counters and mutexes in MassEntityManager.

diff --git a/Engine/Source/Runtime/Core/Public/Misc/TransactionallySafeMutex.h b/Engine/Source/Runtime/Core/Public/Misc/TransactionallySafeMutex.h
--- a/Engine/Source/Runtime/Core/Public/Misc/TransactionallySafeMutex.h
+++ b/Engine/Source/Runtime/Core/Public/Misc/TransactionallySafeMutex.h
@@ -132,6 +132,11 @@ struct FTransactionallySafeMutexDefinition final
 		}
 	}
 
+	bool IsLocked() const
+	{
+		return State->TransactionalLockCount != 0;
+	}
+
 private:
 	UE_NONCOPYABLE(FTransactionallySafeMutexDefinition)
 
diff --git a/Engine/Source/Runtime/MassEntity/Private/MassEntityManagerStorage.cpp b/Engine/Source/Runtime/MassEntity/Private/MassEntityManagerStorage.cpp
--- a/Engine/Source/Runtime/MassEntity/Private/MassEntityManagerStorage.cpp
+++ b/Engine/Source/Runtime/MassEntity/Private/MassEntityManagerStorage.cpp
@@ -422,7 +422,7 @@ namespace UE::Mass
 		EntityData.bIsAllocated = 1;
 		int32 SerialNumber = EntityData.GetSerialNumber();
 
-		EntityCount.fetch_add(1llu);
+		EntityCount.Increment();
 	
 		FMassEntityHandle Handle;
 		Handle.SerialNumber = SerialNumber;
@@ -471,7 +471,7 @@ namespace UE::Mass
 			CountLeft -= CountToProcess;
 		}
 
-		EntityCount.fetch_add(CountAdded);
+		EntityCount.Add(CountAdded);
 
 		return CountAdded;
 	}
@@ -526,7 +526,7 @@ namespace UE::Mass
 		// Free any remaining handles
 		FreeRunOfHandles();
 
-		EntityCount.fetch_sub(DeallocateCount);
+		EntityCount.Subtract(DeallocateCount);
 	
 		return DeallocateCount;
 	}
@@ -558,7 +558,7 @@ namespace UE::Mass
 			}
 		}
 
-		EntityCount.fetch_sub(Handles.Num());
+		EntityCount.Subtract(Handles.Num());
 	
 		return Handles.Num();
 	}
diff --git a/Engine/Source/Runtime/MassEntity/Public/MassEntityManagerStorage.h b/Engine/Source/Runtime/MassEntity/Public/MassEntityManagerStorage.h
--- a/Engine/Source/Runtime/MassEntity/Public/MassEntityManagerStorage.h
+++ b/Engine/Source/Runtime/MassEntity/Public/MassEntityManagerStorage.h
@@ -3,6 +3,9 @@
 
 #include "Async/Mutex.h"
 #include "Containers/ChunkedArray.h"
+#include "HAL/ThreadSafeCounter.h"
+#include "Misc/TransactionallySafeCounter.h"
+#include "Misc/TransactionallySafeMutex.h"
 #include "Misc/TVariant.h"
 #include "Templates/SharedPointer.h"
 #include "MassProcessingTypes.h"
@@ -217,17 +220,17 @@ namespace UE::Mass
 		bool AddPage();
 
 		// Number of allocated Entities
-		std::atomic<uint64> EntityCount = 0;
+		TTransactionallySafeCounter<FThreadSafeCounter> EntityCount = 0;
 		uint32 MaximumEntityCountShift = 0;
 		uint32 MaxEntitiesPerPage = 0;
 		uint32 PageCount = 0;
 		// ALWAYS acquire FreeListMutex before this one
-		UE::FMutex PageAllocateMutex;
+		FTransactionallySafeMutex PageAllocateMutex;
 		// Pointer to array of pages
 		FEntityData** EntityPages = nullptr;
 	
 		TArray<int32> EntityFreeIndexList;
-		UE::FMutex FreeListMutex;
+		FTransactionallySafeMutex FreeListMutex;
 	};
 	
 } // namespace UE::Mass
\ No newline at end of file
