[Backout] - CL36640763
Original CL Desc
-----------------------------------------------------------------
[AutoRTFM] Use transactionally safe counters and mutexes in MassEntityManager.

diff --git a/Engine/Source/Runtime/Core/Public/Misc/TransactionallySafeMutex.h b/Engine/Source/Runtime/Core/Public/Misc/TransactionallySafeMutex.h
--- a/Engine/Source/Runtime/Core/Public/Misc/TransactionallySafeMutex.h
+++ b/Engine/Source/Runtime/Core/Public/Misc/TransactionallySafeMutex.h
@@ -132,11 +132,6 @@ struct FTransactionallySafeMutexDefinition final
 		}
 	}
 
-	bool IsLocked() const
-	{
-		return State->TransactionalLockCount != 0;
-	}
-
 private:
 	UE_NONCOPYABLE(FTransactionallySafeMutexDefinition)
 
diff --git a/Engine/Source/Runtime/MassEntity/Private/MassEntityManagerStorage.cpp b/Engine/Source/Runtime/MassEntity/Private/MassEntityManagerStorage.cpp
--- a/Engine/Source/Runtime/MassEntity/Private/MassEntityManagerStorage.cpp
+++ b/Engine/Source/Runtime/MassEntity/Private/MassEntityManagerStorage.cpp
@@ -422,7 +422,7 @@ namespace UE::Mass
 		EntityData.bIsAllocated = 1;
 		int32 SerialNumber = EntityData.GetSerialNumber();
 
-		EntityCount.Increment();
+		EntityCount.fetch_add(1llu);
 	
 		FMassEntityHandle Handle;
 		Handle.SerialNumber = SerialNumber;
@@ -471,7 +471,7 @@ namespace UE::Mass
 			CountLeft -= CountToProcess;
 		}
 
-		EntityCount.Add(CountAdded);
+		EntityCount.fetch_add(CountAdded);
 
 		return CountAdded;
 	}
@@ -526,7 +526,7 @@ namespace UE::Mass
 		// Free any remaining handles
 		FreeRunOfHandles();
 
-		EntityCount.Subtract(DeallocateCount);
+		EntityCount.fetch_sub(DeallocateCount);
 	
 		return DeallocateCount;
 	}
@@ -558,7 +558,7 @@ namespace UE::Mass
 			}
 		}
 
-		EntityCount.Subtract(Handles.Num());
+		EntityCount.fetch_sub(Handles.Num());
 	
 		return Handles.Num();
 	}
diff --git a/Engine/Source/Runtime/MassEntity/Public/MassEntityManagerStorage.h b/Engine/Source/Runtime/MassEntity/Public/MassEntityManagerStorage.h
--- a/Engine/Source/Runtime/MassEntity/Public/MassEntityManagerStorage.h
+++ b/Engine/Source/Runtime/MassEntity/Public/MassEntityManagerStorage.h
@@ -3,9 +3,6 @@
 
 #include "Async/Mutex.h"
 #include "Containers/ChunkedArray.h"
-#include "HAL/ThreadSafeCounter.h"
-#include "Misc/TransactionallySafeCounter.h"
-#include "Misc/TransactionallySafeMutex.h"
 #include "Misc/TVariant.h"
 #include "Templates/SharedPointer.h"
 #include "MassProcessingTypes.h"
@@ -220,17 +217,17 @@ namespace UE::Mass
 		bool AddPage();
 
 		// Number of allocated Entities
-		TTransactionallySafeCounter<FThreadSafeCounter> EntityCount = 0;
+		std::atomic<uint64> EntityCount = 0;
 		uint32 MaximumEntityCountShift = 0;
 		uint32 MaxEntitiesPerPage = 0;
 		uint32 PageCount = 0;
 		// ALWAYS acquire FreeListMutex before this one
-		FTransactionallySafeMutex PageAllocateMutex;
+		UE::FMutex PageAllocateMutex;
 		// Pointer to array of pages
 		FEntityData** EntityPages = nullptr;
 	
 		TArray<int32> EntityFreeIndexList;
-		FTransactionallySafeMutex FreeListMutex;
+		UE::FMutex FreeListMutex;
 	};
 	
 } // namespace UE::Mass
\ No newline at end of file
