[AutoRTFM] Update MassCommandBuffer for transactional safety.

We now use a transactionally-safe MTAccessDetector and command counter.

diff --git a/Engine/Source/Runtime/MassEntity/Public/MassCommands.h b/Engine/Source/Runtime/MassEntity/Public/MassCommands.h
--- a/Engine/Source/Runtime/MassEntity/Public/MassCommands.h
+++ b/Engine/Source/Runtime/MassEntity/Public/MassCommands.h
@@ -3,7 +3,8 @@
 #pragma once
 
 #include "ProfilingDebugging/CsvProfilerConfig.h"
-#include "Misc/MTAccessDetector.h"
+#include "AutoRTFM/AutoRTFM.h"
+#include "Misc/MTTransactionallySafeAccessDetector.h"
 #include "MassEntityTypes.h"
 #include "MassEntityUtils.h"
 #include "MassEntityManager.h"
@@ -97,8 +98,9 @@ struct MASSENTITY_API FMassBatchedCommand
 
 	bool HasWork() const { return bHasWork; }
 	EMassCommandOperationType GetOperationType() const { return OperationType; }
-
+	
 	template<typename T>
+	UE_AUTORTFM_ALWAYS_OPEN
 	static uint32 GetCommandIndex()
 	{
 		static const uint32 ThisTypesStaticIndex = CommandsCounter++;
@@ -165,7 +167,7 @@ protected:
 	virtual int32 GetNumOperationsStat() const override { return TargetEntities.Num(); }
 #endif // CSV_PROFILER_STATS || WITH_MASSENTITY_DEBUG
 
-	UE_MT_DECLARE_RW_ACCESS_DETECTOR(EntitiesAccessDetector); 
+	UE_MT_DECLARE_TS_RW_ACCESS_DETECTOR(EntitiesAccessDetector); 
 	TArray<FMassEntityHandle> TargetEntities;
 };
 
