[Mass] Renamed all `FMassEntityManager& System` to `EntityManager in MassCommands header

#rnx

diff --git a/Engine/Source/Runtime/MassEntity/Public/MassCommands.h b/Engine/Source/Runtime/MassEntity/Public/MassCommands.h
--- a/Engine/Source/Runtime/MassEntity/Public/MassCommands.h
+++ b/Engine/Source/Runtime/MassEntity/Public/MassCommands.h
@@ -90,7 +90,7 @@ struct MASSENTITY_API FMassBatchedCommand
 #endif // CSV_PROFILER_STATS || WITH_MASSENTITY_DEBUG
 	virtual ~FMassBatchedCommand() { Reset(); }
 
-	virtual void Execute(FMassEntityManager& System) const = 0;
+	virtual void Execute(FMassEntityManager& EntityManager) const = 0;
 	virtual void Reset()
 	{
 		bHasWork = false;
@@ -184,13 +184,13 @@ struct FMassCommandDestroyEntities : public FMassBatchedEntityCommand
 	}
 
 protected:
-	virtual void Execute(FMassEntityManager& System) const override
+	virtual void Execute(FMassEntityManager& EntityManager) const override
 	{
 		TRACE_CPUPROFILER_EVENT_SCOPE(MassCommandDestroyEntities_Execute);
 
 		TArray<FMassArchetypeEntityCollection> EntityCollectionsToDestroy;
-		UE::Mass::Utils::CreateEntityCollections(System, TargetEntities, FMassArchetypeEntityCollection::FoldDuplicates, EntityCollectionsToDestroy);
-		System.BatchDestroyEntityChunks(EntityCollectionsToDestroy);
+		UE::Mass::Utils::CreateEntityCollections(EntityManager, TargetEntities, FMassArchetypeEntityCollection::FoldDuplicates, EntityCollectionsToDestroy);
+		EntityManager.BatchDestroyEntityChunks(EntityCollectionsToDestroy);
 	}
 };
 
@@ -207,12 +207,12 @@ struct FMassCommandAddFragmentsInternal : public FMassBatchedEntityCommand
 	{}
 
 protected:
-	virtual void Execute(FMassEntityManager& System) const override
+	virtual void Execute(FMassEntityManager& EntityManager) const override
 	{
 		TRACE_CPUPROFILER_EVENT_SCOPE(MassCommandAddFragments_Execute);
 		TArray<FMassArchetypeEntityCollection> EntityCollections;
-		UE::Mass::Utils::CreateEntityCollections(System, TargetEntities, FMassArchetypeEntityCollection::FoldDuplicates, EntityCollections);
-		System.BatchChangeFragmentCompositionForEntities(EntityCollections, FragmentsAffected, FMassFragmentBitSet());
+		UE::Mass::Utils::CreateEntityCollections(EntityManager, TargetEntities, FMassArchetypeEntityCollection::FoldDuplicates, EntityCollections);
+		EntityManager.BatchChangeFragmentCompositionForEntities(EntityCollections, FragmentsAffected, FMassFragmentBitSet());
 	}
 	FMassFragmentBitSet FragmentsAffected;
 };
@@ -230,12 +230,12 @@ struct FMassCommandRemoveFragmentsInternal : public FMassBatchedEntityCommand
 	{}
 
 protected:
-	virtual void Execute(FMassEntityManager& System) const override
+	virtual void Execute(FMassEntityManager& EntityManager) const override
 	{
 		TRACE_CPUPROFILER_EVENT_SCOPE(MassCommandRemoveFragments_Execute);
 		TArray<FMassArchetypeEntityCollection> EntityCollections;
-		UE::Mass::Utils::CreateEntityCollections(System, TargetEntities, FMassArchetypeEntityCollection::FoldDuplicates, EntityCollections);
-		System.BatchChangeFragmentCompositionForEntities(EntityCollections, FMassFragmentBitSet(), FragmentsAffected);
+		UE::Mass::Utils::CreateEntityCollections(EntityManager, TargetEntities, FMassArchetypeEntityCollection::FoldDuplicates, EntityCollections);
+		EntityManager.BatchChangeFragmentCompositionForEntities(EntityCollections, FMassFragmentBitSet(), FragmentsAffected);
 	}
 	FMassFragmentBitSet FragmentsAffected;
 };
@@ -260,13 +260,13 @@ struct FMassCommandChangeTags : public FMassBatchedEntityCommand
 	{}
 
 protected:
-	virtual void Execute(FMassEntityManager& System) const override
+	virtual void Execute(FMassEntityManager& EntityManager) const override
 	{
 		TRACE_CPUPROFILER_EVENT_SCOPE(MassCommandChangeTags_Execute);
 		TArray<FMassArchetypeEntityCollection> EntityCollections;
-		UE::Mass::Utils::CreateEntityCollections(System, TargetEntities, FMassArchetypeEntityCollection::FoldDuplicates, EntityCollections);
+		UE::Mass::Utils::CreateEntityCollections(EntityManager, TargetEntities, FMassArchetypeEntityCollection::FoldDuplicates, EntityCollections);
 
-		System.BatchChangeTagsForEntities(EntityCollections, TagsToAdd, TagsToRemove);
+		EntityManager.BatchChangeTagsForEntities(EntityCollections, TagsToAdd, TagsToRemove);
 	}
 
 	virtual SIZE_T GetAllocatedSize() const override
@@ -363,7 +363,7 @@ protected:
 		return Super::GetAllocatedSize() + Fragments.GetAllocatedSize() + FragmentsAffected.GetAllocatedSize();
 	}
 
-	virtual void Execute(FMassEntityManager& System) const override
+	virtual void Execute(FMassEntityManager& EntityManager) const override
 	{
 		TRACE_CPUPROFILER_EVENT_SCOPE(MassCommandAddFragmentInstances_Execute);
 
@@ -372,10 +372,10 @@ protected:
 		Fragments.GetAsGenericMultiArray(GenericMultiArray);
 
 		TArray<FMassArchetypeEntityCollectionWithPayload> EntityCollections;
-		FMassArchetypeEntityCollectionWithPayload::CreateEntityRangesWithPayload(System, TargetEntities, FMassArchetypeEntityCollection::FoldDuplicates
+		FMassArchetypeEntityCollectionWithPayload::CreateEntityRangesWithPayload(EntityManager, TargetEntities, FMassArchetypeEntityCollection::FoldDuplicates
 			, FMassGenericPayloadView(GenericMultiArray), EntityCollections);
 
-		System.BatchAddFragmentInstancesForEntities(EntityCollections, FragmentsAffected);
+		EntityManager.BatchAddFragmentInstancesForEntities(EntityCollections, FragmentsAffected);
 	}
 
 	mutable UE::Mass::TMultiArray<TOthers...> Fragments;
@@ -394,7 +394,7 @@ struct FMassCommandBuildEntity : public FMassCommandAddFragmentInstances<TOthers
 
 protected:
 
-	virtual void Execute(FMassEntityManager& System) const override
+	virtual void Execute(FMassEntityManager& EntityManager) const override
 	{
 		TRACE_CPUPROFILER_EVENT_SCOPE(MassCommandBuildEntity_Execute);
 
@@ -403,13 +403,13 @@ protected:
 		Super::Fragments.GetAsGenericMultiArray(GenericMultiArray);
 
 		TArray<FMassArchetypeEntityCollectionWithPayload> EntityCollections;
-		FMassArchetypeEntityCollectionWithPayload::CreateEntityRangesWithPayload(System, Super::TargetEntities, FMassArchetypeEntityCollection::FoldDuplicates
+		FMassArchetypeEntityCollectionWithPayload::CreateEntityRangesWithPayload(EntityManager, Super::TargetEntities, FMassArchetypeEntityCollection::FoldDuplicates
 			, FMassGenericPayloadView(GenericMultiArray), EntityCollections);
 
 		check(EntityCollections.Num() <= 1);
 		if (EntityCollections.Num())
 		{
-			System.BatchBuildEntities(EntityCollections[0], Super::FragmentsAffected, FMassArchetypeSharedFragmentValues());
+			EntityManager.BatchBuildEntities(EntityCollections[0], Super::FragmentsAffected, FMassArchetypeSharedFragmentValues());
 		}
 	}
 };
@@ -457,7 +457,7 @@ protected:
 		return TotalSize;
 	}
 
-	virtual void Execute(FMassEntityManager& System) const 
+	virtual void Execute(FMassEntityManager& EntityManager) const 
 	{
 		TRACE_CPUPROFILER_EVENT_SCOPE(MassCommandBuildEntityWithSharedFragments_Execute);
 
@@ -470,13 +470,13 @@ protected:
 			It.Value.Fragments.GetAsGenericMultiArray(GenericMultiArray);
 
 			TArray<FMassArchetypeEntityCollectionWithPayload> EntityCollections;
-			FMassArchetypeEntityCollectionWithPayload::CreateEntityRangesWithPayload(System, It.Value.TargetEntities, FMassArchetypeEntityCollection::FoldDuplicates
+			FMassArchetypeEntityCollectionWithPayload::CreateEntityRangesWithPayload(EntityManager, It.Value.TargetEntities, FMassArchetypeEntityCollection::FoldDuplicates
 				, FMassGenericPayloadView(GenericMultiArray), EntityCollections);
 			checkf(EntityCollections.Num() <= 1, TEXT("We expect TargetEntities to only contain archetype-less entities, ones that need to be \'build\'"));
 
 			if (EntityCollections.Num())
 			{
-				System.BatchBuildEntities(EntityCollections[0], FragmentsAffected, It.Value.SharedFragmentValues);
+				EntityManager.BatchBuildEntities(EntityCollections[0], FragmentsAffected, It.Value.SharedFragmentValues);
 			}
 
 			GenericMultiArray.Reset();
@@ -530,7 +530,7 @@ template<EMassCommandOperationType OpType>
 struct FMassDeferredCommand : public FMassBatchedCommand
 {
 	using Super = FMassBatchedCommand;
-	using FExecFunction = TFunction<void(FMassEntityManager& System)>;
+	using FExecFunction = TFunction<void(FMassEntityManager& EntityManager)>;
 
 	FMassDeferredCommand()
 		: Super(OpType DEBUG_NAME("BatchedDeferredCommand"))
@@ -554,13 +554,13 @@ protected:
 		return DeferredFunctions.GetAllocatedSize();
 	}
 
-	virtual void Execute(FMassEntityManager& System) const override
+	virtual void Execute(FMassEntityManager& EntityManager) const override
 	{
 		TRACE_CPUPROFILER_EVENT_SCOPE(MassDeferredCommand_Execute);
 
 		for (const FExecFunction& ExecFunction : DeferredFunctions)
 		{
-			ExecFunction(System);
+			ExecFunction(EntityManager);
 		}
 	}
 
