[Mass] Added a concept of "ActivationState" for Mass processors. These are Active, OneShor and Inactive. Processor CDOs can also be configured (via project settings or DefaultMass.ini) which affects all instances created after that point
Made changes to MassDebugger to mark the deactivated processors.

diff --git a/Engine/Source/Editor/MassEntityDebugger/Private/MassDebuggerModel.cpp b/Engine/Source/Editor/MassEntityDebugger/Private/MassDebuggerModel.cpp
--- a/Engine/Source/Editor/MassEntityDebugger/Private/MassDebuggerModel.cpp
+++ b/Engine/Source/Editor/MassEntityDebugger/Private/MassDebuggerModel.cpp
@@ -197,6 +197,11 @@ void FMassDebuggerProcessorData::SetProcessor(const UMassProcessor& InProcessor)
 	UE::Mass::Debugger::Private::MakeDisplayName(Name, Label);
 
 	ProcessorHash = UE::Mass::Debugger::Private::CalcProcessorHash(InProcessor);
+	bIsActive = InProcessor.IsActive();
+	if (bIsActive == false)
+	{
+		Label.InsertAt(0, TEXT("[INACTIVE] "));
+	}
 
 #if WITH_MASSENTITY_DEBUG
 	FStringOutputDevice DescriptionDevice;
diff --git a/Engine/Source/Editor/MassEntityDebugger/Private/MassDebuggerModel.h b/Engine/Source/Editor/MassEntityDebugger/Private/MassDebuggerModel.h
--- a/Engine/Source/Editor/MassEntityDebugger/Private/MassDebuggerModel.h
+++ b/Engine/Source/Editor/MassEntityDebugger/Private/MassDebuggerModel.h
@@ -104,14 +104,15 @@ private:
 public:
 	FString Name;
 	FString Label;
-	uint32 ProcessorHash = 0; 
+	uint32 ProcessorHash = 0;
+	bool bIsActive = true;
+	EMassDebuggerProcessorSelection Selection = EMassDebuggerProcessorSelection::None;
 	TSharedPtr<FMassDebuggerQueryData> ProcessorRequirements;
 	TArray<TSharedPtr<FMassDebuggerQueryData>> Queries;
 	TArray<TSharedPtr<FMassDebuggerArchetypeData>> ValidArchetypes;
 #if WITH_MASSENTITY_DEBUG
 	FString Description;
 #endif // WITH_MASSENTITY_DEBUG	
-	EMassDebuggerProcessorSelection Selection = EMassDebuggerProcessorSelection::None;
 };
 
 struct FMassDebuggerProcessingGraphNode
diff --git a/Engine/Source/Editor/MassEntityDebugger/Private/SMassProcessingGraphView.cpp b/Engine/Source/Editor/MassEntityDebugger/Private/SMassProcessingGraphView.cpp
--- a/Engine/Source/Editor/MassEntityDebugger/Private/SMassProcessingGraphView.cpp
+++ b/Engine/Source/Editor/MassEntityDebugger/Private/SMassProcessingGraphView.cpp
@@ -69,17 +69,15 @@ public:
 								switch (Item->Node.GraphNodeSelection)
 								{
 									case EMassDebuggerProcessingGraphNodeSelection::None:
-										return FLinearColor::White;
-										break;
+										return Item->Node.ProcessorData->bIsActive
+											? FLinearColor::White
+											: FLinearColor::Gray;
 									case EMassDebuggerProcessingGraphNodeSelection::WaitFor:
 										return FMath::Lerp(Foreground, FLinearColor::Green, 0.75f);
-										break;
 									case EMassDebuggerProcessingGraphNodeSelection::Block:
 										return FMath::Lerp(Foreground, FLinearColor::Red, 0.75f);
-										break;
 									default:
 										return FLinearColor::White;
-										break;
 								}
 							})
 					]
diff --git a/Engine/Source/Editor/MassEntityDebugger/Private/SMassProcessorsView.cpp b/Engine/Source/Editor/MassEntityDebugger/Private/SMassProcessorsView.cpp
--- a/Engine/Source/Editor/MassEntityDebugger/Private/SMassProcessorsView.cpp
+++ b/Engine/Source/Editor/MassEntityDebugger/Private/SMassProcessorsView.cpp
@@ -45,6 +45,7 @@ public:
 		[
 			SNew(STextBlock)
 			.Text(FText::FromString(Item->Label))
+			.TextStyle(FAppStyle::Get(), Item->bIsActive ? "NormalText" : "NormalText.Subdued")
 		];
 	}
 
diff --git a/Engine/Source/Runtime/MassEntity/Private/MassExecutor.cpp b/Engine/Source/Runtime/MassEntity/Private/MassExecutor.cpp
--- a/Engine/Source/Runtime/MassEntity/Private/MassExecutor.cpp
+++ b/Engine/Source/Runtime/MassEntity/Private/MassExecutor.cpp
@@ -14,7 +14,10 @@ FORCEINLINE void ExecuteProcessors(FMassEntityManager& EntityManager, TArrayView
 {
 	for (UMassProcessor* Proc : Processors)
 	{
-		Proc->CallExecute(EntityManager, ExecutionContext);
+		if (LIKELY(Proc->IsActive()))
+		{
+			Proc->CallExecute(EntityManager, ExecutionContext);
+		}
 	}
 }
 
diff --git a/Engine/Source/Runtime/MassEntity/Private/MassProcessor.cpp b/Engine/Source/Runtime/MassEntity/Private/MassProcessor.cpp
--- a/Engine/Source/Runtime/MassEntity/Private/MassProcessor.cpp
+++ b/Engine/Source/Runtime/MassEntity/Private/MassProcessor.cpp
@@ -205,24 +205,32 @@ void UMassProcessor::PostInitProperties()
 
 void UMassProcessor::CallExecute(FMassEntityManager& EntityManager, FMassExecutionContext& Context)
 {
-	TRACE_CPUPROFILER_EVENT_SCOPE_TEXT(*StatId);
-	LLM_SCOPE_BYNAME(TEXT("Mass/ExecuteProcessor"));
-	// Not using a more specific scope by default (i.e. LLM_SCOPE_BYNAME(*StatId)) since LLM is more strict regarding the provided string (no spaces or '_')
+	if (LIKELY(ensureMsgf(IsActive(), TEXT("Trying to CallExecute for an inactive processor %s"), *GetProcessorName())))
+	{
+		TRACE_CPUPROFILER_EVENT_SCOPE_TEXT(*StatId);
+		LLM_SCOPE_BYNAME(TEXT("Mass/ExecuteProcessor"));
+		// Not using a more specific scope by default (i.e. LLM_SCOPE_BYNAME(*StatId)) since LLM is more strict regarding the provided string (no spaces or '_')
 
 #if WITH_MASSENTITY_DEBUG
-	Context.DebugSetExecutionDesc(FString::Printf(TEXT("%s (%s)"), *GetProcessorName(), EntityManager.GetWorld() ? *ToString(EntityManager.GetWorld()->GetNetMode()) : TEXT("No World")));
+		Context.DebugSetExecutionDesc(FString::Printf(TEXT("%s (%s)"), *GetProcessorName(), EntityManager.GetWorld() ? *ToString(EntityManager.GetWorld()->GetNetMode()) : TEXT("No World")));
 #endif
-	// CacheSubsystemRequirements will return true only if all requirements declared with ProcessorRequirements are met
-	// meaning if it fails there's no point in calling Execute.
-	// Note that we're not testing individual queries in OwnedQueries - processors can function just fine with some 
-	// of their queries not having anything to do.
-	if (Context.CacheSubsystemRequirements(ProcessorRequirements))
-	{
-		Execute(EntityManager, Context);
-	}
-	else
-	{
-		PROCESSOR_LOG(VeryVerbose, TEXT("%s Skipping Execute due to subsystem requirements not being met"), *GetProcessorName());
+		// CacheSubsystemRequirements will return true only if all requirements declared with ProcessorRequirements are met
+		// meaning if it fails there's no point in calling Execute.
+		// Note that we're not testing individual queries in OwnedQueries - processors can function just fine with some 
+		// of their queries not having anything to do.
+		if (Context.CacheSubsystemRequirements(ProcessorRequirements))
+		{
+			Execute(EntityManager, Context);
+		}
+		else
+		{
+			PROCESSOR_LOG(VeryVerbose, TEXT("%s Skipping Execute due to subsystem requirements not being met"), *GetProcessorName());
+		}
+
+		if (ActivationState == EActivationState::OneShot)
+		{
+			MakeInactive();
+		}
 	}
 }
 
@@ -272,14 +280,17 @@ void UMassProcessor::RegisterQuery(FMassEntityQuery& Query)
 FGraphEventRef UMassProcessor::DispatchProcessorTasks(const TSharedPtr<FMassEntityManager>& EntityManager, FMassExecutionContext& ExecutionContext, const FGraphEventArray& Prerequisites)
 {
 	FGraphEventRef ReturnVal;
-	if (bRequiresGameThreadExecution)
+	if (LIKELY(ensureMsgf(IsActive(), TEXT("Trying to dispatch processor task for inactive processor %s"), *GetProcessorName())))
 	{
-		ReturnVal = TGraphTask<FMassProcessorsTask_GameThread>::CreateTask(&Prerequisites).ConstructAndDispatchWhenReady(EntityManager, ExecutionContext, *this);
+		if (bRequiresGameThreadExecution)
+		{
+			ReturnVal = TGraphTask<FMassProcessorsTask_GameThread>::CreateTask(&Prerequisites).ConstructAndDispatchWhenReady(EntityManager, ExecutionContext, *this);
+		}
+		else
+		{
+			ReturnVal = TGraphTask<FMassProcessorTask>::CreateTask(&Prerequisites).ConstructAndDispatchWhenReady(EntityManager, ExecutionContext, *this);
+		}
 	}
-	else
-	{
-		ReturnVal = TGraphTask<FMassProcessorTask>::CreateTask(&Prerequisites).ConstructAndDispatchWhenReady(EntityManager, ExecutionContext, *this);
-	}	
 	return ReturnVal;
 }
 
@@ -321,20 +332,52 @@ void UMassCompositeProcessor::ConfigureQueries()
 FGraphEventRef UMassCompositeProcessor::DispatchProcessorTasks(const TSharedPtr<FMassEntityManager>& EntityManager, FMassExecutionContext& ExecutionContext, const FGraphEventArray& InPrerequisites)
 {
 	FGraphEventArray Events;
-	Events.Reserve(FlatProcessingGraph.Num());
+	Events.AddDefaulted(FlatProcessingGraph.Num());
+
+	FGraphEventArray Prerequisites;
+	// we'll fill this one with dependencies of disabled processors. We initialize it lazily
+	TArray<FGraphEventArray> AdditionalEvents;
 		
-	for (FDependencyNode& ProcessingNode : FlatProcessingGraph)
+	for (int32 NodeIndex = 0; NodeIndex < FlatProcessingGraph.Num(); ++NodeIndex)
 	{
-		FGraphEventArray Prerequisites;
-		for (const int32 DependencyIndex : ProcessingNode.Dependencies)
-		{
-			Prerequisites.Add(Events[DependencyIndex]);
-		}
+		FDependencyNode& ProcessingNode = FlatProcessingGraph[NodeIndex];
 
-		// we don't expect any group nodes at this point. If we get any there's a bug in dependencies solving
-		if (ensure(ProcessingNode.Processor))
+		if (ensureMsgf(ProcessingNode.Processor, TEXT("We don't expect any group nodes at this point. If we get any there's a bug in dependencies solving.")))
 		{
-			Events.Add(ProcessingNode.Processor->DispatchProcessorTasks(EntityManager, ExecutionContext, Prerequisites));
+			Prerequisites.Reset(ProcessingNode.Dependencies.Num());
+			for (const int32 DependencyIndex : ProcessingNode.Dependencies)
+			{
+				checkSlow(DependencyIndex < NodeIndex);
+				Prerequisites.Add(Events[DependencyIndex]);
+			}
+			// this means there are some inactive processors so we need to consider additional dependencies
+			if (AdditionalEvents.Num())
+			{
+				for (const int32 DependencyIndex : ProcessingNode.Dependencies)
+				{
+					Prerequisites.Append(AdditionalEvents[DependencyIndex]);
+				}
+			}
+
+			if (ProcessingNode.Processor->IsActive())
+			{
+				Events[NodeIndex] = ProcessingNode.Processor->DispatchProcessorTasks(EntityManager, ExecutionContext, Prerequisites);
+			}
+			else
+			{
+				if (AdditionalEvents.Num() == 0)
+				{
+					// lazy initialization
+					AdditionalEvents.AddDefaulted(FlatProcessingGraph.Num());
+				}
+				// if the processor is not going to run at all we store its Prerequisites so that
+				// processors waiting for this given processor to finish will keep their place
+				// in the overall processing graph
+				// NOTE: this is safer than just ignoring the dependencies since even though this
+				// processor is not running, the subsequent processors might unknowingly rely on
+				// implicit dependencies that the current processor was ensuring. 
+				AdditionalEvents[NodeIndex].Append(MoveTemp(Prerequisites));
+			}
 		}
 	}
 
@@ -352,8 +395,17 @@ FGraphEventRef UMassCompositeProcessor::DispatchProcessorTasks(const TSharedPtr<
 			}
 
 			check(ProcessingNode.Processor);
-			PROCESSOR_TASK_LOG(TEXT("Task %u %s%s%s"), Events[i]->GetTraceId(), *ProcessingNode.Processor->GetProcessorName()
-				, DependenciesDesc.Len() > 0 ? TEXT(" depends on ") : TEXT(""), *DependenciesDesc);
+			if (Events[i].IsValid())
+			{
+				PROCESSOR_TASK_LOG(TEXT("Task %u %s%s%s"), Events[i]->GetTraceId(), *ProcessingNode.Processor->GetProcessorName()
+				   , DependenciesDesc.Len() > 0 ? TEXT(" depends on ") : TEXT(""), *DependenciesDesc);
+			}
+			else
+			{
+				ensureMsgf(ProcessingNode.Processor->IsActive() == false, TEXT("This path is expected to trigger only for inactive processors"))
+				PROCESSOR_TASK_LOG(TEXT("Task [INACTIVE] %s%s%s"), *ProcessingNode.Processor->GetProcessorName()
+				   , DependenciesDesc.Len() > 0 ? TEXT(" depends on ") : TEXT(""), *DependenciesDesc);
+			}
 		}
 	}
 #endif // WITH_MASSENTITY_DEBUG
@@ -368,8 +420,10 @@ void UMassCompositeProcessor::Execute(FMassEntityManager& EntityManager, FMassEx
 {
 	for (UMassProcessor* Proc : ChildPipeline.GetMutableProcessors())
 	{
-		check(Proc);
-		Proc->CallExecute(EntityManager, Context);
+		if (LIKELY(ensure(Proc) && Proc->IsActive()))
+		{
+			Proc->CallExecute(EntityManager, Context);
+		}
 	}
 }
 
diff --git a/Engine/Source/Runtime/MassEntity/Public/MassProcessor.h b/Engine/Source/Runtime/MassEntity/Public/MassProcessor.h
--- a/Engine/Source/Runtime/MassEntity/Public/MassProcessor.h
+++ b/Engine/Source/Runtime/MassEntity/Public/MassProcessor.h
@@ -38,12 +38,21 @@ struct FMassProcessorExecutionOrder
 	TArray<FName> ExecuteAfter;
 };
 
+UENUM()
+enum class EActivationState : uint8
+	{
+		Inactive,
+		Active,
+		OneShot,	// one-shot processor will auto-disable itself after the next CallExecute call
+	};
 
 UCLASS(abstract, EditInlineNew, CollapseCategories, config = Mass, defaultconfig, ConfigDoNotCheckDefaults)
 class MASSENTITY_API UMassProcessor : public UObject
 {
 	GENERATED_BODY()
+
 public:
+
 	UMassProcessor();
 	explicit UMassProcessor(const FObjectInitializer& ObjectInitializer);
 
@@ -97,6 +106,23 @@ public:
 	void MarkAsDynamic();
 	bool IsDynamic() const;
 
+	/**
+	 * Marks processor as "Active" (@see ActivationState for details). If called during Mass processing the
+	 * call will take effect next phase.
+	 */
+	void MakeActive();
+	/**
+	 * Marks processor as "One Shot" (@see ActivationState for details). If called during Mass processing
+	 * the call will take effect next phase. The processor will auto-disable after execution.
+	 */
+	void MakeOneShot();
+	/**
+	 * Deactivate the processor, it will no longer execute its `Execute` function.
+	 */
+	void MakeInactive();
+
+	bool IsActive() const;
+
 	bool ShouldAutoAddToGlobalList() const;
 #if WITH_EDITOR
 	bool ShouldShowUpInSettings() const;
@@ -143,6 +169,16 @@ protected:
 	 *  instances of a given class in a single FMassRuntimePipeline and during dependency solving. */
 	uint8 bAllowMultipleInstances : 1 = false;
 
+	UPROPERTY(EditDefaultsOnly, Category = Processor, config)
+	uint8 bRequiresGameThreadExecution : 1 = false;
+
+#if WITH_EDITORONLY_DATA
+	/** Used to permanently remove a given processor class from PipeSetting's listing. Used primarily for test-time 
+	 *  processor classes, but can also be used by project-specific code to prune the processor list. */
+	UPROPERTY(config)
+	uint8 bCanShowUpInSettings : 1 = true;
+#endif // WITH_EDITORONLY_DATA
+
 private:
 	/**
 	 * Gets set to true when an instance of the processor gets added to the phase processing as a "dynamic processor".
@@ -158,16 +194,19 @@ private:
 	/** Used to track whether Initialized has been called. */
 	uint8 bInitialized : 1 = false;
 
-protected:
+	/**
+	 * Processors can be activated/deactivated at runtime. Deactivating a running processor will not disrupt the processing
+	 * graph since the disabled processor's dependencies will get passed down to the subsequent processors depending on this one.
+	 * Deactivating processor's CDO will result in every instance starting off as disabled. Those will still be considered
+	 * while building the processor dependency graph and one included in the processing graph will function just as the
+	 * processor instances disabled at runtime (i.e. won't run, but pass down their dependencies).
+	 * A special type of activation is "One Shot" mode - it works just like "Active" state, but it will auto-disable itself
+	 * upon completion of the next CallExecute call.
+	 */
 	UPROPERTY(EditDefaultsOnly, Category = Processor, config)
-	uint8 bRequiresGameThreadExecution : 1 = false;
+	EActivationState ActivationState = EActivationState::Active;
 
-#if WITH_EDITORONLY_DATA
-	/** Used to permanently remove a given processor class from PipeSetting's listing. Used primarily for test-time 
-	 *  processor classes, but can also be used by project-specific code to prune the processor list. */
-	UPROPERTY(config)
-	uint8 bCanShowUpInSettings : 1 = true;
-#endif // WITH_EDITORONLY_DATA
+protected:
 
 	friend UMassCompositeProcessor;
 	friend FMassDebugger;
@@ -343,6 +382,26 @@ inline bool UMassProcessor::IsDynamic() const
 	return bIsDynamic != 0;
 }
 
+inline void UMassProcessor::MakeActive()
+{
+	ActivationState = EActivationState::Active;
+}
+
+inline void UMassProcessor::MakeOneShot()
+{
+	ActivationState = EActivationState::OneShot;
+}
+
+inline void UMassProcessor::MakeInactive()
+{
+	ActivationState = EActivationState::Inactive;
+}
+
+inline bool UMassProcessor::IsActive() const
+{
+	return ActivationState != EActivationState::Inactive;
+}
+
 inline bool UMassProcessor::ShouldAutoAddToGlobalList() const
 {
 	return bAutoRegisterWithProcessingPhases;
