[Mass] modified how we access MassEntityManager's internal storage. The change gives compilers more clues that the value returned by GetEntityStorageInterface() doesn't change thus increasing chances of the access being optimized in loops.
Also renamed MASS_CONCURRENT_RESERVE to WITH_MASS_CONCURRENT_RESERVE and made it configurable via *.Build.cs files by setting REQUESTED_MASS_CONCURRENT_RESERVE to 1. Will still be forced to 1 in editor builds since TEDS requires it and is enabled by default

#rnx

diff --git a/Engine/Source/Developer/MassEntityTestSuite/Private/MassEntityStorageTest.cpp b/Engine/Source/Developer/MassEntityTestSuite/Private/MassEntityStorageTest.cpp
--- a/Engine/Source/Developer/MassEntityTestSuite/Private/MassEntityStorageTest.cpp
+++ b/Engine/Source/Developer/MassEntityTestSuite/Private/MassEntityStorageTest.cpp
@@ -4,7 +4,7 @@
 #include "MassEntityManager.h"
 
 // not point in "parity testing" if there's only sequential Mass storage available
-#if MASS_CONCURRENT_RESERVE
+#if WITH_MASS_CONCURRENT_RESERVE
 
 #define LOCTEXT_NAMESPACE "MassTest"
 
@@ -391,4 +391,4 @@ UE_ENABLE_OPTIMIZATION_SHIP
 
 #undef LOCTEXT_NAMESPACE
 
-#endif // MASS_CONCURRENT_RESERVE
\ No newline at end of file
+#endif // WITH_MASS_CONCURRENT_RESERVE
\ No newline at end of file
diff --git a/Engine/Source/Runtime/MassEntity/Private/MassEntityManager.cpp b/Engine/Source/Runtime/MassEntity/Private/MassEntityManager.cpp
--- a/Engine/Source/Runtime/MassEntity/Private/MassEntityManager.cpp
+++ b/Engine/Source/Runtime/MassEntity/Private/MassEntityManager.cpp
@@ -138,85 +138,6 @@ void FMassEntityManager::FEntityCreationContext::ForceUpdateCurrentThreadID()
 	OwnerThreadId = FPlatformTLS::GetCurrentThreadId();
 }
 
-//-----------------------------------------------------------------------------
-// FMassEntityManager
-//-----------------------------------------------------------------------------
-#if MASS_CONCURRENT_RESERVE
-UE::Mass::IEntityStorageInterface& FMassEntityManager::GetEntityStorageInterface()
-{
-	using namespace UE::Mass;
-	struct StorageSelector
-	{
-		UE::Mass::IEntityStorageInterface* operator()(FEmptyVariantState&) const
-		{
-			checkf(false, TEXT("Attempt to use EntityStorageInterface without initialization"));
-			return nullptr;
-		}
-		UE::Mass::IEntityStorageInterface* operator()(FSingleThreadedEntityStorage& Storage) const
-		{
-			return &Storage;
-		}
-		UE::Mass::IEntityStorageInterface* operator()(FConcurrentEntityStorage& Storage) const
-		{
-			return &Storage;
-		}
-	};
-
-	UE::Mass::IEntityStorageInterface* Interface = Visit(StorageSelector{}, EntityStorage);
-
-	return *Interface;
-}
-
-const UE::Mass::IEntityStorageInterface& FMassEntityManager::GetEntityStorageInterface() const
-{
-	using namespace UE::Mass;
-	struct StorageSelector
-	{
-		const UE::Mass::IEntityStorageInterface* operator()(const FEmptyVariantState&) const
-		{
-			checkf(false, TEXT("Attempt to use EntityStorageInterface without initialization"));
-			return nullptr;
-		}
-		const UE::Mass::IEntityStorageInterface* operator()(const FSingleThreadedEntityStorage& Storage) const
-		{
-			return &Storage;
-		}
-		const UE::Mass::IEntityStorageInterface* operator()(const FConcurrentEntityStorage& Storage) const
-		{
-			return &Storage;
-		}
-	};
-
-	const UE::Mass::IEntityStorageInterface* Interface = Visit(StorageSelector{}, EntityStorage);
-
-	return *Interface;
-}
-#else
-UE::Mass::FSingleThreadedEntityStorage& FMassEntityManager::GetEntityStorageInterface()
-{
-	// Get will assert if not initialized
-	return EntityStorage.Get<UE::Mass::FSingleThreadedEntityStorage>();
-}
-
-const UE::Mass::FSingleThreadedEntityStorage& FMassEntityManager::GetEntityStorageInterface() const
-{
-	// Get will assert if not initialized
-	return EntityStorage.Get<UE::Mass::FSingleThreadedEntityStorage>();
-}
-#endif
-
-#if WITH_MASSENTITY_DEBUG
-UE::Mass::IEntityStorageInterface& FMassEntityManager::DebugGetEntityStorageInterface()
-{
-	return GetEntityStorageInterface();
-}
-
-const UE::Mass::IEntityStorageInterface& FMassEntityManager::DebugGetEntityStorageInterface() const
-{
-	return GetEntityStorageInterface();
-}
-#endif
-
 //-----------------------------------------------------------------------------
 // FMassEntityManager
 //-----------------------------------------------------------------------------
@@ -294,7 +215,7 @@ namespace UE::Mass::Private
 		}
 		void operator()(const FMassEntityManager_InitParams_Concurrent& Params)
 		{
-#if MASS_CONCURRENT_RESERVE
+#if WITH_MASS_CONCURRENT_RESERVE
 			EntityStorage->Emplace<UE::Mass::FConcurrentEntityStorage>();
 			EntityStorage->Get<UE::Mass::FConcurrentEntityStorage>().Initialize(Params);
 #else
@@ -2077,8 +1998,16 @@ FMassRequirementAccessDetector& FMassEntityManager::GetRequirementAccessDetector
 	return RequirementAccessDetector;
 }
 
-#endif // WITH_MASSENTITY_DEBUG
+UE::Mass::FStorageType& FMassEntityManager::DebugGetEntityStorageInterface()
+{
+	return GetEntityStorageInterface();
+}
 
+const UE::Mass::FStorageType& FMassEntityManager::DebugGetEntityStorageInterface() const
+{
+	return GetEntityStorageInterface();
+}
+#endif // WITH_MASSENTITY_DEBUG
 
 //-----------------------------------------------------------------------------
 // DEPRECATED
diff --git a/Engine/Source/Runtime/MassEntity/Private/MassEntitySubsystem.cpp b/Engine/Source/Runtime/MassEntity/Private/MassEntitySubsystem.cpp
--- a/Engine/Source/Runtime/MassEntity/Private/MassEntitySubsystem.cpp
+++ b/Engine/Source/Runtime/MassEntity/Private/MassEntitySubsystem.cpp
@@ -56,7 +56,7 @@ void UMassEntitySubsystem::Initialize(FSubsystemCollectionBase& Collection)
 	Super::Initialize(Collection);
 
 	FMassEntityManagerStorageInitParams InitializationParams;
-	
+#if WITH_MASS_CONCURRENT_RESERVE
 	if (UE::Mass::Private::bEnableMassConcurrentReserveRuntime)
 	{
 		InitializationParams.Emplace<FMassEntityManager_InitParams_Concurrent>(
@@ -67,6 +67,7 @@ void UMassEntitySubsystem::Initialize(FSubsystemCollectionBase& Collection)
 			});
 	}
 	else
+#endif // WITH_MASS_CONCURRENT_RESERVE
 	{
 		InitializationParams.Emplace<FMassEntityManager_InitParams_SingleThreaded>();
 	}
diff --git a/Engine/Source/Runtime/MassEntity/Public/MassEntityManager.h b/Engine/Source/Runtime/MassEntity/Public/MassEntityManager.h
--- a/Engine/Source/Runtime/MassEntity/Public/MassEntityManager.h
+++ b/Engine/Source/Runtime/MassEntity/Public/MassEntityManager.h
@@ -32,7 +32,23 @@ namespace UE::Mass::Private
 	struct FEntityStorageInitializer;
 }
 
-#define MASS_CONCURRENT_RESERVE 1
+// use REQUESTED_MASS_CONCURRENT_RESERVE=0 in your project's Build.cs file to disable concurrent storage
+// NOTE that it will always be enabled in WITH_EDITOR since editor code requires it. 
+// @see WITH_MASS_CONCURRENT_RESERVE
+#ifndef REQUESTED_MASS_CONCURRENT_RESERVE
+#define REQUESTED_MASS_CONCURRENT_RESERVE 1
+#endif
+
+#define WITH_MASS_CONCURRENT_RESERVE (REQUESTED_MASS_CONCURRENT_RESERVE || WITH_EDITOR)
+
+namespace UE::Mass
+{
+#if WITH_MASS_CONCURRENT_RESERVE
+	using FStorageType = IEntityStorageInterface;
+#else
+	using FStorageType = FSingleThreadedEntityStorage;
+#endif //WITH_MASS_CONCURRENT_RESERVE
+}
 
 /** 
  * The type responsible for hosting Entities managing Archetypes.
@@ -975,9 +991,9 @@ PRAGMA_ENABLE_DEPRECATION_WARNINGS
 	FMassRequirementAccessDetector& GetRequirementAccessDetector();
 
 	// For use by the friend MassDebugger
-	UE::Mass::IEntityStorageInterface& DebugGetEntityStorageInterface();
+	UE::Mass::FStorageType& DebugGetEntityStorageInterface();
 	// For use by the friend MassDebugger
-	const UE::Mass::IEntityStorageInterface& DebugGetEntityStorageInterface() const;
+	const UE::Mass::FStorageType& DebugGetEntityStorageInterface() const;
 #endif // WITH_MASSENTITY_DEBUG
 
 protected:
@@ -1024,14 +1040,35 @@ private:
 
 	TSharedRef<FEntityCreationContext> InternalBatchCreateReservedEntities(const FMassArchetypeHandle& ArchetypeHandle,
 		const FMassArchetypeSharedFragmentValues& SharedFragmentValues, TConstArrayView<FMassEntityHandle> ReservedEntities);
-	
-#if MASS_CONCURRENT_RESERVE
-	UE::Mass::IEntityStorageInterface& GetEntityStorageInterface();
-	const UE::Mass::IEntityStorageInterface& GetEntityStorageInterface() const;
-#else
-	UE::Mass::FSingleThreadedEntityStorage& GetEntityStorageInterface();
-	const UE::Mass::FSingleThreadedEntityStorage& GetEntityStorageInterface() const;
+
+	UE::Mass::FStorageType& GetEntityStorageInterface() const
+	{
+#if WITH_MASS_CONCURRENT_RESERVE
+		using namespace UE::Mass;
+		struct StorageSelector
+		{
+			UE::Mass::IEntityStorageInterface* operator()(FEmptyVariantState&) const
+			{
+				checkf(false, TEXT("Attempt to use EntityStorageInterface without initialization"));
+				return nullptr;
+			}
+			UE::Mass::IEntityStorageInterface* operator()(FSingleThreadedEntityStorage& Storage) const
+			{
+				return &Storage;
+			}
+			UE::Mass::IEntityStorageInterface* operator()(FConcurrentEntityStorage& Storage) const
+			{
+				return &Storage;
+			}
+		};
+
+		UE::Mass::IEntityStorageInterface* Interface = Visit(StorageSelector{}, EntityStorage);
+
+		return *Interface;
+#else	
+		return EntityStorage.Get<UE::Mass::FSingleThreadedEntityStorage>();
 #endif
+	}
 
 	/**
 	 * If ActiveCreationContext is not valid the function creates a new shared FEntityCreationContext instance and returns that.
@@ -1085,7 +1122,7 @@ private:
 		FEmptyVariantState,
 		UE::Mass::FSingleThreadedEntityStorage,
 		UE::Mass::FConcurrentEntityStorage>;
-	FEntityStorageContainerType EntityStorage;
+	mutable FEntityStorageContainerType EntityStorage;
 
 	/** Never access directly, use GetOrMakeCreationContext instead. */
 	TWeakPtr<FEntityCreationContext> ActiveCreationContext;
