[Mass] Fixed a bug in FMassArchetypeData::InitializeWithSimilar that could result in archetype's FragmentConfigs not getting set up properly

The bug has not harmed us for far becase FMassArchetypeCompositionDescriptor's move semantics work as expected. The bug surface by some local changes that broken FMassArchetypeCompositionDescriptor's moveability and started crashing on misconfigured FragmentConfigs

#rnx

diff --git a/Engine/Source/Runtime/MassEntity/Private/MassArchetypeData.cpp b/Engine/Source/Runtime/MassEntity/Private/MassArchetypeData.cpp
--- a/Engine/Source/Runtime/MassEntity/Private/MassArchetypeData.cpp
+++ b/Engine/Source/Runtime/MassEntity/Private/MassArchetypeData.cpp
@@ -99,7 +99,7 @@ void FMassArchetypeData::InitializeWithSimilar(const FMassArchetypeData& BaseArc
 
 	// note that we're calling this function rarely, so we can be a little bit inefficient here.
 	CompositionDescriptor = MoveTemp(NewComposition);
-	if (CompositionDescriptor.Fragments != NewComposition.Fragments)
+	if (CompositionDescriptor.Fragments != BaseArchetype.GetCompositionDescriptor().Fragments)
 	{
 		ConfigureFragments();
 	}
