[Mass] Fixed a bug in how FMassArchetypeEntityCollection::FArchetypeEntityRange.Length == 0 is treated - the comments (and the design) say it means "all the entities in the chunk starting at index SubchunkStart", while only FMassArchetypeData::ExecutionFunctionForChunk was actually respecting it.

The change also adds unit test for the modified code.

diff --git a/Engine/Source/Developer/MassEntityTestSuite/Private/MassChunkCollectionTest.cpp b/Engine/Source/Developer/MassEntityTestSuite/Private/MassChunkCollectionTest.cpp
--- a/Engine/Source/Developer/MassEntityTestSuite/Private/MassChunkCollectionTest.cpp
+++ b/Engine/Source/Developer/MassEntityTestSuite/Private/MassChunkCollectionTest.cpp
@@ -1,8 +1,11 @@
 // Copyright Epic Games, Inc. All Rights Reserved.
 
+#include "MassArchetypeData.h"
 #include "MassEntityManager.h"
 #include "MassEntityTestTypes.h"
+#include "MassExecutionContext.h"
 #include "Algo/RandomShuffle.h"
+#include "MassExecutor.h"
 
 #define LOCTEXT_NAMESPACE "MassTest"
 
@@ -384,6 +387,134 @@ struct FEntityCollection_WithPayload : FEntityCollection_WithPayloadBase
 	}
 };
 IMPLEMENT_AI_INSTANT_TEST(FEntityCollection_WithPayload, "System.Mass.EntityCollection.Create.WithPayload");
+
+struct FEntityCollection_ZeroLength_Destroy : FEntityTestBase
+{
+	virtual bool InstantTest() override
+	{
+		CA_ASSUME(EntityManager);
+		const FMassArchetypeHandle SourceArchetype = FloatsArchetype;
+		const int32 EntitiesPerChunk = FMassArchetypeHelper::ArchetypeDataFromHandleChecked(SourceArchetype).GetNumEntitiesPerChunk();
+		const int32 NumChunks = 3;
+		const int32 EntityCount = NumChunks * EntitiesPerChunk;
+		TArray<FMassEntityHandle> Entities;
+		EntityManager->BatchCreateEntities(SourceArchetype, {}, EntityCount, Entities);
+
+		AITEST_EQUAL("Entity creation generated expected number of entities", EntityCount, FMassArchetypeHelper::ArchetypeDataFromHandleChecked(SourceArchetype).GetNumEntities());
+
+		FMassEntityQuery Query;
+		Query.CacheArchetypes(*EntityManager);
+		FMassArchetypeEntityCollection Collection(SourceArchetype);
+		AITEST_EQUAL("Created collection has expected number of ranges (matching number of chunks created"
+			, Collection.GetRanges().Num(), NumChunks);
+		
+		EntityManager->BatchDestroyEntityChunks(Collection);
+		AITEST_EQUAL("Entity destruction destroyed all entities", FMassArchetypeHelper::ArchetypeDataFromHandleChecked(SourceArchetype).GetNumEntities(), 0);
+
+		return true;
+	}
+};
+IMPLEMENT_AI_INSTANT_TEST(FEntityCollection_ZeroLength_Destroy, "System.Mass.Entity.ZeroLengthRanges.Destroy");
+
+struct FEntityCollection_ZeroLength_Execute : FEntityTestBase
+{
+	virtual bool InstantTest() override
+	{
+		CA_ASSUME(EntityManager);
+		const FMassArchetypeHandle SourceArchetype = FloatsArchetype;
+		const int32 EntitiesPerChunk = FMassArchetypeHelper::ArchetypeDataFromHandleChecked(SourceArchetype).GetNumEntitiesPerChunk();
+		const int32 NumChunks = 3;
+		const int32 EntityCount = NumChunks * EntitiesPerChunk;
+		TArray<FMassEntityHandle> Entities;
+		EntityManager->BatchCreateEntities(SourceArchetype, {}, EntityCount, Entities);
+		FMassArchetypeEntityCollection Collection(SourceArchetype);
+
+		AITEST_EQUAL("Entity creation generated expected number of entities", EntityCount, FMassArchetypeHelper::ArchetypeDataFromHandleChecked(SourceArchetype).GetNumEntities());
+
+		int32 ProcessedEntitiesCount = 0;
+		UMassTestProcessorBase* Processor = NewObject<UMassTestProcessorBase>();
+		Processor->ForEachEntityChunkExecutionFunction = [&ProcessedEntitiesCount](FMassExecutionContext& Context)
+		{
+			ProcessedEntitiesCount += Context.GetNumEntities();
+		};
+		Processor->EntityQuery.AddRequirement<FTestFragment_Float>(EMassFragmentAccess::ReadOnly);
+
+		FMassProcessingContext ProcessingContext(EntityManager, /*DeltaSeconds=*/0.f);
+		UE::Mass::Executor::RunProcessorsView(MakeArrayView((UMassProcessor**)&Processor, 1), ProcessingContext, MakeConstArrayView(&Collection, 1));
+
+		AITEST_EQUAL("Number of entities processed matches expectations"
+			, ProcessedEntitiesCount, EntityCount);
+
+		return true;
+	}
+};
+IMPLEMENT_AI_INSTANT_TEST(FEntityCollection_ZeroLength_Execute, "System.Mass.Entity.ZeroLengthRanges.Execute");
+
+struct FEntityCollection_ZeroLength_BatchMove : FEntityTestBase
+{
+	virtual bool InstantTest() override
+	{
+		CA_ASSUME(EntityManager);
+		const FMassArchetypeHandle SourceArchetype = FloatsArchetype;
+		const int32 EntitiesPerChunk = FMassArchetypeHelper::ArchetypeDataFromHandleChecked(SourceArchetype).GetNumEntitiesPerChunk();
+		const int32 NumChunks = 3;
+		const int32 EntityCount = NumChunks * EntitiesPerChunk;
+		TArray<FMassEntityHandle> Entities;
+		EntityManager->BatchCreateEntities(SourceArchetype, {}, EntityCount, Entities);
+		FMassArchetypeEntityCollection Collection(SourceArchetype);
+
+		const UScriptStruct* AddedTag = FTestTag_A::StaticStruct();
+
+		AITEST_EQUAL("Entity creation generated expected number of entities", EntityCount, FMassArchetypeHelper::ArchetypeDataFromHandleChecked(SourceArchetype).GetNumEntities());
+		check(EntityCount == Entities.Num());
+
+		EntityManager->BatchChangeTagsForEntities(MakeArrayView(&Collection, 1), FMassTagBitSet(*AddedTag), {});
+
+		const FMassArchetypeHandle NewArchetype = EntityManager->GetArchetypeForEntity(Entities[0]);
+		AITEST_NOT_EQUAL("Entities have changed their host archetype", NewArchetype, SourceArchetype);
+
+		const FMassArchetypeHandle ExpectedArchetype = EntityManager->CreateArchetype(SourceArchetype, MakeArrayView(&AddedTag, 1));
+		AITEST_EQUAL("The new archetype matches expectations", NewArchetype, ExpectedArchetype);
+		AITEST_EQUAL("All the entities have been moved to the expected archetype", EntityManager->DebugGetArchetypeEntitiesCount(NewArchetype), EntityCount);
+
+		return true;
+	}
+};
+IMPLEMENT_AI_INSTANT_TEST(FEntityCollection_ZeroLength_BatchMove, "System.Mass.Entity.ZeroLengthRanges.BatchMove");
+
+struct FEntityCollection_ZeroLength_BatchSetFragmentValues : FEntityTestBase
+{
+	virtual bool InstantTest() override
+	{
+		CA_ASSUME(EntityManager);
+		const FMassArchetypeHandle SourceArchetype = FloatsArchetype;
+		const int32 EntitiesPerChunk = FMassArchetypeHelper::ArchetypeDataFromHandleChecked(SourceArchetype).GetNumEntitiesPerChunk();
+		const int32 NumChunks = 3;
+		const int32 EntityCount = NumChunks * EntitiesPerChunk;
+		TArray<FMassEntityHandle> Entities;
+		EntityManager->BatchCreateEntities(SourceArchetype, {}, EntityCount, Entities);
+
+		AITEST_EQUAL("Entity creation generated expected number of entities", EntityCount, FMassArchetypeHelper::ArchetypeDataFromHandleChecked(SourceArchetype).GetNumEntities());
+		check(EntityCount == Entities.Num());
+
+		FMassArchetypeEntityCollection Collection(SourceArchetype);
+		FMassArchetypeEntityCollectionWithPayload EntityCollectionWithEmptyPayload(Collection);
+
+		const FMassFragmentBitSet AddedFragmentsBitSet(*FTestFragment_Int::StaticStruct());
+		EntityManager->BatchAddFragmentInstancesForEntities(MakeArrayView(&EntityCollectionWithEmptyPayload, 1), AddedFragmentsBitSet);
+
+		const FMassArchetypeHandle NewArchetype = EntityManager->GetArchetypeForEntity(Entities[0]);
+		AITEST_NOT_EQUAL("Entities have changed their host archetype", NewArchetype, SourceArchetype);
+
+		const FMassArchetypeHandle ExpectedArchetype = FloatsIntsArchetype;
+		AITEST_EQUAL("The new archetype matches expectations", NewArchetype, ExpectedArchetype);
+		AITEST_EQUAL("All the entities have been moved to the expected archetype", EntityManager->DebugGetArchetypeEntitiesCount(NewArchetype), EntityCount);
+
+		return true;
+	}
+};
+IMPLEMENT_AI_INSTANT_TEST(FEntityCollection_ZeroLength_BatchSetFragmentValues, "System.Mass.Entity.ZeroLengthRanges.BatchSetFragmentValues");
+
 #endif // WITH_MASSENTITY_DEBUG
 //
 }
diff --git a/Engine/Source/Runtime/MassEntity/Private/MassArchetypeData.cpp b/Engine/Source/Runtime/MassEntity/Private/MassArchetypeData.cpp
--- a/Engine/Source/Runtime/MassEntity/Private/MassArchetypeData.cpp
+++ b/Engine/Source/Runtime/MassEntity/Private/MassArchetypeData.cpp
@@ -323,19 +323,20 @@ void FMassArchetypeData::BatchDestroyEntityChunks(FMassArchetypeEntityCollection
 	for (const FMassArchetypeEntityCollection::FArchetypeEntityRange EntityRange : SortedRangeCollection)
 	{ 
 		FMassArchetypeChunk& Chunk = Chunks[EntityRange.ChunkIndex];
+		const int32 RangeLength = EntityRange.Length > 0 ? EntityRange.Length : (Chunk.GetNumInstances() - EntityRange.SubchunkStart);
 
 		// gather entities we're about to remove
 		FMassEntityHandle* DyingEntityPtr = &Chunk.GetEntityArrayElementRef(EntityListOffsetWithinChunk, EntityRange.SubchunkStart);
-		OutEntitiesRemoved.Append(DyingEntityPtr, EntityRange.Length);
+		OutEntitiesRemoved.Append(DyingEntityPtr, RangeLength);
 
 		for (const FMassArchetypeFragmentConfig& FragmentConfig : FragmentConfigs)
 		{
 			// Destroy the fragment data
 			void* DyingFragmentPtr = FragmentConfig.GetFragmentData(Chunk.GetRawMemory(), EntityRange.SubchunkStart);
-			FragmentConfig.FragmentType->DestroyStruct(DyingFragmentPtr, EntityRange.Length);
+			FragmentConfig.FragmentType->DestroyStruct(DyingFragmentPtr, RangeLength);
 		}
 
-		BatchRemoveEntitiesInternal(EntityRange.ChunkIndex, EntityRange.SubchunkStart, EntityRange.Length);
+		BatchRemoveEntitiesInternal(EntityRange.ChunkIndex, EntityRange.SubchunkStart, RangeLength);
 	}
 
 	for (int i = InitialOutEntitiesCount; i < OutEntitiesRemoved.Num(); ++i)
@@ -478,8 +479,8 @@ void FMassArchetypeData::ExecuteFunction(FMassExecutionContext& RunContext, cons
 	{
 		FMassArchetypeChunk& Chunk = Chunks[ChunkIterator->ChunkIndex];
 
-		const int32 ChunkLength = ChunkIterator->Length > 0 ? ChunkIterator->Length : (Chunk.GetNumInstances() - ChunkIterator->SubchunkStart);
-		if (ChunkLength)
+		const int32 SubchunkLength = ChunkIterator->Length > 0 ? ChunkIterator->Length : (Chunk.GetNumInstances() - ChunkIterator->SubchunkStart);
+		if (SubchunkLength)
 		{
 			const uint32 SharedFragmentValuesHash = GetTypeHash(Chunk.GetSharedFragmentValues());
 			if (PrevSharedFragmentValuesHash != SharedFragmentValuesHash)
@@ -489,14 +490,14 @@ void FMassArchetypeData::ExecuteFunction(FMassExecutionContext& RunContext, cons
 				BindSharedFragmentRequirements(RunContext, Chunk.GetMutableSharedFragmentValues(), RequirementMapping.SharedFragments);
 			}
 
-			checkf((ChunkIterator->SubchunkStart + ChunkLength) <= Chunk.GetNumInstances() && ChunkLength > 0, TEXT("Invalid subchunk, it is going over the number of instances in the chunk or it is empty."));
+			checkf((ChunkIterator->SubchunkStart + SubchunkLength) <= Chunk.GetNumInstances() && SubchunkLength > 0, TEXT("Invalid subchunk, it is going over the number of instances in the chunk or it is empty."));
 
 			RunContext.SetCurrentChunkSerialModificationNumber(Chunk.GetSerialModificationNumber());
 			BindChunkFragmentRequirements(RunContext, RequirementMapping.ChunkFragments, Chunk);
 
 			if (!ChunkCondition || ChunkCondition(RunContext))
 			{
-				BindEntityRequirements(RunContext, RequirementMapping.EntityFragments, Chunk, ChunkIterator->SubchunkStart, ChunkLength);
+				BindEntityRequirements(RunContext, RequirementMapping.EntityFragments, Chunk, ChunkIterator->SubchunkStart, SubchunkLength);
 				Function(RunContext);
 			}
 		}
@@ -540,9 +541,9 @@ void FMassArchetypeData::ExecuteFunction(FMassExecutionContext& RunContext, cons
 void FMassArchetypeData::ExecutionFunctionForChunk(FMassExecutionContext& RunContext, const FMassExecuteFunction& Function, const FMassQueryRequirementIndicesMapping& RequirementMapping, const FMassArchetypeEntityCollection::FArchetypeEntityRange& EntityRange, const FMassChunkConditionFunction& ChunkCondition)
 {
 	FMassArchetypeChunk& Chunk = Chunks[EntityRange.ChunkIndex];
-	const int32 ChunkLength = EntityRange.Length > 0 ? EntityRange.Length : (Chunk.GetNumInstances() - EntityRange.SubchunkStart);
+	const int32 RangeLength = EntityRange.Length > 0 ? EntityRange.Length : (Chunk.GetNumInstances() - EntityRange.SubchunkStart);
 
-	if (ChunkLength)
+	if (RangeLength > 0)
 	{
 		BindConstSharedFragmentRequirements(RunContext, Chunk.GetSharedFragmentValues(), RequirementMapping.ConstSharedFragments);
 		BindSharedFragmentRequirements(RunContext, Chunk.GetMutableSharedFragmentValues(), RequirementMapping.SharedFragments);
@@ -553,7 +554,7 @@ void FMassArchetypeData::ExecutionFunctionForChunk(FMassExecutionContext& RunCon
 
 		if (!ChunkCondition || ChunkCondition(RunContext))
 		{
-			BindEntityRequirements(RunContext, RequirementMapping.EntityFragments, Chunk, EntityRange.SubchunkStart, ChunkLength);
+			BindEntityRequirements(RunContext, RequirementMapping.EntityFragments, Chunk, EntityRange.SubchunkStart, RangeLength);
 			Function(RunContext);
 		}
 	}
@@ -1086,12 +1087,13 @@ void FMassArchetypeData::BatchMoveEntitiesToAnotherArchetype(const FMassArchetyp
 
 	for (const FMassArchetypeEntityCollection::FArchetypeEntityRange EntityRange : Subchunks)
 	{
-		if (!ensureMsgf(EntityRange.IsSet() && EntityRange.Length > 0, TEXT("We only expect to get valid EntityRanges at this point.")))
+		if (!ensureMsgf(EntityRange.IsSet() && EntityRange.Length >= 0, TEXT("We only expect to get valid EntityRanges at this point.")))
 		{
 			continue;
 		}
 
 		FMassArchetypeChunk& Chunk = Chunks[EntityRange.ChunkIndex];
+		const int32 RangeLength = EntityRange.Length > 0 ? EntityRange.Length : (Chunk.GetNumInstances() - EntityRange.SubchunkStart);
 
 		// 0 - consider compacting new archetype to ensure larger empty spaces
 		// 1. find next free spot in the destination archetype
@@ -1099,7 +1101,7 @@ void FMassArchetypeData::BatchMoveEntitiesToAnotherArchetype(const FMassArchetyp
 
 		// gather entities we're about to remove
 		FMassEntityHandle* DyingEntityPtr = &Chunk.GetEntityArrayElementRef(EntityListOffsetWithinChunk, EntityRange.SubchunkStart);
-		OutEntitiesBeingMoved.Append(DyingEntityPtr, EntityRange.Length);
+		OutEntitiesBeingMoved.Append(DyingEntityPtr, RangeLength);
 
 		FMassArchetypeEntityCollection::FArchetypeEntityRange ResultSubChunk;
 		ResultSubChunk.ChunkIndex = 0;
@@ -1113,7 +1115,7 @@ void FMassArchetypeData::BatchMoveEntitiesToAnotherArchetype(const FMassArchetyp
 
 			if (ChangeSharedFragments == false)
 			{
-				ResultSubChunk = NewArchetype.PrepareNextEntitiesSpanInternal(MakeArrayView(DyingEntityPtr + NumberMoved, EntityRange.Length - NumberMoved)
+				ResultSubChunk = NewArchetype.PrepareNextEntitiesSpanInternal(MakeArrayView(DyingEntityPtr + NumberMoved, RangeLength - NumberMoved)
 					, Chunk.GetSharedFragmentValues(), ResultSubChunk.ChunkIndex);
 			}
 			else
@@ -1134,7 +1136,7 @@ void FMassArchetypeData::BatchMoveEntitiesToAnotherArchetype(const FMassArchetyp
 				}
 				NewSharedValues.Sort();
 
-				ResultSubChunk = NewArchetype.PrepareNextEntitiesSpanInternal(MakeArrayView(DyingEntityPtr + NumberMoved, EntityRange.Length - NumberMoved)
+				ResultSubChunk = NewArchetype.PrepareNextEntitiesSpanInternal(MakeArrayView(DyingEntityPtr + NumberMoved, RangeLength - NumberMoved)
 					, NewSharedValues, ResultSubChunk.ChunkIndex);
 			}
 
@@ -1156,7 +1158,7 @@ void FMassArchetypeData::BatchMoveEntitiesToAnotherArchetype(const FMassArchetyp
 				}
 			}
 
-		} while (NumberMoved < EntityRange.Length);
+		} while (NumberMoved < RangeLength);
 
 	}
 
@@ -1336,11 +1338,12 @@ void FMassArchetypeData::BatchSetFragmentValues(TConstArrayView<FMassArchetypeEn
 	for (const FMassArchetypeEntityCollection::FArchetypeEntityRange EntityRange : EntityCollection)
 	{
 		FMassArchetypeChunk& Chunk = Chunks[EntityRange.ChunkIndex];
+		const int32 RangeLength = EntityRange.Length > 0 ? EntityRange.Length : (Chunk.GetNumInstances() - EntityRange.SubchunkStart);
 
 		for (int i = 0; i < Payload.Num(); ++i)
 		{
 			FStructArrayView FragmentPayload = Payload[i];
-			check(FragmentPayload.Num() - EntitiesHandled >= EntityRange.Length);
+			check(FragmentPayload.Num() - EntitiesHandled >= RangeLength);
 
 			const UScriptStruct* FragmentType = FragmentPayload.GetScriptStruct();
 			check(FragmentType);
@@ -1349,10 +1352,10 @@ void FMassArchetypeData::BatchSetFragmentValues(TConstArrayView<FMassArchetypeEn
 			void* Dst = FragmentConfigs[FragmentIndex].GetFragmentData(Chunk.GetRawMemory(), EntityRange.SubchunkStart);
 			const void* Src = FragmentPayload.GetDataAt(EntitiesHandled);
 
-			FragmentType->CopyScriptStruct(Dst, Src, EntityRange.Length);
+			FragmentType->CopyScriptStruct(Dst, Src, RangeLength);
 		}
 
-		EntitiesHandled += EntityRange.Length;
+		EntitiesHandled += RangeLength;
 	}
 }
 
diff --git a/Engine/Source/Runtime/MassEntity/Public/MassArchetypeTypes.h b/Engine/Source/Runtime/MassEntity/Public/MassArchetypeTypes.h
--- a/Engine/Source/Runtime/MassEntity/Public/MassArchetypeTypes.h
+++ b/Engine/Source/Runtime/MassEntity/Public/MassArchetypeTypes.h
@@ -58,8 +58,15 @@ public:
 	struct FArchetypeEntityRange
 	{
 		int32 ChunkIndex = INDEX_NONE;
+		 /** 
+		  * The index of the first entity within the specified chunk that starts this subchunk.
+		  */
 		int32 SubchunkStart = 0;
-		/** negative or 0-length means "all available entities within chunk" */
+		/** 
+		 * The number of entities in this subchunk.
+		 * If Length is 0 or negative, it indicates that the range covers all remaining entities 
+		 * in the chunk starting from SubchunkStart.
+		 */
 		int32 Length = 0;
 
 		FArchetypeEntityRange() = default;
