[Mass] Removed Mass code marked as deprecated for 5.3 and down.

#rnx

diff --git a/Engine/Plugins/Runtime/MassGameplay/Source/MassLOD/Private/MassLODSubsystem.cpp b/Engine/Plugins/Runtime/MassGameplay/Source/MassLOD/Private/MassLODSubsystem.cpp
--- a/Engine/Plugins/Runtime/MassGameplay/Source/MassLOD/Private/MassLODSubsystem.cpp
+++ b/Engine/Plugins/Runtime/MassGameplay/Source/MassLOD/Private/MassLODSubsystem.cpp
@@ -407,18 +407,6 @@ void UMassLODSubsystem::SynchronizeViewers()
 	}
 }
 
-void UMassLODSubsystem::AddViewer(APlayerController* PlayerController, FName StreamingSourceName /* = NAME_None*/)
-{
-	if (PlayerController)
-	{
-		AddPlayerViewer(*PlayerController);
-	}
-	else
-	{
-		AddStreamingSourceViewer(StreamingSourceName);
-	}
-}
-
 void UMassLODSubsystem::AddPlayerViewer(APlayerController& PlayerController)
 {
 #if WITH_EDITOR
diff --git a/Engine/Plugins/Runtime/MassGameplay/Source/MassLOD/Public/MassLODSubsystem.h b/Engine/Plugins/Runtime/MassGameplay/Source/MassLOD/Public/MassLODSubsystem.h
--- a/Engine/Plugins/Runtime/MassGameplay/Source/MassLOD/Public/MassLODSubsystem.h
+++ b/Engine/Plugins/Runtime/MassGameplay/Source/MassLOD/Public/MassLODSubsystem.h
@@ -64,11 +64,6 @@ private:
 	TObjectPtr<APlayerController> PlayerController = nullptr;
 };
 
-UE_DEPRECATED(5.3, "FOnViewerAdded is deprecated. Use UMassLODSubsystem::FOnViewerAdded instead.")
-DECLARE_MULTICAST_DELEGATE_ThreeParams(FOnViewerAdded, FMassViewerHandle ViewerHandle, APlayerController* PlayerController, FName StreamingSourceName);
-UE_DEPRECATED(5.3, "FOnViewerRemoved is deprecated. Use UMassLODSubsystem::FOnViewerRemoved instead.")
-DECLARE_MULTICAST_DELEGATE_ThreeParams(FOnViewerRemoved, FMassViewerHandle ViewerHandle, APlayerController* PlayerController, FName StreamingSourceName);
-
 /*
  * Manager responsible to manage and synchronized available viewers
  */
@@ -131,9 +126,6 @@ protected:
 	/** Synchronizes the viewers from the engine PlayerController list */
 	void SynchronizeViewers();
 
-	UE_DEPRECATED(5.3, "AddViewer has been deprecated. Use AddPlayerViewer or AddStreamingSourceViewer instead")
-	void AddViewer(APlayerController* PlayerController, FName StreamingSourceName = NAME_None);
-
 	/** Adds the given player as a viewer to the list and sends notification about addition */
 	void AddPlayerViewer(APlayerController& PlayerController);
 
diff --git a/Engine/Plugins/Runtime/MassGameplay/Source/MassSpawner/Private/MassEntityConfigAsset.cpp b/Engine/Plugins/Runtime/MassGameplay/Source/MassSpawner/Private/MassEntityConfigAsset.cpp
--- a/Engine/Plugins/Runtime/MassGameplay/Source/MassSpawner/Private/MassEntityConfigAsset.cpp
+++ b/Engine/Plugins/Runtime/MassGameplay/Source/MassSpawner/Private/MassEntityConfigAsset.cpp
@@ -6,7 +6,6 @@
 #include "MassSpawnerTypes.h"
 #include "MassSpawnerSubsystem.h"
 #include "MassEntityTemplateRegistry.h"
-#include "Serialization/ArchiveObjectCrc32.h"
 #include "VisualLogger/VisualLogger.h"
 #include "Engine/World.h"
 #if WITH_EDITOR
@@ -186,34 +185,6 @@ bool FMassEntityConfig::ValidateEntityTemplate(const UWorld& World)
 	return BuildContext.BuildFromTraits(CombinedTraits, World);
 }
 
-//-----------------------------------------------------------------------------
-// DEPRECATED
-//-----------------------------------------------------------------------------
-const FMassEntityTemplate& FMassEntityConfig::GetOrCreateEntityTemplate(const UWorld& World, const UObject& InConfigOwner) const
-{
-	return GetOrCreateEntityTemplate(World);
-}
-
-void FMassEntityConfig::DestroyEntityTemplate(const UWorld& World, const UObject& InConfigOwner) const
-{
-	DestroyEntityTemplate(World);
-}
-
-const FMassEntityTemplate& FMassEntityConfig::GetEntityTemplateChecked(const UWorld& World, const UObject& InConfigOwner) const
-{
-	return GetEntityTemplateChecked(World);
-}
-
-bool FMassEntityConfig::ValidateEntityTemplate(const UWorld& World, const UObject& InConfigOwner)
-{
-	return ValidateEntityTemplate(World);
-}
-
-void FMassEntityConfig::GetCombinedTraits(TArray<UMassEntityTraitBase*>& OutTraits, TArray<const UObject*>& Visited, const UObject& InConfigOwner) const
-{
-	return GetCombinedTraitsInternal(OutTraits, Visited);
-}
-
 #if WITH_EDITOR
 void FMassEntityConfig::PostDuplicate(const bool bDuplicateForPIE)
 {
diff --git a/Engine/Plugins/Runtime/MassGameplay/Source/MassSpawner/Private/MassEntityTemplateRegistry.cpp b/Engine/Plugins/Runtime/MassGameplay/Source/MassSpawner/Private/MassEntityTemplateRegistry.cpp
--- a/Engine/Plugins/Runtime/MassGameplay/Source/MassSpawner/Private/MassEntityTemplateRegistry.cpp
+++ b/Engine/Plugins/Runtime/MassGameplay/Source/MassSpawner/Private/MassEntityTemplateRegistry.cpp
@@ -273,21 +273,4 @@ bool FMassEntityTemplateBuildContext::ValidateBuildContext(const UWorld& World)
 	return (ErrorCount == 0);
 }
 
-//-----------------------------------------------------------------------------
-// DEPRECATED
-//-----------------------------------------------------------------------------
-FMassEntityTemplate* FMassEntityTemplateRegistry::FindMutableTemplateFromTemplateID(FMassEntityTemplateID TemplateID)
-{
-	return nullptr;
-}
-
-FMassEntityTemplate& FMassEntityTemplateRegistry::CreateTemplate(const uint32 HashLookup, FMassEntityTemplateID TemplateID)
-{
-	static FMassEntityTemplate Dummy;
-	return Dummy;
-}
-
-void FMassEntityTemplateRegistry::InitializeEntityTemplate(FMassEntityTemplate& InOutTemplate) const
-{}
-
 #undef LOCTEXT_NAMESPACE 
diff --git a/Engine/Plugins/Runtime/MassGameplay/Source/MassSpawner/Public/MassEntityConfigAsset.h b/Engine/Plugins/Runtime/MassGameplay/Source/MassSpawner/Public/MassEntityConfigAsset.h
--- a/Engine/Plugins/Runtime/MassGameplay/Source/MassSpawner/Public/MassEntityConfigAsset.h
+++ b/Engine/Plugins/Runtime/MassGameplay/Source/MassSpawner/Public/MassEntityConfigAsset.h
@@ -60,18 +60,6 @@ struct MASSSPAWNER_API FMassEntityConfig
 
 	void SetOwner(UObject& InOwner) { ConfigOwner = &InOwner; }
 
-	UE_DEPRECATED(5.3, "This flavor of GetOrCreateEntityTemplate is deprecated. Use the one without the ConfigOwner parameter (now a property of the FMassEntityConfig itself)")
-	const FMassEntityTemplate& GetOrCreateEntityTemplate(const UWorld& World, const UObject& ConfigOwner) const;
-
-	UE_DEPRECATED(5.3, "This flavor of DestroyEntityTemplate is deprecated. Use the one without the ConfigOwner parameter (now a property of the FMassEntityConfig itself)")
-	void DestroyEntityTemplate(const UWorld& World, const UObject& ConfigOwner) const;
-
-	UE_DEPRECATED(5.3, "This flavor of GetEntityTemplateChecked is deprecated. Use the one without the ConfigOwner parameter (now a property of the FMassEntityConfig itself)")
-	const FMassEntityTemplate& GetEntityTemplateChecked(const UWorld& World, const UObject& ConfigOwner) const;
-
-	UE_DEPRECATED(5.3, "This flavor of ValidateEntityTemplate is deprecated. Use the one without the ConfigOwner parameter (now a property of the FMassEntityConfig itself)")
-	bool ValidateEntityTemplate(const UWorld& World, const UObject& ConfigOwner);
-
 	bool IsEmpty() const { return Parent == nullptr && Traits.Num() == 0; }
 
 	const FGuid& GetGuid() const { return ConfigGuid; }
@@ -90,10 +78,6 @@ protected:
 	/** Looks for a trait of the indicated type, accepting all child classes as well, unless bExactMatch == true */
 	UMassEntityTraitBase* FindTraitInternal(TSubclassOf<UMassEntityTraitBase> TraitClass, const bool bExactMatch = false) const;
 
-	/** Combines traits based on the config hierarchy and returns list of unique traits */
-	UE_DEPRECATED(5.3, "This flavor of GetCombinedTraits is deprecated. Use the one without the ConfigOwner parameter (now a property of the FMassEntityConfig itself)")
-	void GetCombinedTraits(TArray<UMassEntityTraitBase*>& OutTraits, TArray<const UObject*>& Visited, const UObject& ConfigOwner) const;
-
 	/** Reference to parent config asset */
 	UPROPERTY(Category = "Derived Traits", EditAnywhere)
 	TObjectPtr<const UMassEntityConfigAsset> Parent = nullptr;
diff --git a/Engine/Plugins/Runtime/MassGameplay/Source/MassSpawner/Public/MassEntityTemplateRegistry.h b/Engine/Plugins/Runtime/MassGameplay/Source/MassSpawner/Public/MassEntityTemplateRegistry.h
--- a/Engine/Plugins/Runtime/MassGameplay/Source/MassSpawner/Public/MassEntityTemplateRegistry.h
+++ b/Engine/Plugins/Runtime/MassGameplay/Source/MassSpawner/Public/MassEntityTemplateRegistry.h
@@ -396,17 +396,8 @@ struct MASSSPAWNER_API FMassEntityTemplateRegistry
 	 */
 	const TSharedRef<FMassEntityTemplate>& FindOrAddTemplate(FMassEntityTemplateID TemplateID, FMassEntityTemplateData&& TemplateData);
 
-	UE_DEPRECATED(5.3, "We no longer support fething mutable templates from the TemplateRegistry. Stored templates are considered const.")
-	FMassEntityTemplate* FindMutableTemplateFromTemplateID(FMassEntityTemplateID TemplateID);
-
-	UE_DEPRECATED(5.3, "CreateTemplate is no longer available. Use FindOrAddTemplate instead.")
-	FMassEntityTemplate& CreateTemplate(const uint32 HashLookup, FMassEntityTemplateID TemplateID);
-
 	void DestroyTemplate(FMassEntityTemplateID TemplateID);
 
-	UE_DEPRECATED(5.3, "InitializeEntityTemplate is no longer available. Use FindOrAddTemplate instead.")
-	void InitializeEntityTemplate(FMassEntityTemplate& InOutTemplate) const;
-
 	FMassEntityManager& GetEntityManagerChecked() { check(EntityManager); return *EntityManager; }
 
 protected:
diff --git a/Engine/Plugins/Runtime/MassGameplay/Source/MassSpawner/Public/MassSpawnerSubsystem.h b/Engine/Plugins/Runtime/MassGameplay/Source/MassSpawner/Public/MassSpawnerSubsystem.h
--- a/Engine/Plugins/Runtime/MassGameplay/Source/MassSpawner/Public/MassSpawnerSubsystem.h
+++ b/Engine/Plugins/Runtime/MassGameplay/Source/MassSpawner/Public/MassSpawnerSubsystem.h
@@ -35,12 +35,6 @@ public:
 
 	void DestroyEntities(TConstArrayView<FMassEntityHandle> Entities);
 
-	UE_DEPRECATED(5.3, "This flavor of DestroyEntities has been deprecated. Use the one without the FMassEntityTemplateID parameter")
-	void DestroyEntities(const FMassEntityTemplateID TemplateID, TConstArrayView<FMassEntityHandle> Entities)
-	{
-		DestroyEntities(Entities);
-	}
-
 	const FMassEntityTemplateRegistry& GetTemplateRegistryInstance() const { return TemplateRegistryInstance; }
 	FMassEntityTemplateRegistry& GetMutableTemplateRegistryInstance() { return TemplateRegistryInstance; }
 
diff --git a/Engine/Source/Editor/MassEntityEditor/Public/MassEntityEditorModule.h b/Engine/Source/Editor/MassEntityEditor/Public/MassEntityEditorModule.h
--- a/Engine/Source/Editor/MassEntityEditor/Public/MassEntityEditorModule.h
+++ b/Engine/Source/Editor/MassEntityEditor/Public/MassEntityEditorModule.h
@@ -20,12 +20,20 @@ public:
 	virtual void StartupModule() override;
 	virtual void ShutdownModule() override;
 
-	virtual TSharedPtr<FExtensibilityManager> GetMenuExtensibilityManager() override { return MenuExtensibilityManager; }
-	virtual TSharedPtr<FExtensibilityManager> GetToolBarExtensibilityManager() override { return ToolBarExtensibilityManager; }
-
-	UE_DEPRECATED(5.3, "Please use GetProcessorClassCache instead")
-	TSharedPtr<FGraphNodeClassHelper> GetProcassorClassCache() { return ProcessorClassCache; }
-	TSharedPtr<FGraphNodeClassHelper> GetProcessorClassCache() { return ProcessorClassCache; }
+	virtual TSharedPtr<FExtensibilityManager> GetMenuExtensibilityManager() override
+	{
+		return MenuExtensibilityManager;
+	}
+
+	virtual TSharedPtr<FExtensibilityManager> GetToolBarExtensibilityManager() override
+	{
+		return ToolBarExtensibilityManager;
+	}
+
+	TSharedPtr<FGraphNodeClassHelper> GetProcessorClassCache()
+	{
+		return ProcessorClassCache;
+	}
 
 protected:
 #if WITH_UNREAL_DEVELOPER_TOOLS
diff --git a/Engine/Source/Runtime/MassEntity/Private/MassProcessor.cpp b/Engine/Source/Runtime/MassEntity/Private/MassProcessor.cpp
--- a/Engine/Source/Runtime/MassEntity/Private/MassProcessor.cpp
+++ b/Engine/Source/Runtime/MassEntity/Private/MassProcessor.cpp
@@ -622,12 +622,3 @@ UMassCompositeProcessor* UMassCompositeProcessor::FindOrAddGroupProcessor(FName
 
 	return GroupProcessor;
 }
-
-//-----------------------------------------------------------------------------
-// DEPRECATED
-//-----------------------------------------------------------------------------
-void UMassCompositeProcessor::Populate(TConstArrayView<FMassProcessorOrderInfo> OrderedProcessors)
-{
-	TArray<FMassProcessorOrderInfo> OrderedProcessorsCopy(OrderedProcessors);
-	UpdateProcessorsCollection(OrderedProcessorsCopy, EProcessorExecutionFlags::None);
-}
diff --git a/Engine/Source/Runtime/MassEntity/Public/MassEntityUtils.h b/Engine/Source/Runtime/MassEntity/Public/MassEntityUtils.h
--- a/Engine/Source/Runtime/MassEntity/Public/MassEntityUtils.h
+++ b/Engine/Source/Runtime/MassEntity/Public/MassEntityUtils.h
@@ -21,12 +21,6 @@ MASSENTITY_API extern EProcessorExecutionFlags DetermineProcessorExecutionFlags(
 /** based on the given World (which can be null), the function determines additional level tick types for the processing phases */
 uint8 DetermineProcessorSupportedTickTypes(const UWorld* World);
 
-UE_DEPRECATED(5.3, "This function is deprecated. Please use the properly spelled GetProcessorExecutionFlagsForWorld instead.")
-FORCEINLINE EProcessorExecutionFlags GetProcessorExecutionFlagsForWold(const UWorld& World)
-{
-	return GetProcessorExecutionFlagsForWorld(World);
-}
-
 /** 
  * Fills OutEntityCollections with per-archetype FMassArchetypeEntityCollection instances. 
  * @param DuplicatesHandling used to inform the function whether to expect duplicates.
diff --git a/Engine/Source/Runtime/MassEntity/Public/MassExecutionContext.h b/Engine/Source/Runtime/MassEntity/Public/MassExecutionContext.h
--- a/Engine/Source/Runtime/MassEntity/Public/MassExecutionContext.h
+++ b/Engine/Source/Runtime/MassEntity/Public/MassExecutionContext.h
@@ -400,69 +400,6 @@ public:
 	//-----------------------------------------------------------------------------
 	// DEPRECATED functions
 	//-----------------------------------------------------------------------------
-
-	template<typename T>
-	UE_DEPRECATED(5.2, "The Get*Subsystem(World) functions have been deprecated. Use the world-less flavor.")
-	T* GetMutableSubsystem(const UWorld* World)
-	{
-		return SubsystemAccess.GetMutableSubsystem<T>();
-	}
-
-	template<typename T>
-	UE_DEPRECATED(5.2, "The Get*Subsystem(World) functions have been deprecated. Use the world-less flavor.")
-	T& GetMutableSubsystemChecked(const UWorld* World)
-	{
-		return SubsystemAccess.GetMutableSubsystemChecked<T>();
-	}
-
-	template<typename T>
-	UE_DEPRECATED(5.2, "The Get*Subsystem(World) functions have been deprecated. Use the world-less flavor.")
-	const T* GetSubsystem(const UWorld* World)
-	{
-		return SubsystemAccess.GetSubsystem<T>();
-	}
-
-	template<typename T>
-	UE_DEPRECATED(5.2, "The Get*Subsystem(World) functions have been deprecated. Use the world-less flavor.")
-	const T& GetSubsystemChecked(const UWorld* World)
-	{
-		return SubsystemAccess.GetSubsystemChecked<T>();
-	}
-
-	template<typename T>
-	UE_DEPRECATED(5.2, "The Get*Subsystem(World) functions have been deprecated. Use the world-less flavor.")
-	T* GetMutableSubsystem(const UWorld* World, const TSubclassOf<USubsystem> SubsystemClass)
-	{
-		return SubsystemAccess.GetMutableSubsystem<T>(SubsystemClass);
-	}
-
-	template<typename T>
-	UE_DEPRECATED(5.2, "The Get*Subsystem(World) functions have been deprecated. Use the world-less flavor.")
-	T& GetMutableSubsystemChecked(const UWorld* World, const TSubclassOf<USubsystem> SubsystemClass)
-	{
-		return SubsystemAccess.GetMutableSubsystemChecked<T>(SubsystemClass);
-	}
-
-	template<typename T>
-	UE_DEPRECATED(5.2, "The Get*Subsystem(World) functions have been deprecated. Use the world-less flavor.")
-	const T* GetSubsystem(const UWorld* World, const TSubclassOf<USubsystem> SubsystemClass)
-	{
-		return SubsystemAccess.GetSubsystem<T>(SubsystemClass);
-	}
-
-	template<typename T>
-	UE_DEPRECATED(5.2, "The Get*Subsystem(World) functions have been deprecated. Use the world-less flavor.")
-	const T& GetSubsystemChecked(const UWorld* World, const TSubclassOf<USubsystem> SubsystemClass)
-	{
-		return SubsystemAccess.GetSubsystemChecked<T>(SubsystemClass);
-	}
-
-	UE_DEPRECATED(5.2, "This version of CacheSubsystemRequirements is deprecated. Use the one without the UWorld parameter")
-	bool CacheSubsystemRequirements(const UWorld*, const FMassSubsystemRequirements& SubsystemRequirements)
-	{
-		return CacheSubsystemRequirements(SubsystemRequirements);
-	}
-
 	UE_DEPRECATED(5.4, "Deprecated in favor of 'SetCurrentArchetypeCompositionDescriptor' as this provides information on the entire archetype.")
 	void SetCurrentArchetypesTagBitSet(const FMassTagBitSet&) {}
 };
diff --git a/Engine/Source/Runtime/MassEntity/Public/MassExternalSubsystemTraits.h b/Engine/Source/Runtime/MassEntity/Public/MassExternalSubsystemTraits.h
--- a/Engine/Source/Runtime/MassEntity/Public/MassExternalSubsystemTraits.h
+++ b/Engine/Source/Runtime/MassEntity/Public/MassExternalSubsystemTraits.h
@@ -34,30 +34,6 @@ struct TMassExternalSubsystemTraits final
 	};
 };
 
-namespace FMassExternalSubsystemTraits
-{
-	/**
-	 * Every TMassExternalSubsystemTraits specialization needs to implement the following. Not supplying default implementations
-	 * to be able to catch missing implementations and header inclusion at compilation time.
-	 * 
-	 * This is a getter function that given a UWorld* fetches an instance.
-	*/
-	template<typename T, typename = typename TEnableIf<TIsDerivedFrom<T, UWorldSubsystem>::IsDerived>::Type>
-	UE_DEPRECATED(5.2, "FMassExternalSubsystemTraits::GetInstance has been marked as deprecated. Use FMassSubsystemAccess::FetchSubsystemInstance instead.")
-	FORCEINLINE T* GetInstance(const UWorld* World)
-	{ 
-		// note that the default implementation works only for UWorldSubsystems
-		return UWorld::GetSubsystem<T>(World);
-	}
-
-	template<typename T, typename = typename TEnableIf<TIsDerivedFrom<T, UWorldSubsystem>::IsDerived>::Type>
-	UE_DEPRECATED(5.2, "FMassExternalSubsystemTraits::GetInstance has been marked as deprecated. Use FMassSubsystemAccess::FetchSubsystemInstance instead.")
-	FORCEINLINE T* GetInstance(const UWorld* World, const TSubclassOf<UWorldSubsystem> SubsystemClass)
-	{
-		return World ? Cast<T>(World->GetSubsystemBase(SubsystemClass)) : (T*)nullptr;
-	}
-}
-
 /** 
  * Shared Fragments' traits.
  * @see TMassExternalSubsystemTraits
diff --git a/Engine/Source/Runtime/MassEntity/Public/MassProcessingPhaseManager.h b/Engine/Source/Runtime/MassEntity/Public/MassProcessingPhaseManager.h
--- a/Engine/Source/Runtime/MassEntity/Public/MassProcessingPhaseManager.h
+++ b/Engine/Source/Runtime/MassEntity/Public/MassProcessingPhaseManager.h
@@ -113,19 +113,6 @@ struct MASSENTITY_API FMassPhaseProcessorConfigurationHelper
 	EMassProcessingPhase Phase;
 	bool bInitializeCreatedProcessors = true;
 	bool bIsGameRuntime = true;
-
-	UE_DEPRECATED(5.2, "This flavor of Configure has been deprecated. Use the one requiring the first parameter to be an array view of additional processors")
-	void Configure(const TSharedPtr<FMassEntityManager>& EntityManager = TSharedPtr<FMassEntityManager>(), FMassProcessorDependencySolver::FResult* OutOptionalResult = nullptr)
-	{
-		Configure({}, EProcessorExecutionFlags::None, EntityManager, OutOptionalResult);
-	}
-
-	UE_DEPRECATED(5.3, "This flavor of Configure has been deprecated. Use the one requiring the first parameter to be an array view of additional processors and EProcessorExecutionFlags to be provided")
-	void Configure(TArrayView<UMassProcessor* const> DynamicProcessors, const TSharedPtr<FMassEntityManager>& EntityManager = TSharedPtr<FMassEntityManager>(),
-		FMassProcessorDependencySolver::FResult* OutOptionalResult = nullptr)
-	{
-		Configure(DynamicProcessors, EProcessorExecutionFlags::None, EntityManager, OutOptionalResult);
-	}
 };
 
 /** 
diff --git a/Engine/Source/Runtime/MassEntity/Public/MassProcessor.h b/Engine/Source/Runtime/MassEntity/Public/MassProcessor.h
--- a/Engine/Source/Runtime/MassEntity/Public/MassProcessor.h
+++ b/Engine/Source/Runtime/MassEntity/Public/MassProcessor.h
@@ -324,8 +324,6 @@ protected:
 	// DEPRECATED
 	//-----------------------------------------------------------------------------
 public:
-	UE_DEPRECATED(5.3, "Populate is deprecated. Please use UpdateProcessorsCollection instead.")
-	void Populate(TConstArrayView<FMassProcessorOrderInfo> OrderedProcessors);
 };
 
 
