[Mass] Issue #623014

diff --git a/Engine/Source/Runtime/MassEntity/Public/MassEntityManager.h b/Engine/Source/Runtime/MassEntity/Public/MassEntityManager.h
--- a/Engine/Source/Runtime/MassEntity/Public/MassEntityManager.h
+++ b/Engine/Source/Runtime/MassEntity/Public/MassEntityManager.h
@@ -3,6 +3,7 @@
 #pragma once
 
 #include "UObject/GCObject.h"
+#include "Containers/StaticArray.h"
 #include "MassEntityTypes.h"
 #include "MassProcessingTypes.h"
 #include "MassEntityQuery.h"
diff --git a/Engine/Source/Runtime/MassEntity/Public/MassExternalSubsystemTraits.h b/Engine/Source/Runtime/MassEntity/Public/MassExternalSubsystemTraits.h
--- a/Engine/Source/Runtime/MassEntity/Public/MassExternalSubsystemTraits.h
+++ b/Engine/Source/Runtime/MassEntity/Public/MassExternalSubsystemTraits.h
@@ -2,9 +2,6 @@
 
 #pragma once
 
-#include "Engine/World.h"
-#include "Subsystems/WorldSubsystem.h"
-
 /**
  * Traits describing how a given piece of code can be used by Mass. We require author or user of a given subsystem to 
  * define its traits. To do it add the following in an accessible location. 
@@ -46,3 +43,8 @@ struct TMassSharedFragmentTraits final
 		GameThreadOnly = false,
 	};
 };
+
+#if UE_ENABLE_INCLUDE_ORDER_DEPRECATED_IN_5_6
+#include "Engine/World.h"
+#include "Subsystems/WorldSubsystem.h"
+#endif // UE_ENABLE_INCLUDE_ORDER_DEPRECATED_IN_5_6
\ No newline at end of file
diff --git a/Engine/Source/Runtime/MassEntity/Public/MassRequirements.h b/Engine/Source/Runtime/MassEntity/Public/MassRequirements.h
--- a/Engine/Source/Runtime/MassEntity/Public/MassRequirements.h
+++ b/Engine/Source/Runtime/MassEntity/Public/MassRequirements.h
@@ -3,12 +3,15 @@
 #pragma once
 
 #include "MassEntityTypes.h"
+#include "MassExternalSubsystemTraits.h"
+#include "Templates/SubclassOf.h"
 #include "MassRequirements.generated.h"
 
 struct FMassDebugger;
 struct FMassArchetypeHandle;
 struct FMassExecutionRequirements;
 struct FMassRequirementAccessDetector;
+class USubsystem;
 
 UENUM()
 enum class EMassFragmentAccess : uint8
