[Mass] marked FMassProcessingPhaseManager::EnableTickFunctions as virtual to support overriding how phase ticks get executed.
This change is following a suggestion by an external Mass user.

diff --git a/Engine/Source/Runtime/MassEntity/Public/MassProcessingPhaseManager.h b/Engine/Source/Runtime/MassEntity/Public/MassProcessingPhaseManager.h
--- a/Engine/Source/Runtime/MassEntity/Public/MassProcessingPhaseManager.h
+++ b/Engine/Source/Runtime/MassEntity/Public/MassProcessingPhaseManager.h
@@ -181,7 +181,8 @@ protected:
 	}
 	// End of FGCObject interface
 
-	void EnableTickFunctions(const UWorld& World);
+	/** Override this function if you want to modify how the phase tick functions get executed. */
+	virtual void EnableTickFunctions(const UWorld& World);
 
 	/** Creates phase processors instances for each declared phase name, based on MassEntitySettings */
 	void CreatePhases();
