Changes to MassProcessor's Query configuration and initialization flow.

ConfigureQueries is no longer getting called as part of PostInitProperties and instead is called as part of processor's Initialize. Initialize in turn got a new parameter - the entity manager instance. This change is in preparation for the upcoming Mass changes and enforces the existing use pattern where a given Processor instance is only ever being used with a single MassEntityManager. Upcoming Mass changes build on that principle.

diff --git a/Engine/Plugins/AI/MassAI/Source/MassAIBehavior/Private/MassStateTreeProcessors.cpp b/Engine/Plugins/AI/MassAI/Source/MassAIBehavior/Private/MassStateTreeProcessors.cpp
--- a/Engine/Plugins/AI/MassAI/Source/MassAIBehavior/Private/MassStateTreeProcessors.cpp
+++ b/Engine/Plugins/AI/MassAI/Source/MassAIBehavior/Private/MassStateTreeProcessors.cpp
@@ -81,10 +81,10 @@ UMassStateTreeFragmentDestructor::UMassStateTreeFragmentDestructor()
 	bRequiresGameThreadExecution = true;
 }
 
-void UMassStateTreeFragmentDestructor::Initialize(UObject& Owner)
+void UMassStateTreeFragmentDestructor::InitializeInternal(UObject& Owner, const TSharedRef<FMassEntityManager>& EntityManager)
 {
 	SignalSubsystem = UWorld::GetSubsystem<UMassSignalSubsystem>(Owner.GetWorld());
-	Super::Initialize(Owner);
+	Super::InitializeInternal(Owner, EntityManager);
 }
 
 void UMassStateTreeFragmentDestructor::ConfigureQueries()
@@ -236,9 +236,9 @@ UMassStateTreeProcessor::UMassStateTreeProcessor(const FObjectInitializer& Objec
 	ExecutionOrder.ExecuteBefore.Add(UE::Mass::ProcessorGroupNames::Tasks);
 }
 
-void UMassStateTreeProcessor::Initialize(UObject& Owner)
+void UMassStateTreeProcessor::InitializeInternal(UObject& Owner, const TSharedRef<FMassEntityManager>& EntityManager)
 {
-	Super::Initialize(Owner);
+	Super::InitializeInternal(Owner, EntityManager);
 
 	UMassSignalSubsystem* SignalSubsystem = UWorld::GetSubsystem<UMassSignalSubsystem>(Owner.GetWorld());
 
diff --git a/Engine/Plugins/AI/MassAI/Source/MassAIBehavior/Private/MassZoneGraphAnnotationProcessors.cpp b/Engine/Plugins/AI/MassAI/Source/MassAIBehavior/Private/MassZoneGraphAnnotationProcessors.cpp
--- a/Engine/Plugins/AI/MassAI/Source/MassAIBehavior/Private/MassZoneGraphAnnotationProcessors.cpp
+++ b/Engine/Plugins/AI/MassAI/Source/MassAIBehavior/Private/MassZoneGraphAnnotationProcessors.cpp
@@ -67,9 +67,9 @@ UMassZoneGraphAnnotationTagUpdateProcessor::UMassZoneGraphAnnotationTagUpdatePro
 	ExecutionOrder.ExecuteBefore.Add(UE::Mass::ProcessorGroupNames::Behavior);
 }
 
-void UMassZoneGraphAnnotationTagUpdateProcessor::Initialize(UObject& Owner)
+void UMassZoneGraphAnnotationTagUpdateProcessor::InitializeInternal(UObject& Owner, const TSharedRef<FMassEntityManager>& EntityManager)
 {
-	Super::Initialize(Owner);
+	Super::InitializeInternal(Owner, EntityManager);
 
 	UMassSignalSubsystem* SignalSubsystem = UWorld::GetSubsystem<UMassSignalSubsystem>(Owner.GetWorld());
 	if (ensureMsgf(SignalSubsystem, TEXT("Expecting to run this processor with UMassSignalSubsystem available")))
diff --git a/Engine/Plugins/AI/MassAI/Source/MassAIBehavior/Public/MassStateTreeProcessors.h b/Engine/Plugins/AI/MassAI/Source/MassAIBehavior/Public/MassStateTreeProcessors.h
--- a/Engine/Plugins/AI/MassAI/Source/MassAIBehavior/Public/MassStateTreeProcessors.h
+++ b/Engine/Plugins/AI/MassAI/Source/MassAIBehavior/Public/MassStateTreeProcessors.h
@@ -24,7 +24,7 @@ public:
 	UMassStateTreeFragmentDestructor();
 
 protected:
-	virtual void Initialize(UObject& Owner) override;
+	virtual void InitializeInternal(UObject& Owner, const TSharedRef<FMassEntityManager>& EntityManager) override;
 	virtual void ConfigureQueries() override;
 	virtual void Execute(FMassEntityManager& EntityManager, FMassExecutionContext& Context) override;
 
@@ -69,7 +69,7 @@ public:
 	UMassStateTreeProcessor(const FObjectInitializer& ObjectInitializer = FObjectInitializer::Get());
 
 protected:
-	virtual void Initialize(UObject& Owner) override;
+	virtual void InitializeInternal(UObject& Owner, const TSharedRef<FMassEntityManager>& EntityManager) override;
 	virtual void ConfigureQueries() override;
 	virtual void SignalEntities(FMassEntityManager& EntityManager, FMassExecutionContext& Context, FMassSignalNameLookup& EntitySignals) override;
 
diff --git a/Engine/Plugins/AI/MassAI/Source/MassAIBehavior/Public/MassZoneGraphAnnotationProcessors.h b/Engine/Plugins/AI/MassAI/Source/MassAIBehavior/Public/MassZoneGraphAnnotationProcessors.h
--- a/Engine/Plugins/AI/MassAI/Source/MassAIBehavior/Public/MassZoneGraphAnnotationProcessors.h
+++ b/Engine/Plugins/AI/MassAI/Source/MassAIBehavior/Public/MassZoneGraphAnnotationProcessors.h
@@ -42,7 +42,7 @@ public:
 	UMassZoneGraphAnnotationTagUpdateProcessor();
 	
 protected:
-	virtual void Initialize(UObject& Owner) override;
+	virtual void InitializeInternal(UObject& Owner, const TSharedRef<FMassEntityManager>& EntityManager) override;
 	virtual void ConfigureQueries() override;
 	virtual void Execute(FMassEntityManager& EntityManager, FMassExecutionContext& Context) override;
 
diff --git a/Engine/Plugins/AI/MassAI/Source/MassNavigation/Private/MassAvoidanceProcessors.cpp b/Engine/Plugins/AI/MassAI/Source/MassNavigation/Private/MassAvoidanceProcessors.cpp
--- a/Engine/Plugins/AI/MassAI/Source/MassNavigation/Private/MassAvoidanceProcessors.cpp
+++ b/Engine/Plugins/AI/MassAI/Source/MassNavigation/Private/MassAvoidanceProcessors.cpp
@@ -371,9 +371,9 @@ void UMassMovingAvoidanceProcessor::ConfigureQueries()
 	EntityQuery.AddConstSharedRequirement<FMassMovementParameters>(EMassFragmentPresence::All);
 }
 
-void UMassMovingAvoidanceProcessor::Initialize(UObject& Owner)
+void UMassMovingAvoidanceProcessor::InitializeInternal(UObject& Owner, const TSharedRef<FMassEntityManager>& EntityManager)
 {
-	Super::Initialize(Owner);
+	Super::InitializeInternal(Owner, EntityManager);
 
 	World = Owner.GetWorld();
 	NavigationSubsystem = UWorld::GetSubsystem<UMassNavigationSubsystem>(Owner.GetWorld());
@@ -1034,9 +1034,9 @@ void UMassStandingAvoidanceProcessor::ConfigureQueries()
 	EntityQuery.AddConstSharedRequirement<FMassStandingAvoidanceParameters>(EMassFragmentPresence::All);
 }
 
-void UMassStandingAvoidanceProcessor::Initialize(UObject& Owner)
+void UMassStandingAvoidanceProcessor::InitializeInternal(UObject& Owner, const TSharedRef<FMassEntityManager>& EntityManager)
 {
-	Super::Initialize(Owner);
+	Super::InitializeInternal(Owner, EntityManager);
 
 	World = Owner.GetWorld();
 	NavigationSubsystem = UWorld::GetSubsystem<UMassNavigationSubsystem>(Owner.GetWorld());
diff --git a/Engine/Plugins/AI/MassAI/Source/MassNavigation/Public/Avoidance/MassAvoidanceProcessors.h b/Engine/Plugins/AI/MassAI/Source/MassNavigation/Public/Avoidance/MassAvoidanceProcessors.h
--- a/Engine/Plugins/AI/MassAI/Source/MassNavigation/Public/Avoidance/MassAvoidanceProcessors.h
+++ b/Engine/Plugins/AI/MassAI/Source/MassNavigation/Public/Avoidance/MassAvoidanceProcessors.h
@@ -23,7 +23,7 @@ public:
 
 protected:
 	virtual void ConfigureQueries() override;
-	virtual void Initialize(UObject& Owner) override;
+	virtual void InitializeInternal(UObject& Owner, const TSharedRef<FMassEntityManager>& EntityManager) override;
 	virtual void Execute(FMassEntityManager& EntityManager, FMassExecutionContext& Context) override;
 
 private:
@@ -43,7 +43,7 @@ public:
 
 protected:
 	virtual void ConfigureQueries() override;
-	virtual void Initialize(UObject& Owner) override;
+	virtual void InitializeInternal(UObject& Owner, const TSharedRef<FMassEntityManager>& EntityManager) override;
 	virtual void Execute(FMassEntityManager& EntityManager, FMassExecutionContext& Context) override;
 
 private:
diff --git a/Engine/Plugins/AI/MassAI/Source/MassZoneGraphNavigation/Private/MassZoneGraphNavigationProcessors.cpp b/Engine/Plugins/AI/MassAI/Source/MassZoneGraphNavigation/Private/MassZoneGraphNavigationProcessors.cpp
--- a/Engine/Plugins/AI/MassAI/Source/MassZoneGraphNavigation/Private/MassZoneGraphNavigationProcessors.cpp
+++ b/Engine/Plugins/AI/MassAI/Source/MassZoneGraphNavigation/Private/MassZoneGraphNavigationProcessors.cpp
@@ -121,9 +121,9 @@ UMassZoneGraphPathFollowProcessor::UMassZoneGraphPathFollowProcessor()
 	ExecutionOrder.ExecuteBefore.Add(UE::Mass::ProcessorGroupNames::Avoidance);
 }
 
-void UMassZoneGraphPathFollowProcessor::Initialize(UObject& Owner)
+void UMassZoneGraphPathFollowProcessor::InitializeInternal(UObject& Owner, const TSharedRef<FMassEntityManager>& EntityManager)
 {
-	Super::Initialize(Owner);
+	Super::InitializeInternal(Owner, EntityManager);
 	SignalSubsystem = UWorld::GetSubsystem<UMassSignalSubsystem>(Owner.GetWorld());
 }
 
@@ -433,9 +433,9 @@ void UMassZoneGraphLaneCacheBoundaryProcessor::ConfigureQueries()
 	EntityQuery.AddTagRequirement<FMassOffLODTag>(EMassFragmentPresence::None);
 }
 
-void UMassZoneGraphLaneCacheBoundaryProcessor::Initialize(UObject& Owner)
+void UMassZoneGraphLaneCacheBoundaryProcessor::InitializeInternal(UObject& Owner, const TSharedRef<FMassEntityManager>& EntityManager)
 {
-	Super::Initialize(Owner);
+	Super::InitializeInternal(Owner, EntityManager);
 
 	WeakWorld = Owner.GetWorld();
 }
diff --git a/Engine/Plugins/AI/MassAI/Source/MassZoneGraphNavigation/Public/MassZoneGraphNavigationProcessors.h b/Engine/Plugins/AI/MassAI/Source/MassZoneGraphNavigation/Public/MassZoneGraphNavigationProcessors.h
--- a/Engine/Plugins/AI/MassAI/Source/MassZoneGraphNavigation/Public/MassZoneGraphNavigationProcessors.h
+++ b/Engine/Plugins/AI/MassAI/Source/MassZoneGraphNavigation/Public/MassZoneGraphNavigationProcessors.h
@@ -38,7 +38,7 @@ class MASSZONEGRAPHNAVIGATION_API UMassZoneGraphPathFollowProcessor : public UMa
 protected:
 	UMassZoneGraphPathFollowProcessor();
 	
-	virtual void Initialize(UObject& Owner) override;
+	virtual void InitializeInternal(UObject& Owner, const TSharedRef<FMassEntityManager>& EntityManager) override;
 	virtual void ConfigureQueries() override;
 	virtual void Execute(FMassEntityManager& EntityManager, FMassExecutionContext& Context) override;
 
@@ -60,7 +60,7 @@ public:
 
 protected:
 	virtual void ConfigureQueries() override;
-	virtual void Initialize(UObject& Owner) override;
+	virtual void InitializeInternal(UObject& Owner, const TSharedRef<FMassEntityManager>& EntityManager) override;
 	virtual void Execute(FMassEntityManager& EntityManager, FMassExecutionContext& Context) override;
 
 private:
diff --git a/Engine/Plugins/AI/MassCrowd/Source/MassCrowd/Private/MassCrowdNavigationProcessor.cpp b/Engine/Plugins/AI/MassCrowd/Source/MassCrowd/Private/MassCrowdNavigationProcessor.cpp
--- a/Engine/Plugins/AI/MassCrowd/Source/MassCrowd/Private/MassCrowdNavigationProcessor.cpp
+++ b/Engine/Plugins/AI/MassCrowd/Source/MassCrowd/Private/MassCrowdNavigationProcessor.cpp
@@ -31,9 +31,9 @@ void UMassCrowdLaneTrackingSignalProcessor::ConfigureQueries()
 	EntityQuery.AddSubsystemRequirement<UMassCrowdSubsystem>(EMassFragmentAccess::ReadWrite);
 }
 
-void UMassCrowdLaneTrackingSignalProcessor::Initialize(UObject& Owner)
+void UMassCrowdLaneTrackingSignalProcessor::InitializeInternal(UObject& Owner, const TSharedRef<FMassEntityManager>& EntityManager)
 {
-	Super::Initialize(Owner);
+	Super::InitializeInternal(Owner, EntityManager);
 	
 	UMassSignalSubsystem* SignalSubsystem = UWorld::GetSubsystem<UMassSignalSubsystem>(Owner.GetWorld());
 	SubscribeToSignal(*SignalSubsystem, UE::Mass::Signals::CurrentLaneChanged);
@@ -110,9 +110,9 @@ UMassCrowdDynamicObstacleProcessor::UMassCrowdDynamicObstacleProcessor()
 	ExecutionOrder.ExecuteBefore.Add(UE::Mass::ProcessorGroupNames::UpdateAnnotationTags);
 }
 
-void UMassCrowdDynamicObstacleProcessor::Initialize(UObject& Owner)
+void UMassCrowdDynamicObstacleProcessor::InitializeInternal(UObject& Owner, const TSharedRef<FMassEntityManager>& EntityManager)
 {
-	Super::Initialize(Owner);
+	Super::InitializeInternal(Owner, EntityManager);
 
 	ZoneGraphAnnotationSubsystem = UWorld::GetSubsystem<UZoneGraphAnnotationSubsystem>(Owner.GetWorld());
 	checkf(ZoneGraphAnnotationSubsystem != nullptr, TEXT("UZoneGraphAnnotationSubsystem is mandatory when using this processor."));
@@ -274,9 +274,9 @@ UMassCrowdDynamicObstacleDeinitializer::UMassCrowdDynamicObstacleDeinitializer()
 	Operation = EMassObservedOperation::Remove;
 }
 
-void UMassCrowdDynamicObstacleDeinitializer::Initialize(UObject& Owner)
+void UMassCrowdDynamicObstacleDeinitializer::InitializeInternal(UObject& Owner, const TSharedRef<FMassEntityManager>& EntityManager)
 {
-	Super::Initialize(Owner);
+	Super::InitializeInternal(Owner, EntityManager);
 
 	UWorld* World = Owner.GetWorld();
 	ZoneGraphAnnotationSubsystem = UWorld::GetSubsystem<UZoneGraphAnnotationSubsystem>(World);
diff --git a/Engine/Plugins/AI/MassCrowd/Source/MassCrowd/Private/MassCrowdServerRepresentationLODProcessor.cpp b/Engine/Plugins/AI/MassCrowd/Source/MassCrowd/Private/MassCrowdServerRepresentationLODProcessor.cpp
--- a/Engine/Plugins/AI/MassCrowd/Source/MassCrowd/Private/MassCrowdServerRepresentationLODProcessor.cpp
+++ b/Engine/Plugins/AI/MassCrowd/Source/MassCrowd/Private/MassCrowdServerRepresentationLODProcessor.cpp
@@ -43,11 +43,11 @@ void UMassCrowdServerRepresentationLODProcessor::ConfigureQueries()
 	ProcessorRequirements.AddSubsystemRequirement<UMassLODSubsystem>(EMassFragmentAccess::ReadOnly);
 }
 
-void UMassCrowdServerRepresentationLODProcessor::Initialize(UObject& InOwner)
+void UMassCrowdServerRepresentationLODProcessor::InitializeInternal(UObject& InOwner, const TSharedRef<FMassEntityManager>& EntityManager)
 {
 	LODCalculator.Initialize(LODDistance, BufferHysteresisOnDistancePercentage / 100.0f, LODMaxCount);
 
-	Super::Initialize(InOwner);
+	Super::InitializeInternal(InOwner, EntityManager);
 }
 
 void UMassCrowdServerRepresentationLODProcessor::Execute(FMassEntityManager& EntityManager, FMassExecutionContext& Context)
diff --git a/Engine/Plugins/AI/MassCrowd/Source/MassCrowd/Private/MassCrowdVisualizationProcessor.cpp b/Engine/Plugins/AI/MassCrowd/Source/MassCrowd/Private/MassCrowdVisualizationProcessor.cpp
--- a/Engine/Plugins/AI/MassCrowd/Source/MassCrowd/Private/MassCrowdVisualizationProcessor.cpp
+++ b/Engine/Plugins/AI/MassCrowd/Source/MassCrowd/Private/MassCrowdVisualizationProcessor.cpp
@@ -64,9 +64,9 @@ void UMassDebugCrowdVisualizationProcessor::ConfigureQueries()
 	EntityQuery.RequireMutatingWorldAccess(); // due to UWorld mutable access
 }
 
-void UMassDebugCrowdVisualizationProcessor::Initialize(UObject& Owner)
+void UMassDebugCrowdVisualizationProcessor::InitializeInternal(UObject& Owner, const TSharedRef<FMassEntityManager>& EntityManager)
 {
-	Super::Initialize(Owner);
+	Super::InitializeInternal(Owner, EntityManager);
 
 	World = Owner.GetWorld();
 	check(World);
diff --git a/Engine/Plugins/AI/MassCrowd/Source/MassCrowd/Public/MassCrowdNavigationProcessor.h b/Engine/Plugins/AI/MassCrowd/Source/MassCrowd/Public/MassCrowdNavigationProcessor.h
--- a/Engine/Plugins/AI/MassCrowd/Source/MassCrowd/Public/MassCrowdNavigationProcessor.h
+++ b/Engine/Plugins/AI/MassCrowd/Source/MassCrowd/Public/MassCrowdNavigationProcessor.h
@@ -19,7 +19,7 @@ public:
 	UMassCrowdLaneTrackingSignalProcessor();
 
 protected:
-	virtual void Initialize(UObject& Owner) override;
+	virtual void InitializeInternal(UObject& Owner, const TSharedRef<FMassEntityManager>& EntityManager) override;
 	virtual void ConfigureQueries() override;
 	virtual void SignalEntities(FMassEntityManager& EntityManager, FMassExecutionContext& Context, FMassSignalNameLookup& /*Unused*/) override;
 };
@@ -50,7 +50,7 @@ public:
 	UMassCrowdDynamicObstacleProcessor();
 
 protected:
-	virtual void Initialize(UObject& Owner) override;
+	virtual void InitializeInternal(UObject& Owner, const TSharedRef<FMassEntityManager>& EntityManager)override;
 	virtual void ConfigureQueries() override;
 	virtual void Execute(FMassEntityManager& EntityManager, FMassExecutionContext& Context);
 
@@ -89,7 +89,7 @@ public:
 	UMassCrowdDynamicObstacleDeinitializer();
 
 protected:
-	virtual void Initialize(UObject& Owner) override;
+	virtual void InitializeInternal(UObject& Owner, const TSharedRef<FMassEntityManager>& EntityManager) override;
 	virtual void ConfigureQueries() override;
 	virtual void Execute(FMassEntityManager& EntityManager, FMassExecutionContext& Context) override;
 
diff --git a/Engine/Plugins/AI/MassCrowd/Source/MassCrowd/Public/MassCrowdServerRepresentationLODProcessor.h b/Engine/Plugins/AI/MassCrowd/Source/MassCrowd/Public/MassCrowdServerRepresentationLODProcessor.h
--- a/Engine/Plugins/AI/MassCrowd/Source/MassCrowd/Public/MassCrowdServerRepresentationLODProcessor.h
+++ b/Engine/Plugins/AI/MassCrowd/Source/MassCrowd/Public/MassCrowdServerRepresentationLODProcessor.h
@@ -23,7 +23,7 @@ public:
 protected:
 
 	virtual void ConfigureQueries() override;
-	virtual void Initialize(UObject& InOwner) override;
+	virtual void InitializeInternal(UObject& InOwner, const TSharedRef<FMassEntityManager>& EntityManager) override;
 	virtual void Execute(FMassEntityManager& EntityManager, FMassExecutionContext& Context) override;
 
 	/** Distance where each LOD becomes relevant */
diff --git a/Engine/Plugins/AI/MassCrowd/Source/MassCrowd/Public/MassCrowdVisualizationProcessor.h b/Engine/Plugins/AI/MassCrowd/Source/MassCrowd/Public/MassCrowdVisualizationProcessor.h
--- a/Engine/Plugins/AI/MassCrowd/Source/MassCrowd/Public/MassCrowdVisualizationProcessor.h
+++ b/Engine/Plugins/AI/MassCrowd/Source/MassCrowd/Public/MassCrowdVisualizationProcessor.h
@@ -36,7 +36,7 @@ protected:
 	/** Configure the owned FMassEntityQuery instances to express processor's requirements */
 	virtual void ConfigureQueries() override;
 
-	virtual void Initialize(UObject& Owner) override;
+	virtual void InitializeInternal(UObject& Owner, const TSharedRef<FMassEntityManager>& EntityManager) override;
 
 	/**
 	 * Execution method for this processor
diff --git a/Engine/Plugins/Runtime/MassGameplay/Source/MassEQS/Private/Processors/MassEnvQueryProcessorBase.cpp b/Engine/Plugins/Runtime/MassGameplay/Source/MassEQS/Private/Processors/MassEnvQueryProcessorBase.cpp
--- a/Engine/Plugins/Runtime/MassGameplay/Source/MassEQS/Private/Processors/MassEnvQueryProcessorBase.cpp
+++ b/Engine/Plugins/Runtime/MassGameplay/Source/MassEQS/Private/Processors/MassEnvQueryProcessorBase.cpp
@@ -4,9 +4,9 @@
 #include "Engine/World.h"
 #include "MassEQSSubsystem.h"
 
-void UMassEnvQueryProcessorBase::Initialize(UObject& Owner)
+void UMassEnvQueryProcessorBase::InitializeInternal(UObject& Owner, const TSharedRef<FMassEntityManager>& EntityManager)
 {
-	Super::Initialize(Owner);
+	Super::InitializeInternal(Owner, EntityManager);
 
 	if (CorrespondingRequestClass)
 	{
diff --git a/Engine/Plugins/Runtime/MassGameplay/Source/MassEQS/Public/Processors/MassEnvQueryProcessorBase.h b/Engine/Plugins/Runtime/MassGameplay/Source/MassEQS/Public/Processors/MassEnvQueryProcessorBase.h
--- a/Engine/Plugins/Runtime/MassGameplay/Source/MassEQS/Public/Processors/MassEnvQueryProcessorBase.h
+++ b/Engine/Plugins/Runtime/MassGameplay/Source/MassEQS/Public/Processors/MassEnvQueryProcessorBase.h
@@ -14,7 +14,7 @@ class UMassEnvQueryProcessorBase : public UMassProcessor
 	GENERATED_BODY()
 
 protected:
-	virtual void Initialize(UObject& Owner);
+	virtual void InitializeInternal(UObject& Owner, const TSharedRef<FMassEntityManager>& EntityManager);
 	
 	TSubclassOf<UEnvQueryNode> CorrespondingRequestClass = nullptr;
 	int32 CachedRequestQueryIndex = -1;
diff --git a/Engine/Plugins/Runtime/MassGameplay/Source/MassReplication/Private/MassReplicationProcessor.cpp b/Engine/Plugins/Runtime/MassGameplay/Source/MassReplication/Private/MassReplicationProcessor.cpp
--- a/Engine/Plugins/Runtime/MassGameplay/Source/MassReplication/Private/MassReplicationProcessor.cpp
+++ b/Engine/Plugins/Runtime/MassGameplay/Source/MassReplication/Private/MassReplicationProcessor.cpp
@@ -68,9 +68,9 @@ void UMassReplicationProcessor::ConfigureQueries()
 	ProcessorRequirements.AddSubsystemRequirement<UMassLODSubsystem>(EMassFragmentAccess::ReadOnly);
 }
 
-void UMassReplicationProcessor::Initialize(UObject& Owner)
+void UMassReplicationProcessor::InitializeInternal(UObject& Owner, const TSharedRef<FMassEntityManager>& EntityManager)
 {
-	Super::Initialize(Owner);
+	Super::InitializeInternal(Owner, EntityManager);
 
 #if UE_REPLICATION_COMPILE_SERVER_CODE
 	UWorld* World = Owner.GetWorld();
diff --git a/Engine/Plugins/Runtime/MassGameplay/Source/MassReplication/Public/MassReplicationProcessor.h b/Engine/Plugins/Runtime/MassGameplay/Source/MassReplication/Public/MassReplicationProcessor.h
--- a/Engine/Plugins/Runtime/MassGameplay/Source/MassReplication/Public/MassReplicationProcessor.h
+++ b/Engine/Plugins/Runtime/MassGameplay/Source/MassReplication/Public/MassReplicationProcessor.h
@@ -33,7 +33,7 @@ public:
 
 protected:
 	virtual void ConfigureQueries() override;
-	virtual void Initialize(UObject& Owner) override;
+	virtual void InitializeInternal(UObject& Owner, const TSharedRef<FMassEntityManager>& EntityManager) override;
 	virtual void Execute(FMassEntityManager& EntityManager, FMassExecutionContext& Context) override;
 
 	void PrepareExecution(FMassEntityManager& EntityManager);
diff --git a/Engine/Plugins/Runtime/MassGameplay/Source/MassSmartObjects/Private/MassSmartObjectRegistration.cpp b/Engine/Plugins/Runtime/MassGameplay/Source/MassSmartObjects/Private/MassSmartObjectRegistration.cpp
--- a/Engine/Plugins/Runtime/MassGameplay/Source/MassSmartObjects/Private/MassSmartObjectRegistration.cpp
+++ b/Engine/Plugins/Runtime/MassGameplay/Source/MassSmartObjects/Private/MassSmartObjectRegistration.cpp
@@ -170,9 +170,9 @@ void UMassActiveSmartObjectSignalProcessor::ConfigureQueries()
 	OutsideSmartObjectActiveRangeQuery.AddTagRequirement<FMassInActiveSmartObjectsRangeTag>(EMassFragmentPresence::None);
 }
 
-void UMassActiveSmartObjectSignalProcessor::Initialize(UObject& Owner)
+void UMassActiveSmartObjectSignalProcessor::InitializeInternal(UObject& Owner, const TSharedRef<FMassEntityManager>& EntityManager)
 {
-	Super::Initialize(Owner);
+	Super::InitializeInternal(Owner, EntityManager);
 
 	UMassSignalSubsystem* SignalSubsystem = UWorld::GetSubsystem<UMassSignalSubsystem>(Owner.GetWorld());
 	SubscribeToSignal(*SignalSubsystem, UE::Mass::Signals::ActorInstanceHandleChanged);
diff --git a/Engine/Plugins/Runtime/MassGameplay/Source/MassSmartObjects/Public/MassSmartObjectRegistration.h b/Engine/Plugins/Runtime/MassGameplay/Source/MassSmartObjects/Public/MassSmartObjectRegistration.h
--- a/Engine/Plugins/Runtime/MassGameplay/Source/MassSmartObjects/Public/MassSmartObjectRegistration.h
+++ b/Engine/Plugins/Runtime/MassGameplay/Source/MassSmartObjects/Public/MassSmartObjectRegistration.h
@@ -135,7 +135,7 @@ public:
 
 protected:
 	virtual void ConfigureQueries() override;
-	virtual void Initialize(UObject& Owner) override;
+	virtual void InitializeInternal(UObject& Owner, const TSharedRef<FMassEntityManager>& EntityManager) override;
 	virtual void SignalEntities(FMassEntityManager& EntityManager, FMassExecutionContext& Context, FMassSignalNameLookup& /*Unused*/) override;
 
 private:
diff --git a/Engine/Plugins/Runtime/MassGameplay/Source/MassSpawner/Private/MassSpawner.cpp b/Engine/Plugins/Runtime/MassGameplay/Source/MassSpawner/Private/MassSpawner.cpp
--- a/Engine/Plugins/Runtime/MassGameplay/Source/MassSpawner/Private/MassSpawner.cpp
+++ b/Engine/Plugins/Runtime/MassGameplay/Source/MassSpawner/Private/MassSpawner.cpp
@@ -410,15 +410,21 @@ UMassProcessor* AMassSpawner::GetPostSpawnProcessor(TSubclassOf<UMassProcessor>
 		}
 	);
 
-	if (Initializer == nullptr)
+	if (Initializer)
 	{
-		UMassProcessor* NewInitializer = NewObject<UMassProcessor>(this, ProcessorClass);
-		NewInitializer->Initialize(*this);
+		return *Initializer;
+	}
+
+	UMassProcessor* NewInitializer = NewObject<UMassProcessor>(this, ProcessorClass);
+	FMassEntityManager* EntityManager = UE::Mass::Utils::GetEntityManager(this);
+	if (ensureMsgf(EntityManager, TEXT("Unable to determine the current MassEntityManager")))
+	{
+		NewInitializer->CallInitialize(this, EntityManager->AsShared());
 		PostSpawnProcessors.Add(NewInitializer);
 		return NewInitializer;
 	}
 
-	return *Initializer;
+	return nullptr;
 }
 
 void AMassSpawner::SpawnGeneratedEntities(TConstArrayView<FMassEntitySpawnDataGeneratorResult> Results)
diff --git a/Engine/Plugins/Runtime/MassGameplay/Source/MassSpawner/Private/MassSpawnerSubsystem.cpp b/Engine/Plugins/Runtime/MassGameplay/Source/MassSpawner/Private/MassSpawnerSubsystem.cpp
--- a/Engine/Plugins/Runtime/MassGameplay/Source/MassSpawner/Private/MassSpawnerSubsystem.cpp
+++ b/Engine/Plugins/Runtime/MassGameplay/Source/MassSpawner/Private/MassSpawnerSubsystem.cpp
@@ -90,7 +90,7 @@ void UMassSpawnerSubsystem::DestroyEntities(TConstArrayView<FMassEntityHandle> E
 
 UMassProcessor* UMassSpawnerSubsystem::GetSpawnDataInitializer(TSubclassOf<UMassProcessor> InitializerClass)
 {	
-	if (!InitializerClass)
+	if (!InitializerClass || !EntityManager)
 	{
 		return nullptr;
 	}
@@ -104,7 +104,7 @@ UMassProcessor* UMassSpawnerSubsystem::GetSpawnDataInitializer(TSubclassOf<UMass
 	if (Initializer == nullptr)
 	{
 		UMassProcessor* NewInitializer = NewObject<UMassProcessor>(this, InitializerClass);
-		NewInitializer->Initialize(*this);
+		NewInitializer->CallInitialize(this, EntityManager.ToSharedRef());
 		SpawnDataInitializers.Add(NewInitializer);
 		return NewInitializer;
 	}
diff --git a/Engine/Source/Developer/MassEntityTestSuite/Private/MassCompositeProcessorTest.cpp b/Engine/Source/Developer/MassEntityTestSuite/Private/MassCompositeProcessorTest.cpp
--- a/Engine/Source/Developer/MassEntityTestSuite/Private/MassCompositeProcessorTest.cpp
+++ b/Engine/Source/Developer/MassEntityTestSuite/Private/MassCompositeProcessorTest.cpp
@@ -59,7 +59,7 @@ struct FCompositeProcessorTest_NoWork : FEntityTestBase
 				Processors.Add(Processor);
 			}
 
-			CompositeProcessor->SetChildProcessors(MoveTemp(Processors));
+			CompositeProcessor->SetChildProcessors(Processors);
 		}
 
 		FMassProcessingContext ProcessingContext(EntityManager, /*DeltaSeconds=*/0.f);
@@ -74,7 +74,7 @@ struct FCompositeProcessorTest_NoWork : FEntityTestBase
 
 		TimesExecuted = 0;
 		{
-			TArray<UMassProcessor*> Processors;
+			TArray<TObjectPtr<UMassProcessor>> Processors;
 			for (int i = 0; i < 3; ++i)
 			{
 				UMassTestProcessorBase* Processor = NewObject<UMassTestProcessorBase>();
@@ -123,7 +123,7 @@ struct FCompositeProcessorTest_MultipleSubProcessors : FEntityTestBase
 				Processors.Add(Processor);
 			}
 
-			CompositeProcessor->SetChildProcessors(MoveTemp(Processors));
+			CompositeProcessor->SetChildProcessors(Processors);
 		}
 
 		FMassProcessingContext ProcessingContext(EntityManager, /*DeltaSeconds=*/0.f);
diff --git a/Engine/Source/Developer/MassEntityTestSuite/Private/MassExecutionTest.cpp b/Engine/Source/Developer/MassEntityTestSuite/Private/MassExecutionTest.cpp
--- a/Engine/Source/Developer/MassEntityTestSuite/Private/MassExecutionTest.cpp
+++ b/Engine/Source/Developer/MassEntityTestSuite/Private/MassExecutionTest.cpp
@@ -157,7 +157,7 @@ struct FExecution_Sparse : FEntityTestBase
 		{
 			TArray<UMassProcessor*> Processors;
 			Processors.Add(Processor);
-			Pipeline.SetProcessors(MoveTemp(Processors));
+			Pipeline.SetProcessors(Processors);
 		}
 
 		FMassArchetypeEntityCollection EntityCollection(FloatsArchetype);
diff --git a/Engine/Source/Runtime/MassEntity/Private/MassEntitySettings.cpp b/Engine/Source/Runtime/MassEntity/Private/MassEntitySettings.cpp
--- a/Engine/Source/Runtime/MassEntity/Private/MassEntitySettings.cpp
+++ b/Engine/Source/Runtime/MassEntity/Private/MassEntitySettings.cpp
@@ -74,23 +74,9 @@ void UMassEntitySettings::BuildPhases()
 			PhaseConfig.PhaseProcessor->SetProcessingPhase(EMassProcessingPhase(i));
 			const FString PhaseDumpDependencyGraphFileName = !DumpDependencyGraphFileName.IsEmpty() ? DumpDependencyGraphFileName + TEXT("_") + PhaseConfig.PhaseName.ToString() : FString();
 
-			FMassProcessorDependencySolver::FResult Result;
-			Result.DependencyGraphFileName = PhaseDumpDependencyGraphFileName;
-			FMassPhaseProcessorConfigurationHelper Configurator(*PhaseConfig.PhaseProcessor, PhaseConfig, *this, EMassProcessingPhase(i));
-			Configurator.bInitializeCreatedProcessors = false;
-			Configurator.bIsGameRuntime = false;
-			Configurator.Configure({}, EProcessorExecutionFlags::All, /*EntityManager=*/nullptr, &Result);
-
-			if (Result.PrunedProcessorClasses.Num() > 0)
-			{
-				UE_VLOG_UELOG(this, LogMass, Warning, TEXT("Calculating dependencies for phase %s resulted in pruned processors:")
-					, *PhaseConfig.PhaseName.ToString());
-				for (const TSubclassOf<UMassProcessor>& ProcessorClass : Result.PrunedProcessorClasses)
-				{
-					UE_VLOG_UELOG(this, LogMass, Warning, TEXT("\t%s"), *GetNameSafe(ProcessorClass));
-				}
-				UE_VLOG_UELOG(this, LogMass, Warning, TEXT("Make sure EntityQueries owned by these processors got registered or override UMassProcessor::ShouldAllowQueryBasedPruning to return the appropriate result"));
-			}
+			FMassRuntimePipeline TmpPipeline(EProcessorExecutionFlags::All);
+			TmpPipeline.CreateFromArray(PhaseConfig.ProcessorCDOs, *PhaseConfig.PhaseProcessor);
+			PhaseConfig.PhaseProcessor->SetChildProcessors(TmpPipeline.MoveProcessorsArray());
 
 			FStringOutputDevice Ar;
 			PhaseConfig.PhaseProcessor->DebugOutputDescription(Ar);
diff --git a/Engine/Source/Runtime/MassEntity/Private/MassEntityUtils.cpp b/Engine/Source/Runtime/MassEntity/Private/MassEntityUtils.cpp
--- a/Engine/Source/Runtime/MassEntity/Private/MassEntityUtils.cpp
+++ b/Engine/Source/Runtime/MassEntity/Private/MassEntityUtils.cpp
@@ -92,6 +92,13 @@ void CreateEntityCollections(const FMassEntityManager& EntityManager, const TCon
 	}
 }
 
+FMassEntityManager* GetEntityManager(const UObject* WorldContextObject)
+{
+	return WorldContextObject
+		? GetEntityManager(WorldContextObject->GetWorld())
+		: nullptr;
+}
+
 FMassEntityManager* GetEntityManager(const UWorld* World)
 {
 	UMassEntitySubsystem* EntityManager = UWorld::GetSubsystem<UMassEntitySubsystem>(World);
diff --git a/Engine/Source/Runtime/MassEntity/Private/MassObserverManager.cpp b/Engine/Source/Runtime/MassEntity/Private/MassObserverManager.cpp
--- a/Engine/Source/Runtime/MassEntity/Private/MassObserverManager.cpp
+++ b/Engine/Source/Runtime/MassEntity/Private/MassObserverManager.cpp
@@ -43,7 +43,7 @@ namespace UE::Mass::ObserverManager
 					Pipeline.AppendProcessor(ProcessorClass, Owner);
 				}
 			}
-			Pipeline.Initialize(Owner);
+			Pipeline.Initialize(Owner, EntityManager.AsShared());
 		}
 	};
 	} // Private
@@ -304,7 +304,7 @@ void FMassObserverManager::AddObserverInstance(const UScriptStruct& FragmentOrTa
 	// calling initialize to ensure the given processor is related to the same EntityManager
 	if (UObject* Owner = EntityManager.GetOwner())
 	{	
-		ObserverProcessor.Initialize(*Owner);
+		ObserverProcessor.CallInitialize(Owner, EntityManager.AsShared());
 	}
 }
 
diff --git a/Engine/Source/Runtime/MassEntity/Private/MassProcessingPhaseManager.cpp b/Engine/Source/Runtime/MassEntity/Private/MassProcessingPhaseManager.cpp
--- a/Engine/Source/Runtime/MassEntity/Private/MassProcessingPhaseManager.cpp
+++ b/Engine/Source/Runtime/MassEntity/Private/MassProcessingPhaseManager.cpp
@@ -16,6 +16,7 @@
 #define LOCTEXT_NAMESPACE "Mass"
 
 DECLARE_CYCLE_STAT(TEXT("Mass Phase Tick"), STAT_Mass_PhaseTick, STATGROUP_Mass);
+DECLARE_CYCLE_STAT(TEXT("Mass Phase Configure Pipeline Creation"), STAT_Mass_PhaseConfigurePipelineCreation, STATGROUP_Mass);
 
 namespace UE::Mass::Tweakables
 {
@@ -145,24 +146,40 @@ void FMassProcessingPhase::Initialize(FMassProcessingPhaseManager& InPhaseManage
 // FPhaseProcessorConfigurator
 //----------------------------------------------------------------------//
 void FMassPhaseProcessorConfigurationHelper::Configure(TArrayView<UMassProcessor* const> DynamicProcessors
-	, EProcessorExecutionFlags InWorldExecutionFlags, const TSharedPtr<FMassEntityManager>& EntityManager
-	, FMassProcessorDependencySolver::FResult* OutOptionalResult)
+	, EProcessorExecutionFlags InWorldExecutionFlags, const TSharedRef<FMassEntityManager>& EntityManager
+	, FMassProcessorDependencySolver::FResult& InOutOptionalResult)
 {
-	FMassRuntimePipeline TmpPipeline(InWorldExecutionFlags);
-	TmpPipeline.CreateFromArray(PhaseConfig.ProcessorCDOs, ProcessorOuter);
-	for (UMassProcessor* Processor : DynamicProcessors)
+	FMassRuntimePipeline TmpPipeline(PhaseProcessor.GetChildProcessors(), InWorldExecutionFlags);
 	{
-		checkf(Processor != nullptr, TEXT("Dynamic processor provided to MASS is null."));
-		if (Processor->GetProcessingPhase() == Phase)
+		SCOPE_CYCLE_COUNTER(STAT_Mass_PhaseConfigurePipelineCreation);
+
+		for (UMassProcessor* PrunedProcessor : InOutOptionalResult.PrunedProcessors)
+		{
+			if (PrunedProcessor)
+			{
+				TmpPipeline.AppendProcessor(*PrunedProcessor);
+			}
+		}
+
+		for (UMassProcessor* Processor : DynamicProcessors)
 		{
-			TmpPipeline.AppendProcessor(*Processor);
+			checkf(Processor != nullptr, TEXT("Dynamic processor provided to MASS is null."));
+			if (Processor->GetProcessingPhase() == Phase)
+			{
+				TmpPipeline.AppendUniqueProcessor(*Processor);
+			}
 		}
+
+		UObject* Owner = EntityManager->GetOwner();
+		check(Owner);
+		// @todo consider doing this only during initial config.
+		TmpPipeline.AppendUniqueRuntimeProcessorCopies(PhaseConfig.ProcessorCDOs, *Owner, EntityManager);
 	}
 
 	TArray<FMassProcessorOrderInfo> SortedProcessors;
 	FMassProcessorDependencySolver Solver(TmpPipeline.GetMutableProcessors(), bIsGameRuntime);
 
-	Solver.ResolveDependencies(SortedProcessors, EntityManager, OutOptionalResult);
+	Solver.ResolveDependencies(SortedProcessors, EntityManager, &InOutOptionalResult);
 
 	PhaseProcessor.UpdateProcessorsCollection(SortedProcessors, InWorldExecutionFlags);
 
@@ -178,12 +195,6 @@ void FMassPhaseProcessorConfigurationHelper::Configure(TArrayView<UMassProcessor
 		for (UMassProcessor* Processor : TmpPipeline.GetProcessors())
 		{
 			UE_VLOG_UELOG(&PhaseProcessor, LogMass, Verbose, TEXT("\t%s"), *Processor->GetProcessorName());
-			// we're ignoring dynamic processors since these instances were explicitly created by the user, and we don't
-			// want to destroy them - the user would need to re-add them every time we recalculate the dependency graph
-			if (Processor->IsDynamic() == false)
-			{
-				Processor->MarkAsGarbage();
-			}
 		}
 	}
 #endif // WITH_MASSENTITY_DEBUG
@@ -195,19 +206,26 @@ void FMassPhaseProcessorConfigurationHelper::Configure(TArrayView<UMassProcessor
 
 	if (bInitializeCreatedProcessors)
 	{
-		PhaseProcessor.Initialize(ProcessorOuter);
+		PhaseProcessor.InitializeInternal(ProcessorOuter, EntityManager);
 	}
 }
 
 //----------------------------------------------------------------------//
 // FMassProcessingPhaseManager
 //----------------------------------------------------------------------//
+FMassProcessingPhaseManager::FMassProcessingPhaseManager(EProcessorExecutionFlags InProcessorExecutionFlags) 
+	: ProcessorExecutionFlags(InProcessorExecutionFlags)
+{
+#if WITH_MASSENTITY_DEBUG
+	OnDebugEntityManagerInitializedHandle = FMassDebugger::OnEntityManagerInitialized.AddRaw(this, &FMassProcessingPhaseManager::OnDebugEntityManagerInitialized);
+	OnDebugEntityManagerDeinitializedHandle = FMassDebugger::OnEntityManagerInitialized.AddRaw(this, &FMassProcessingPhaseManager::OnDebugEntityManagerDeinitialized);
+#endif // WITH_MASSENTITY_DEBUG
+}
+
 void FMassProcessingPhaseManager::Initialize(UObject& InOwner, TConstArrayView<FMassProcessingPhaseConfig> InProcessingPhasesConfig, const FString& DependencyGraphFileName)
 {
 	UWorld* World = InOwner.GetWorld();
-#if WITH_EDITOR
-	const bool bCreateProcessorGraphPreview = (World != nullptr) && (World->IsEditorWorld() && !World->IsGameWorld());
-#endif // WITH_EDITOR
+
 	Owner = &InOwner;
 	ProcessingPhasesConfig = InProcessingPhasesConfig;
 
@@ -235,22 +253,6 @@ void FMassProcessingPhaseManager::Initialize(UObject& InOwner, TConstArrayView<F
 		PhaseProcessor->DebugOutputDescription(Ar);
 		UE_VLOG(&InOwner, LogMass, Log, TEXT("Setting new group processor for phase %s:\n%s"), *UEnum::GetValueAsString(Phase), *Ar);
 #endif // WITH_MASSENTITY_DEBUG
-
-#if WITH_EDITOR
-		if (bCreateProcessorGraphPreview)
-		{
-			// populating the phase processor with initial data for editor world so that the default processing graph
-			// can be investigated in the editor without running PIE.
-			// Runtime processing graph are initialized at runtime base on the actual archetypes instantiated at call time.
-			FMassProcessorDependencySolver::FResult Result;
-			Result.DependencyGraphFileName = DependencyGraphFileName;
-			FMassPhaseProcessorConfigurationHelper Configurator(*PhaseProcessor, ProcessingPhasesConfig[PhaseAsInt], InOwner, EMassProcessingPhase(PhaseAsInt));
-			Configurator.bIsGameRuntime = false;
-			// passing EProcessorExecutionFlags::All here to gather all available processors since bCreateProcessorGraphPreview 
-			// is true when we want to preview processors that might be available at runtime.
-			Configurator.Configure({}, EProcessorExecutionFlags::All, /*EntityManager=*/nullptr, &Result);
-		}
-#endif // WITH_EDITOR
 	}
 
 	bIsAllowedToTick = true;
@@ -316,12 +318,19 @@ void FMassProcessingPhaseManager::AddReferencedObjects(FReferenceCollector& Coll
 		}
 	}
 
-	for (auto& DynamicProcessor : DynamicProcessors)
+	auto NullProcessorRemover = [](const TObjectPtr<UMassProcessor>& Processor)
 	{
-		if (DynamicProcessor)
-		{
-			Collector.AddReferencedObject(DynamicProcessor);
-		}
+		return !Processor;
+	};
+
+	check(DynamicProcessors.RemoveAllSwap(NullProcessorRemover) == 0);
+	Collector.AddStableReferenceArray(&DynamicProcessors);
+
+	// we also need to store our pruned processors
+	for (FPhaseGraphBuildState& GraphBuildState : ProcessingGraphBuildStates)
+	{
+		check(GraphBuildState.LastResult.PrunedProcessors.RemoveAllSwap(NullProcessorRemover) == 0);
+		Collector.AddStableReferenceArray(&GraphBuildState.LastResult.PrunedProcessors);
 	}
 }
 
@@ -413,7 +422,7 @@ void FMassProcessingPhaseManager::OnPhaseStart(const FMassProcessingPhase& Phase
 			GraphBuildState.LastResult.Reset();
 
 			FMassPhaseProcessorConfigurationHelper Configurator(*PhaseProcessor, ProcessingPhasesConfig[PhaseAsInt], *Owner.Get(), Phase.Phase);
-			Configurator.Configure(DynamicProcessors, ProcessorExecutionFlags, EntityManager, &GraphBuildState.LastResult);
+			Configurator.Configure(DynamicProcessors, ProcessorExecutionFlags, EntityManager.ToSharedRef(), GraphBuildState.LastResult);
 
 			GraphBuildState.bInitialized = true;
 
@@ -460,6 +469,11 @@ FString FMassProcessingPhaseManager::GetName() const
 
 void FMassProcessingPhaseManager::RegisterDynamicProcessor(UMassProcessor& Processor)
 {
+	if (Processor.IsInitialized() == false)
+	{
+		check(EntityManager->GetOwner());
+		Processor.CallInitialize(EntityManager->GetOwner(), EntityManager.ToSharedRef());
+	}
 	DynamicProcessors.Add(&Processor);
 	Processor.MarkAsDynamic();
 	ProcessingGraphBuildStates[uint32(Processor.GetProcessingPhase())].bProcessorsNeedRebuild = true;
@@ -472,6 +486,10 @@ void FMassProcessingPhaseManager::UnregisterDynamicProcessor(UMassProcessor& Pro
 	{
 		DynamicProcessors.RemoveAt(Index);
 		ProcessingGraphBuildStates[uint32(Processor.GetProcessingPhase())].bProcessorsNeedRebuild = true;
+
+		// it's possible that the given dynamic processor is a part of processing graph at the moment
+		// we need to store the information about its removal and use it when rebuilding the graph next time around.
+		ensure(false);
 	}
 	else
 	{
@@ -487,4 +505,31 @@ void FMassProcessingPhaseManager::OnNewArchetype(const FMassArchetypeHandle& New
 	}
 }
 
+#if WITH_MASSENTITY_DEBUG
+void FMassProcessingPhaseManager::OnDebugEntityManagerInitialized(const FMassEntityManager& InEntityManager)
+{
+	
+}
+
+void FMassProcessingPhaseManager::OnDebugEntityManagerDeinitialized(const FMassEntityManager& InEntityManager)
+{
+	
+}
+#endif // WITH_MASSENTITY_DEBUG
+
+//-----------------------------------------------------------------------------
+// DEPRECATED
+//-----------------------------------------------------------------------------
+UE_DEPRECATED(5.6, "This flavor of Configure is deprecated. Please use the one using a TSharedRef<FMassEntityManager> parameter instead")
+void FMassPhaseProcessorConfigurationHelper::Configure(TArrayView<UMassProcessor* const> DynamicProcessors, EProcessorExecutionFlags InWorldExecutionFlags
+	, const TSharedPtr<FMassEntityManager>& EntityManager
+	, FMassProcessorDependencySolver::FResult* OutOptionalResult)
+{
+	if (ensureMsgf(EntityManager, TEXT("Configuring processors without a valid EntityManager is no longer supported"))
+		&& OutOptionalResult)
+	{
+		Configure(DynamicProcessors, InWorldExecutionFlags, EntityManager.ToSharedRef(), *OutOptionalResult);
+	}
+}
+
 #undef LOCTEXT_NAMESPACE
diff --git a/Engine/Source/Runtime/MassEntity/Private/MassProcessingTypes.cpp b/Engine/Source/Runtime/MassEntity/Private/MassProcessingTypes.cpp
--- a/Engine/Source/Runtime/MassEntity/Private/MassProcessingTypes.cpp
+++ b/Engine/Source/Runtime/MassEntity/Private/MassProcessingTypes.cpp
@@ -44,12 +44,19 @@ FMassProcessingContext::~FMassProcessingContext()
 //----------------------------------------------------------------------//
 //  FMassRuntimePipeline
 //----------------------------------------------------------------------//
+FMassRuntimePipeline::FMassRuntimePipeline(TConstArrayView<TObjectPtr<UMassProcessor>> SeedProcessors, EProcessorExecutionFlags WorldExecutionFlags)
+	: Processors(SeedProcessors)
+	, ExecutionFlags(WorldExecutionFlags)
+{
+	
+}
+
 void FMassRuntimePipeline::Reset()
 {
 	Processors.Reset();
 }
 
-void FMassRuntimePipeline::Initialize(UObject& Owner)
+void FMassRuntimePipeline::Initialize(UObject& Owner, const TSharedRef<FMassEntityManager>& EntityManager)
 {
 	// having nulls in Processors should be rare so we run the "remove all nulls" operation below only if we know 
 	// for sure that there are any nulls to be removed
@@ -62,8 +69,7 @@ void FMassRuntimePipeline::Initialize(UObject& Owner)
 			if (Proc->IsInitialized() == false)
 			{
 				REDIRECT_OBJECT_TO_VLOG(Proc, &Owner);
-				Proc->Initialize(Owner);
-				ensureMsgf(Proc->IsInitialized(), TEXT("Missing Super::Initialize call for %s"), *Proc->GetFullName());
+				Proc->CallInitialize(&Owner, EntityManager);
 			}
 		}
 		else
@@ -78,25 +84,31 @@ void FMassRuntimePipeline::Initialize(UObject& Owner)
 	}
 }
 
-void FMassRuntimePipeline::SetProcessors(TArray<UMassProcessor*>&& InProcessors)
+void FMassRuntimePipeline::SetProcessors(TArrayView<UMassProcessor*> InProcessors)
 {
 	Reset();
 	Processors = InProcessors;
 }
 
+void FMassRuntimePipeline::SetProcessors(TArray<TObjectPtr<UMassProcessor>>&& InProcessors)
+{
+	Reset();
+	Processors = MoveTemp(InProcessors);
+}
+
 void FMassRuntimePipeline::CreateFromArray(TConstArrayView<const UMassProcessor*> InProcessors, UObject& InOwner)
 {
 	Reset();
 	AppendOrOverrideRuntimeProcessorCopies(InProcessors, InOwner);
 }
 
-void FMassRuntimePipeline::InitializeFromArray(TConstArrayView<const UMassProcessor*> InProcessors, UObject& InOwner)
+void FMassRuntimePipeline::InitializeFromArray(TConstArrayView<const UMassProcessor*> InProcessors, UObject& InOwner, const TSharedRef<FMassEntityManager>& EntityManager)
 {
 	CreateFromArray(InProcessors, InOwner);
-	Initialize(InOwner);
+	Initialize(InOwner, EntityManager);
 }
 
-void FMassRuntimePipeline::InitializeFromClassArray(TConstArrayView<TSubclassOf<UMassProcessor>> InProcessorClasses, UObject& InOwner)
+void FMassRuntimePipeline::InitializeFromClassArray(TConstArrayView<TSubclassOf<UMassProcessor>> InProcessorClasses, UObject& InOwner, const TSharedRef<FMassEntityManager>& EntityManager)
 {
 	Reset();
 
@@ -120,7 +132,7 @@ void FMassRuntimePipeline::InitializeFromClassArray(TConstArrayView<TSubclassOf<
 		}
 	}
 
-	Initialize(InOwner);
+	Initialize(InOwner, EntityManager);
 }
 
 bool FMassRuntimePipeline::HasProcessorOfExactClass(TSubclassOf<UMassProcessor> InClass) const
@@ -129,7 +141,7 @@ bool FMassRuntimePipeline::HasProcessorOfExactClass(TSubclassOf<UMassProcessor>
 	return Processors.FindByPredicate([TestClass](const UMassProcessor* Proc){ return Proc != nullptr && Proc->GetClass() == TestClass; }) != nullptr;
 }
 
-void FMassRuntimePipeline::AppendUniqueRuntimeProcessorCopies(TConstArrayView<const UMassProcessor*> InProcessors, UObject& InOwner)
+void FMassRuntimePipeline::AppendUniqueRuntimeProcessorCopies(TConstArrayView<const UMassProcessor*> InProcessors, UObject& InOwner, const TSharedRef<FMassEntityManager>& EntityManager)
 {
 	const UWorld* World = InOwner.GetWorld();
 	const EProcessorExecutionFlags WorldExecutionFlags = UE::Mass::Utils::DetermineProcessorExecutionFlags(World, ExecutionFlags);
@@ -138,7 +150,7 @@ void FMassRuntimePipeline::AppendUniqueRuntimeProcessorCopies(TConstArrayView<co
 	for (const UMassProcessor* Proc : InProcessors)
 	{
 		if (Proc && Proc->ShouldExecute(WorldExecutionFlags)
-			&& (Proc->ShouldAllowMultipleInstances() || (HasProcessorOfExactClass(Proc->GetClass()) == false)))
+			&& HasProcessorOfExactClass(Proc->GetClass()) == false)
 		{
 			// unfortunately the const cast is required since NewObject doesn't support const Template object
 			UMassProcessor* ProcCopy = NewObject<UMassProcessor>(&InOwner, Proc->GetClass(), FName(), RF_NoFlags, const_cast<UMassProcessor*>(Proc));
@@ -167,8 +179,7 @@ void FMassRuntimePipeline::AppendUniqueRuntimeProcessorCopies(TConstArrayView<co
 		if (Proc->IsInitialized() == false)
 		{
 			REDIRECT_OBJECT_TO_VLOG(Proc, &InOwner);
-			Proc->Initialize(InOwner);
-			ensureMsgf(Proc->IsInitialized(), TEXT("Missing Super::Initialize call for %s"), *Proc->GetFullName());
+ 			Proc->CallInitialize(&InOwner, EntityManager);
 		}
 	}
 }
@@ -220,6 +231,11 @@ void FMassRuntimePipeline::AppendProcessor(UMassProcessor& Processor)
 	Processors.Add(&Processor);
 }
 
+void FMassRuntimePipeline::AppendUniqueProcessor(UMassProcessor& Processor)
+{
+	Processors.AddUnique(&Processor);
+}
+
 void FMassRuntimePipeline::AppendProcessor(TSubclassOf<UMassProcessor> ProcessorClass, UObject& InOwner)
 {
 	check(ProcessorClass);
@@ -255,3 +271,46 @@ uint32 GetTypeHash(const FMassRuntimePipeline& Instance)
 	return Hash;
 }
 
+//-----------------------------------------------------------------------------
+// DEPRECATED
+//-----------------------------------------------------------------------------
+void FMassRuntimePipeline::Initialize(UObject& Owner)
+{
+	FMassEntityManager* EntityManager = UE::Mass::Utils::GetEntityManager(Owner.GetWorld());
+	if (ensureMsgf(EntityManager, TEXT("Unable to determine the current MassEntityManager")))
+	{
+		Initialize(Owner, EntityManager->AsShared());
+	}
+}
+
+void FMassRuntimePipeline::InitializeFromArray(TConstArrayView<const UMassProcessor*> InProcessors, UObject& InOwner)
+{
+	FMassEntityManager* EntityManager = UE::Mass::Utils::GetEntityManager(&InOwner);
+	if (ensureMsgf(EntityManager, TEXT("Unable to determine the current MassEntityManager")))
+	{
+		InitializeFromArray(InProcessors, InOwner, EntityManager->AsShared());
+	}
+}
+
+void FMassRuntimePipeline::InitializeFromClassArray(TConstArrayView<TSubclassOf<UMassProcessor>> InProcessorClasses, UObject& InOwner)
+{
+	FMassEntityManager* EntityManager = UE::Mass::Utils::GetEntityManager(&InOwner);
+	if (ensureMsgf(EntityManager, TEXT("Unable to determine the current MassEntityManager")))
+	{
+		InitializeFromClassArray(InProcessorClasses,  InOwner, EntityManager->AsShared());
+	}	
+}
+
+void FMassRuntimePipeline::AppendUniqueRuntimeProcessorCopies(TConstArrayView<const UMassProcessor*> InProcessors, UObject& InOwner)
+{
+	FMassEntityManager* EntityManager = UE::Mass::Utils::GetEntityManager(&InOwner);
+	if (ensureMsgf(EntityManager, TEXT("Unable to determine the current MassEntityManager")))
+	{
+		AppendUniqueRuntimeProcessorCopies(InProcessors, InOwner, EntityManager->AsShared());
+	}	
+}
+
+void FMassRuntimePipeline::SetProcessors(TArray<UMassProcessor*>&& InProcessors)
+{
+	SetProcessors(MakeArrayView(InProcessors.GetData(), InProcessors.Num()));
+}
diff --git a/Engine/Source/Runtime/MassEntity/Private/MassProcessor.cpp b/Engine/Source/Runtime/MassEntity/Private/MassProcessor.cpp
--- a/Engine/Source/Runtime/MassEntity/Private/MassProcessor.cpp
+++ b/Engine/Source/Runtime/MassEntity/Private/MassProcessor.cpp
@@ -122,9 +122,36 @@ UMassProcessor::UMassProcessor()
 {
 }
 
-void UMassProcessor::Initialize(UObject& Owner)
+void UMassProcessor::CallInitialize(const gsl::not_null<UObject*> Owner, const TSharedRef<FMassEntityManager>& EntityManager)
 {
-	bInitialized = true;
+	if (ensure(HasAnyFlags(RF_ClassDefaultObject) == false 
+		&& GetClass()->HasAnyClassFlags(CLASS_Abstract) == false))
+	{
+		ConfigureQueries();
+
+		bool bNeedsGameThread = ProcessorRequirements.DoesRequireGameThreadExecution();
+		for (const FMassEntityQuery* QueryPtr : OwnedQueries)
+		{
+			CA_ASSUME(QueryPtr);
+			bNeedsGameThread = (bNeedsGameThread || QueryPtr->DoesRequireGameThreadExecution());
+		}
+
+		UE_CLOG(bRequiresGameThreadExecution != bNeedsGameThread, LogMass, Verbose
+			, TEXT("%s is marked bRequiresGameThreadExecution = %s, while the registered queries' or processor requirements indicate the opposite")
+			, *GetProcessorName(), bRequiresGameThreadExecution ? TEXT("TRUE") : TEXT("FALSE"));
+
+		// better safe than sorry - if queries or processor requirements indicate the game thread execution is required, then we mark the whole processor as such
+		bRequiresGameThreadExecution = bRequiresGameThreadExecution || bNeedsGameThread;
+
+		InitializeInternal(*Owner, EntityManager);
+
+		bInitialized = true;
+	}
+}
+
+void UMassProcessor::InitializeInternal(UObject&, const TSharedRef<FMassEntityManager>&)
+{
+	// empty in base class
 }
 
 void UMassProcessor::SetShouldAutoRegisterWithGlobalList(const bool bAutoRegister)
@@ -177,27 +204,6 @@ void UMassProcessor::PostInitProperties()
 {
 	Super::PostInitProperties();
 
-	// We want even the CDO processors to be fully set up so that we can reason about processors based on CDOs, without 
-	// needing to instantiate processors based on a given class.
-	// Note that we skip abstract processors because we're not going to instantiate them at runtime anyway.
-	if (GetClass()->HasAnyClassFlags(CLASS_Abstract) == false)
-	{
-		ConfigureQueries();
-
-		bool bNeedsGameThread = ProcessorRequirements.DoesRequireGameThreadExecution();
-		for (const FMassEntityQuery* QueryPtr : OwnedQueries)
-		{
-			CA_ASSUME(QueryPtr);
-			bNeedsGameThread = (bNeedsGameThread || QueryPtr->DoesRequireGameThreadExecution());
-		}
-
-		UE_CLOG(bRequiresGameThreadExecution != bNeedsGameThread, LogMass, Verbose
-			, TEXT("%s is marked bRequiresGameThreadExecution = %s, while the registered queries' or processor requirements indicate the opposite")
-			, *GetProcessorName(), bRequiresGameThreadExecution ? TEXT("TRUE") : TEXT("FALSE"));
-
-		// better safe than sorry - if queries or processor requirements indicate the game thread execution is required, then we marked the whole processor as such
-		bRequiresGameThreadExecution = bRequiresGameThreadExecution || bNeedsGameThread;
-	}
 #if CPUPROFILERTRACE_ENABLED
 	StatId = GetProcessorName();
 #endif
@@ -318,15 +324,19 @@ UMassCompositeProcessor::UMassCompositeProcessor()
 	bAutoRegisterWithProcessingPhases = false;
 }
 
-void UMassCompositeProcessor::SetChildProcessors(TArray<UMassProcessor*>&& InProcessors)
+void UMassCompositeProcessor::SetChildProcessors(TArrayView<UMassProcessor*> InProcessors)
+{
+	ChildPipeline.SetProcessors(InProcessors);
+}
+
+void UMassCompositeProcessor::SetChildProcessors(TArray<TObjectPtr<UMassProcessor>>&& InProcessors)
 {
 	ChildPipeline.SetProcessors(MoveTemp(InProcessors));
 }
 
 void UMassCompositeProcessor::ConfigureQueries()
 {
-	// nothing to do here since ConfigureQueries will get automatically called for all the processors in ChildPipeline
-	// via their individual PostInitProperties call
+	// nothing to do here since ConfigureQueries will get independently called for all the processors during their creation
 }
 
 FGraphEventRef UMassCompositeProcessor::DispatchProcessorTasks(const TSharedPtr<FMassEntityManager>& EntityManager, FMassExecutionContext& ExecutionContext, const FGraphEventArray& InPrerequisites)
@@ -427,10 +437,10 @@ void UMassCompositeProcessor::Execute(FMassEntityManager& EntityManager, FMassEx
 	}
 }
 
-void UMassCompositeProcessor::Initialize(UObject& Owner)
+void UMassCompositeProcessor::InitializeInternal(UObject& Owner, const TSharedRef<FMassEntityManager>& EntityManager)
 {
-	ChildPipeline.Initialize(Owner);
-	Super::Initialize(Owner);
+	ChildPipeline.Initialize(Owner, EntityManager);
+	Super::InitializeInternal(Owner, EntityManager);
 }
 
 void UMassCompositeProcessor::SetProcessors(TArrayView<UMassProcessor*> InProcessorInstances, const TSharedPtr<FMassEntityManager>& EntityManager)
@@ -622,3 +632,30 @@ UMassCompositeProcessor* UMassCompositeProcessor::FindOrAddGroupProcessor(FName
 
 	return GroupProcessor;
 }
+
+//-----------------------------------------------------------------------------
+// DEPRECATED
+//-----------------------------------------------------------------------------
+
+void UMassProcessor::Initialize(UObject& Owner)
+{
+	FMassEntityManager* EntityManager = UE::Mass::Utils::GetEntityManager(Owner.GetWorld());
+	if (ensureMsgf(EntityManager, TEXT("Unable to determine the current MassEntityManager")))
+	{
+		InitializeInternal(Owner, EntityManager->AsShared());
+	}
+}
+
+//void UMassCompositeProcessor::Initialize(UObject& Owner)
+//{
+//	FMassEntityManager* EntityManager = UE::Mass::Utils::GetEntityManager(Owner.GetWorld());
+//	if (ensureMsgf(EntityManager, TEXT("Unable to determine the current MassEntityManager")))
+//	{
+//		InitializeInternal(Owner, EntityManager->AsShared());
+//	}
+//}
+
+void UMassCompositeProcessor::SetChildProcessors(TArray<UMassProcessor*>&& InProcessors)
+{
+	SetChildProcessors(MakeArrayView(InProcessors.GetData(), InProcessors.Num()));
+}
\ No newline at end of file
diff --git a/Engine/Source/Runtime/MassEntity/Private/MassProcessorDependencySolver.cpp b/Engine/Source/Runtime/MassEntity/Private/MassProcessorDependencySolver.cpp
--- a/Engine/Source/Runtime/MassEntity/Private/MassProcessorDependencySolver.cpp
+++ b/Engine/Source/Runtime/MassEntity/Private/MassProcessorDependencySolver.cpp
@@ -1140,7 +1140,7 @@ void FMassProcessorDependencySolver::ResolveDependencies(TArray<FMassProcessorOr
 
 				if (InOutOptionalResult)
 				{
-					InOutOptionalResult->PrunedProcessorClasses.Add(Node.Processor->GetClass());
+					InOutOptionalResult->PrunedProcessors.Add(Node.Processor);
 				}
 
 				// clearing out the processor will result in the rest of the algorithm to treat it as a group - we still 
@@ -1187,21 +1187,19 @@ void FMassProcessorDependencySolver::ResolveDependencies(TArray<FMassProcessorOr
 
 bool FMassProcessorDependencySolver::IsResultUpToDate(const FMassProcessorDependencySolver::FResult& InResult, TSharedPtr<FMassEntityManager> EntityManager)
 {
-	if (InResult.PrunedProcessorClasses.Num() == 0 || !EntityManager || InResult.ArchetypeDataVersion == EntityManager->GetArchetypeDataVersion())
+	if (InResult.PrunedProcessors.Num() == 0 
+		|| !EntityManager 
+		|| InResult.ArchetypeDataVersion == EntityManager->GetArchetypeDataVersion())
 	{
 		return true;
 	}
 
-	// this is inefficient right now since we're using CDOs and need to check all archetypes every time.
 	// Would be more efficient if we had a common place where all processors live, both active and inactive, so that we can utilize those. 
-	for (const TSubclassOf<UMassProcessor>& ProcessorClass : InResult.PrunedProcessorClasses)
+	for (UMassProcessor* PrunedProcessor : InResult.PrunedProcessors)
 	{
-		if (UMassProcessor* ProcessorCDO = ProcessorClass.GetDefaultObject())
+		if (PrunedProcessor && PrunedProcessor->DoesAnyArchetypeMatchOwnedQueries(*EntityManager.Get()))
 		{
-			if (ProcessorCDO->DoesAnyArchetypeMatchOwnedQueries(*EntityManager.Get()))
-			{
-				return false;
-			}
+			return false;
 		}
 	}
 	return true;
diff --git a/Engine/Source/Runtime/MassEntity/Public/MassEntityUtils.h b/Engine/Source/Runtime/MassEntity/Public/MassEntityUtils.h
--- a/Engine/Source/Runtime/MassEntity/Public/MassEntityUtils.h
+++ b/Engine/Source/Runtime/MassEntity/Public/MassEntityUtils.h
@@ -81,6 +81,7 @@ inline void AbstractSort(const int32 NumElements, TPred&& Predicate, TSwap&& Swa
 	}
 }
 
+MASSENTITY_API extern FMassEntityManager* GetEntityManager(const UObject* WorldContextObject);
 MASSENTITY_API extern FMassEntityManager* GetEntityManager(const UWorld* World);
 MASSENTITY_API extern FMassEntityManager& GetEntityManagerChecked(const UWorld& World);
 
diff --git a/Engine/Source/Runtime/MassEntity/Public/MassProcessingPhaseManager.h b/Engine/Source/Runtime/MassEntity/Public/MassProcessingPhaseManager.h
--- a/Engine/Source/Runtime/MassEntity/Public/MassProcessingPhaseManager.h
+++ b/Engine/Source/Runtime/MassEntity/Public/MassProcessingPhaseManager.h
@@ -106,8 +106,8 @@ struct MASSENTITY_API FMassPhaseProcessorConfigurationHelper
 	 * @param InWorldExecutionFlags - provide EProcessorExecutionFlags::None to let underlying code decide
 	 */
 	void Configure(TArrayView<UMassProcessor* const> DynamicProcessors, EProcessorExecutionFlags InWorldExecutionFlags
-		, const TSharedPtr<FMassEntityManager>& EntityManager = TSharedPtr<FMassEntityManager>()
-		, FMassProcessorDependencySolver::FResult* OutOptionalResult = nullptr);
+		, const TSharedRef<FMassEntityManager>& EntityManager
+		, FMassProcessorDependencySolver::FResult& InOutOptionalResult);
 
 	UMassCompositeProcessor& PhaseProcessor;
 	const FMassProcessingPhaseConfig& PhaseConfig;
@@ -115,6 +115,14 @@ struct MASSENTITY_API FMassPhaseProcessorConfigurationHelper
 	EMassProcessingPhase Phase;
 	bool bInitializeCreatedProcessors = true;
 	bool bIsGameRuntime = true;
+
+	//-----------------------------------------------------------------------------
+	// DEPRECATED
+	//-----------------------------------------------------------------------------
+	UE_DEPRECATED(5.6, "This flavor of Configure is deprecated. Please use the one using a TSharedRef<FMassEntityManager> parameter instead")
+	void Configure(TArrayView<UMassProcessor* const> DynamicProcessors, EProcessorExecutionFlags InWorldExecutionFlags
+		, const TSharedPtr<FMassEntityManager>& EntityManager = TSharedPtr<FMassEntityManager>()
+		, FMassProcessorDependencySolver::FResult* OutOptionalResult = nullptr);
 };
 
 /** 
@@ -128,9 +136,7 @@ struct MASSENTITY_API FMassPhaseProcessorConfigurationHelper
 struct MASSENTITY_API FMassProcessingPhaseManager : public FGCObject, public TSharedFromThis<FMassProcessingPhaseManager>
 {
 public:
-	explicit FMassProcessingPhaseManager(EProcessorExecutionFlags InProcessorExecutionFlags = EProcessorExecutionFlags::None) 
-		: ProcessorExecutionFlags(InProcessorExecutionFlags)
-	{}
+	explicit FMassProcessingPhaseManager(EProcessorExecutionFlags InProcessorExecutionFlags = EProcessorExecutionFlags::None);
 	FMassProcessingPhaseManager(const FMassProcessingPhaseManager& Other) = delete;
 	FMassProcessingPhaseManager& operator=(const FMassProcessingPhaseManager& Other) = delete;
 
@@ -227,4 +233,12 @@ protected:
 
 	EProcessorExecutionFlags ProcessorExecutionFlags = EProcessorExecutionFlags::None;
 	bool bIsAllowedToTick = false;
+
+#if WITH_MASSENTITY_DEBUG
+	FDelegateHandle OnDebugEntityManagerInitializedHandle;
+	FDelegateHandle OnDebugEntityManagerDeinitializedHandle;
+
+	void OnDebugEntityManagerInitialized(const FMassEntityManager&);
+	void OnDebugEntityManagerDeinitialized(const FMassEntityManager&);
+#endif // WITH_MASSENTITY_DEBUG
 };
diff --git a/Engine/Source/Runtime/MassEntity/Public/MassProcessingTypes.h b/Engine/Source/Runtime/MassEntity/Public/MassProcessingTypes.h
--- a/Engine/Source/Runtime/MassEntity/Public/MassProcessingTypes.h
+++ b/Engine/Source/Runtime/MassEntity/Public/MassProcessingTypes.h
@@ -82,25 +82,27 @@ private:
 	EProcessorExecutionFlags ExecutionFlags = EProcessorExecutionFlags::None;
 
 public:
-	explicit FMassRuntimePipeline(EProcessorExecutionFlags WorldExecutionFlags = EProcessorExecutionFlags::None)
-		: ExecutionFlags(WorldExecutionFlags)
-	{}
+	explicit FMassRuntimePipeline(EProcessorExecutionFlags WorldExecutionFlags = EProcessorExecutionFlags::None);
+	FMassRuntimePipeline(TConstArrayView<TObjectPtr<UMassProcessor>> SeedProcessors, EProcessorExecutionFlags WorldExecutionFlags);
 
 	void Reset();
-	void Initialize(UObject& Owner);
+	void Initialize(UObject& Owner, const TSharedRef<FMassEntityManager>& EntityManager);
 	
-	/** Creates runtime copies of the given UMassProcessors collection. */
-	void SetProcessors(TArray<UMassProcessor*>&& InProcessors);
+	/** Copies the given array over to this FMassRuntimePipeline instance. */
+	void SetProcessors(TArrayView<UMassProcessor*> InProcessors);
+
+	/** Directly moves the contents of given array to the Processors member array. */
+	void SetProcessors(TArray<TObjectPtr<UMassProcessor>>&& InProcessors);
 
 	/** Creates runtime copies of UMassProcessors given in InProcessors input parameter, using InOwner as new UMassProcessors' outer. */
 	void CreateFromArray(TConstArrayView<const UMassProcessor*> InProcessors, UObject& InOwner);
 
 	/** Calls CreateFromArray and calls Initialize on all processors afterwards. */
-	void InitializeFromArray(TConstArrayView<const UMassProcessor*> InProcessors, UObject& InOwner);
+	void InitializeFromArray(TConstArrayView<const UMassProcessor*> InProcessors, UObject& InOwner, const TSharedRef<FMassEntityManager>& EntityManager);
 	
 	/** Creates runtime instances of UMassProcessors for each processor class given via InProcessorClasses. 
 	 *  The instances will be created with InOwner as outer. */
-	void InitializeFromClassArray(TConstArrayView<TSubclassOf<UMassProcessor>> InProcessorClasses, UObject& InOwner);
+	void InitializeFromClassArray(TConstArrayView<TSubclassOf<UMassProcessor>> InProcessorClasses, UObject& InOwner, const TSharedRef<FMassEntityManager>& EntityManager);
 
 	/** Creates a runtime instance of every processors in the given InProcessors array. If a processor of that class
 	 *  already exists in Processors array it gets overridden. Otherwise it gets added to the end of the collection.*/
@@ -109,11 +111,14 @@ public:
 	/** Creates a runtime instance of every processors in the given array if there's no processor of that class in Processors already.
 	 *  Call this function when adding processors to an already configured FMassRuntimePipeline instance. If you're creating 
 	 *  one from scratch calling any of the InitializeFrom* methods will be more efficient (and will produce same results)
-	 *  or call AppendOrOverrideRuntimeProcessorCopies.*/
-	void AppendUniqueRuntimeProcessorCopies(TConstArrayView<const UMassProcessor*> InProcessors, UObject& InOwner);
+	 *  or call AppendOrOverrideRuntimeProcessorCopies.
+	 *	NOTE: there's a change in functionality since 5.6 - the function will no longer create duplicates for processors returning true from ShouldAllowMultipleInstances
+	 */
+	void AppendUniqueRuntimeProcessorCopies(TConstArrayView<const UMassProcessor*> InProcessors, UObject& InOwner, const TSharedRef<FMassEntityManager>& EntityManager);
 
 	/** Adds InProcessor to Processors without any additional checks */
 	void AppendProcessor(UMassProcessor& Processor);
+	void AppendUniqueProcessor(UMassProcessor& Processor);
 
 	/** Creates an instance of ProcessorClass and adds it to Processors without any additional checks */
 	void AppendProcessor(TSubclassOf<UMassProcessor> ProcessorClass, UObject& InOwner);
@@ -124,13 +129,34 @@ public:
 	UMassCompositeProcessor* FindTopLevelGroupByName(const FName GroupName);
 
 	bool HasProcessorOfExactClass(TSubclassOf<UMassProcessor> InClass) const;
-	bool IsEmpty() const { return Processors.IsEmpty();}
+	bool IsEmpty() const;
+
+	int32 Num() const;
+	TConstArrayView<TObjectPtr<UMassProcessor>> GetProcessors() const;
+	TArrayView<TObjectPtr<UMassProcessor>> GetMutableProcessors();
 
-	int32 Num() const { return Processors.Num(); }
-	TConstArrayView<TObjectPtr<UMassProcessor>> GetProcessors() const { return Processors; }
-	TArrayView<TObjectPtr<UMassProcessor>> GetMutableProcessors() { return Processors; }
+	/** Returns Processors array using move semantics. This operation clears out this FMassRuntimePipeline instance. */
+	TArray<TObjectPtr<UMassProcessor>>&& MoveProcessorsArray();
 
 	MASSENTITY_API friend uint32 GetTypeHash(const FMassRuntimePipeline& Instance);
+
+	//-----------------------------------------------------------------------------
+	// DEPRECATED
+	//-----------------------------------------------------------------------------
+	UE_DEPRECATED(5.6, "This flavor of Initialize is deprecated. Please use the one requiring a FMassEntityManager parameter")
+	void Initialize(UObject& Owner);
+
+	UE_DEPRECATED(5.6, "This flavor of InitializeFromArray is deprecated. Please use the one requiring a FMassEntityManager parameter")
+	void InitializeFromArray(TConstArrayView<const UMassProcessor*> InProcessors, UObject& InOwner);
+	
+	UE_DEPRECATED(5.6, "This flavor of InitializeFromClassArray is deprecated. Please use the one requiring a FMassEntityManager parameter")
+	void InitializeFromClassArray(TConstArrayView<TSubclassOf<UMassProcessor>> InProcessorClasses, UObject& InOwner);
+
+	UE_DEPRECATED(5.6, "This flavor of AppendUniqueRuntimeProcessorCopies is deprecated. Please use the one requiring a FMassEntityManager parameter")
+	void AppendUniqueRuntimeProcessorCopies(TConstArrayView<const UMassProcessor*> InProcessors, UObject& InOwner);
+
+	UE_DEPRECATED(5.6, "This flavor of SetProcessors is deprecated. Please use one of the others.")
+	void SetProcessors(TArray<UMassProcessor*>&& InProcessors);
 };
 
 UENUM()
@@ -161,3 +187,37 @@ struct FMassProcessorOrderInfo
 	TArray<FName> Dependencies;
 	int32 SequenceIndex = INDEX_NONE;
 };
+
+//-----------------------------------------------------------------------------
+// Inlines
+//-----------------------------------------------------------------------------
+inline FMassRuntimePipeline::FMassRuntimePipeline(EProcessorExecutionFlags WorldExecutionFlags)
+	: ExecutionFlags(WorldExecutionFlags)
+{
+	
+}
+
+inline bool FMassRuntimePipeline::IsEmpty() const
+{
+	return Processors.IsEmpty();
+}
+
+inline int32 FMassRuntimePipeline::Num() const
+{
+	return Processors.Num();
+}
+
+inline TConstArrayView<TObjectPtr<UMassProcessor>> FMassRuntimePipeline::GetProcessors() const
+{
+	return Processors;
+}
+
+inline TArrayView<TObjectPtr<UMassProcessor>> FMassRuntimePipeline::GetMutableProcessors()
+{
+	return Processors;
+}
+
+inline TArray<TObjectPtr<UMassProcessor>>&& FMassRuntimePipeline::MoveProcessorsArray()
+{
+	return MoveTemp(Processors);
+}
diff --git a/Engine/Source/Runtime/MassEntity/Public/MassProcessor.h b/Engine/Source/Runtime/MassEntity/Public/MassProcessor.h
--- a/Engine/Source/Runtime/MassEntity/Public/MassProcessor.h
+++ b/Engine/Source/Runtime/MassEntity/Public/MassProcessor.h
@@ -60,8 +60,10 @@ public:
 	explicit UMassProcessor(const FObjectInitializer& ObjectInitializer);
 
 	bool IsInitialized() const;
-	/** Called to initialize the processor. Override to perform custom steps. The super implementation has to be called. */
-	virtual void Initialize(UObject& Owner);
+
+	/** Calls InitializeInternal and handles initialization bookkeeping. */
+	void CallInitialize(const gsl::not_null<UObject*> Owner, const TSharedRef<FMassEntityManager>& EntityManager);
+	
 	virtual FGraphEventRef DispatchProcessorTasks(const TSharedPtr<FMassEntityManager>& EntityManager, FMassExecutionContext& ExecutionContext, const FGraphEventArray& Prerequisites = FGraphEventArray());
 
 	EProcessorExecutionFlags GetExecutionFlags() const;
@@ -145,6 +147,9 @@ public:
 #endif
 	
 protected:
+	/** Called to initialize the processor. Override to perform custom steps. */
+	virtual void InitializeInternal(UObject& Owner, const TSharedRef<FMassEntityManager>& EntityManager);
+
 	virtual void ConfigureQueries() PURE_VIRTUAL(UMassProcessor::ConfigureQueries);
 	virtual void PostInitProperties() override;
 	virtual void Execute(FMassEntityManager& EntityManager, FMassExecutionContext& Context) PURE_VIRTUAL(UMassProcessor::Execute);
@@ -222,6 +227,13 @@ private:
 	 *  @note that it's safe to store pointers here since RegisterQuery does verify that a given registered query is 
 	 *  a member variable of a given processor */
 	TArray<FMassEntityQuery*> OwnedQueries;
+
+	//-----------------------------------------------------------------------------
+	// DEPRECATED
+	//-----------------------------------------------------------------------------
+public:
+	UE_DEPRECATED(5.6, "Initialize is deprecated. Override Initialize(UObject&, const TSharedRef<FMassEntityManager>&) instead. If you want to call the function, use CallInitialize.")
+	virtual void Initialize(UObject& Owner) final;
 };
 
 
@@ -245,9 +257,10 @@ public:
 public:
 	UMassCompositeProcessor();
 
-	void SetChildProcessors(TArray<UMassProcessor*>&& InProcessors);
+	void SetChildProcessors(TArrayView<UMassProcessor*> InProcessors);
+	void SetChildProcessors(TArray<TObjectPtr<UMassProcessor>>&& InProcessors);
 
-	virtual void Initialize(UObject& Owner) override;
+	virtual void InitializeInternal(UObject& Owner, const TSharedRef<FMassEntityManager>& InEntityManager) override;
 	virtual void DebugOutputDescription(FOutputDevice& Ar, int32 Indent = 0) const override;
 	virtual void SetProcessingPhase(EMassProcessingPhase Phase) override;
 
@@ -281,6 +294,8 @@ public:
 
 	virtual FString GetProcessorName() const override;
 
+	TConstArrayView<TObjectPtr<UMassProcessor>> GetChildProcessors() const;
+
 protected:
 	virtual void ConfigureQueries() override;
 	virtual void Execute(FMassEntityManager& EntityManager, FMassExecutionContext& Context) override;
@@ -327,6 +342,12 @@ protected:
 	// DEPRECATED
 	//-----------------------------------------------------------------------------
 public:
+
+	/*UE_DEPRECATED(5.6, "This flavor of Initialize is deprecated. Please use the one requiring a FMassEntityManager parameter")
+	virtual void Initialize(UObject& Owner) final;*/
+
+	UE_DEPRECATED(5.6, "This flavor of SetChildProcessors is deprecated. Please use one of the others.")
+	void SetChildProcessors(TArray<UMassProcessor*>&& InProcessors);
 };
 
 
@@ -427,3 +448,8 @@ inline bool UMassCompositeProcessor::IsEmpty() const
 {
 	return ChildPipeline.IsEmpty();
 }
+
+inline TConstArrayView<TObjectPtr<UMassProcessor>> UMassCompositeProcessor::GetChildProcessors() const
+{
+	return ChildPipeline.GetProcessors();
+}
diff --git a/Engine/Source/Runtime/MassEntity/Public/MassProcessorDependencySolver.h b/Engine/Source/Runtime/MassEntity/Public/MassProcessorDependencySolver.h
--- a/Engine/Source/Runtime/MassEntity/Public/MassProcessorDependencySolver.h
+++ b/Engine/Source/Runtime/MassEntity/Public/MassProcessorDependencySolver.h
@@ -150,13 +150,15 @@ public:
 	struct FResult
 	{
 		FString DependencyGraphFileName;
-		TArray<TSubclassOf<UMassProcessor>> PrunedProcessorClasses;
+		TArray<TObjectPtr<UMassProcessor>> PrunedProcessors;
 		int32 MaxSequenceLength = 0;
 		uint32 ArchetypeDataVersion = 0;
+		UE_DEPRECATED(5.6, "This property is deprecated, replaced by PrunedProcessors")
+		TArray<TSubclassOf<UMassProcessor>> PrunedProcessorClasses;
 
 		void Reset()
 		{
-			PrunedProcessorClasses.Reset();
+			PrunedProcessors.Reset();
 			MaxSequenceLength = 0;
 			ArchetypeDataVersion = 0;
 		}
