[Mass] Fixed removed dynamic processors being left in the processing graph, which in turn causes crashes after GC.

#rnx

diff --git a/Engine/Source/Runtime/MassEntity/Private/MassProcessingPhaseManager.cpp b/Engine/Source/Runtime/MassEntity/Private/MassProcessingPhaseManager.cpp
--- a/Engine/Source/Runtime/MassEntity/Private/MassProcessingPhaseManager.cpp
+++ b/Engine/Source/Runtime/MassEntity/Private/MassProcessingPhaseManager.cpp
@@ -145,7 +145,7 @@ void FMassProcessingPhase::Initialize(FMassProcessingPhaseManager& InPhaseManage
 //----------------------------------------------------------------------//
 // FPhaseProcessorConfigurator
 //----------------------------------------------------------------------//
-void FMassPhaseProcessorConfigurationHelper::Configure(TArrayView<UMassProcessor* const> DynamicProcessors
+void FMassPhaseProcessorConfigurationHelper::Configure(TArrayView<UMassProcessor* const> DynamicProcessors, TArray<TWeakObjectPtr<UMassProcessor>>& InOutRemovedDynamicProcessors
 	, EProcessorExecutionFlags InWorldExecutionFlags, const TSharedRef<FMassEntityManager>& EntityManager
 	, FMassProcessorDependencySolver::FResult& InOutOptionalResult)
 {
@@ -161,6 +161,22 @@ void FMassPhaseProcessorConfigurationHelper::Configure(TArrayView<UMassProcessor
 			}
 		}
 
+		// some previously added dynamic processors were either in the active processor group,
+		// or were among the pruned processors. At this point we have both groups in TmpPipeline now
+		// so we need to check if any of these processors have been removed since las processing
+		// graph recreation
+		for (int32 Index = InOutRemovedDynamicProcessors.Num() - 1; Index >= 0; --Index)
+		{
+			if (const UMassProcessor* RemovedProcessor = InOutRemovedDynamicProcessors[Index].Get())
+			{
+				if (TmpPipeline.RemoveProcessor(*RemovedProcessor) == false)
+				{
+					continue;
+				}
+			}
+			InOutRemovedDynamicProcessors.RemoveAtSwap(Index, 1, EAllowShrinking::No);
+		}
+
 		for (UMassProcessor* Processor : DynamicProcessors)
 		{
 			checkf(Processor != nullptr, TEXT("Dynamic processor provided to MASS is null."));
@@ -422,7 +438,7 @@ void FMassProcessingPhaseManager::OnPhaseStart(const FMassProcessingPhase& Phase
 			GraphBuildState.LastResult.Reset();
 
 			FMassPhaseProcessorConfigurationHelper Configurator(*PhaseProcessor, ProcessingPhasesConfig[PhaseAsInt], *Owner.Get(), Phase.Phase);
-			Configurator.Configure(DynamicProcessors, ProcessorExecutionFlags, EntityManager.ToSharedRef(), GraphBuildState.LastResult);
+			Configurator.Configure(DynamicProcessors, RemovedDynamicProcessors, ProcessorExecutionFlags, EntityManager.ToSharedRef(), GraphBuildState.LastResult);
 
 			GraphBuildState.bInitialized = true;
 
@@ -489,7 +505,7 @@ void FMassProcessingPhaseManager::UnregisterDynamicProcessor(UMassProcessor& Pro
 
 		// it's possible that the given dynamic processor is a part of processing graph at the moment
 		// we need to store the information about its removal and use it when rebuilding the graph next time around.
-		ensure(false);
+		RemovedDynamicProcessors.Add(&Processor);
 	}
 	else
 	{
@@ -528,7 +544,8 @@ void FMassPhaseProcessorConfigurationHelper::Configure(TArrayView<UMassProcessor
 	if (ensureMsgf(EntityManager, TEXT("Configuring processors without a valid EntityManager is no longer supported"))
 		&& OutOptionalResult)
 	{
-		Configure(DynamicProcessors, InWorldExecutionFlags, EntityManager.ToSharedRef(), *OutOptionalResult);
+		static TArray<TWeakObjectPtr<UMassProcessor>> DummyRemovedDynamicProcessors;
+		Configure(DynamicProcessors, DummyRemovedDynamicProcessors, InWorldExecutionFlags, EntityManager.ToSharedRef(), *OutOptionalResult);
 	}
 }
 
diff --git a/Engine/Source/Runtime/MassEntity/Private/MassProcessingTypes.cpp b/Engine/Source/Runtime/MassEntity/Private/MassProcessingTypes.cpp
--- a/Engine/Source/Runtime/MassEntity/Private/MassProcessingTypes.cpp
+++ b/Engine/Source/Runtime/MassEntity/Private/MassProcessingTypes.cpp
@@ -243,9 +243,12 @@ void FMassRuntimePipeline::AppendProcessor(TSubclassOf<UMassProcessor> Processor
 	AppendProcessor(*ProcInstance);
 }
 
-void FMassRuntimePipeline::RemoveProcessor(UMassProcessor& InProcessor)
+bool FMassRuntimePipeline::RemoveProcessor(const UMassProcessor& InProcessor)
 {
-	Processors.Remove(&InProcessor);
+	return Processors.RemoveAll([Processor = &InProcessor](const TObjectPtr<UMassProcessor>& Element)
+		{
+			return Element == Processor;	
+		}) > 0;
 }
 
 UMassCompositeProcessor* FMassRuntimePipeline::FindTopLevelGroupByName(FName GroupName)
diff --git a/Engine/Source/Runtime/MassEntity/Public/MassProcessingPhaseManager.h b/Engine/Source/Runtime/MassEntity/Public/MassProcessingPhaseManager.h
--- a/Engine/Source/Runtime/MassEntity/Public/MassProcessingPhaseManager.h
+++ b/Engine/Source/Runtime/MassEntity/Public/MassProcessingPhaseManager.h
@@ -105,8 +105,8 @@ struct MASSENTITY_API FMassPhaseProcessorConfigurationHelper
 	/** 
 	 * @param InWorldExecutionFlags - provide EProcessorExecutionFlags::None to let underlying code decide
 	 */
-	void Configure(TArrayView<UMassProcessor* const> DynamicProcessors, EProcessorExecutionFlags InWorldExecutionFlags
-		, const TSharedRef<FMassEntityManager>& EntityManager
+	void Configure(TArrayView<UMassProcessor* const> DynamicProcessors, TArray<TWeakObjectPtr<UMassProcessor>>& InOutRemovedDynamicProcessors
+		, EProcessorExecutionFlags InWorldExecutionFlags, const TSharedRef<FMassEntityManager>& EntityManager
 		, FMassProcessorDependencySolver::FResult& InOutOptionalResult);
 
 	UMassCompositeProcessor& PhaseProcessor;
@@ -222,6 +222,7 @@ protected:
 	FPhaseGraphBuildState ProcessingGraphBuildStates[(uint8)EMassProcessingPhase::MAX];
 	TArray<FMassProcessingPhaseConfig> ProcessingPhasesConfig;
 	TArray<TObjectPtr<UMassProcessor>> DynamicProcessors;
+	TArray<TWeakObjectPtr<UMassProcessor>> RemovedDynamicProcessors;
 
 	TSharedPtr<FMassEntityManager> EntityManager;
 
diff --git a/Engine/Source/Runtime/MassEntity/Public/MassProcessingTypes.h b/Engine/Source/Runtime/MassEntity/Public/MassProcessingTypes.h
--- a/Engine/Source/Runtime/MassEntity/Public/MassProcessingTypes.h
+++ b/Engine/Source/Runtime/MassEntity/Public/MassProcessingTypes.h
@@ -123,7 +123,8 @@ public:
 	/** Creates an instance of ProcessorClass and adds it to Processors without any additional checks */
 	void AppendProcessor(TSubclassOf<UMassProcessor> ProcessorClass, UObject& InOwner);
 
-	void RemoveProcessor(UMassProcessor& InProcessor);
+	/** @return whether the given processor has been removed from hosted processors collection */
+	bool RemoveProcessor(const UMassProcessor& InProcessor);
 
 	/** goes through Processor looking for a UMassCompositeProcessor instance which GroupName matches the one given as the parameter */
 	UMassCompositeProcessor* FindTopLevelGroupByName(const FName GroupName);
