[Mass] Minor Mass API extensions, convenience getters mainly

#rnx

diff --git a/Engine/Source/Runtime/MassEntity/Private/MassProcessingPhaseManager.cpp b/Engine/Source/Runtime/MassEntity/Private/MassProcessingPhaseManager.cpp
--- a/Engine/Source/Runtime/MassEntity/Private/MassProcessingPhaseManager.cpp
+++ b/Engine/Source/Runtime/MassEntity/Private/MassProcessingPhaseManager.cpp
@@ -149,7 +149,7 @@ void FMassPhaseProcessorConfigurationHelper::Configure(TArrayView<UMassProcessor
 	, EProcessorExecutionFlags InWorldExecutionFlags, const TSharedRef<FMassEntityManager>& EntityManager
 	, FMassProcessorDependencySolver::FResult& InOutOptionalResult)
 {
-	FMassRuntimePipeline TmpPipeline(PhaseProcessor.GetChildProcessors(), InWorldExecutionFlags);
+	FMassRuntimePipeline TmpPipeline(PhaseProcessor.GetChildProcessorsView(), InWorldExecutionFlags);
 	{
 		SCOPE_CYCLE_COUNTER(STAT_Mass_PhaseConfigurePipelineCreation);
 
diff --git a/Engine/Source/Runtime/MassEntity/Private/MassProcessingTypes.cpp b/Engine/Source/Runtime/MassEntity/Private/MassProcessingTypes.cpp
--- a/Engine/Source/Runtime/MassEntity/Private/MassProcessingTypes.cpp
+++ b/Engine/Source/Runtime/MassEntity/Private/MassProcessingTypes.cpp
@@ -44,13 +44,19 @@ FMassProcessingContext::~FMassProcessingContext()
 //----------------------------------------------------------------------//
 //  FMassRuntimePipeline
 //----------------------------------------------------------------------//
-FMassRuntimePipeline::FMassRuntimePipeline(TConstArrayView<TObjectPtr<UMassProcessor>> SeedProcessors, EProcessorExecutionFlags WorldExecutionFlags)
+FMassRuntimePipeline::FMassRuntimePipeline(TConstArrayView<TObjectPtr<UMassProcessor>> SeedProcessors, const EProcessorExecutionFlags WorldExecutionFlags)
 	: Processors(SeedProcessors)
 	, ExecutionFlags(WorldExecutionFlags)
 {
 	
 }
 
+FMassRuntimePipeline::FMassRuntimePipeline(TConstArrayView<UMassProcessor*> SeedProcessors, const EProcessorExecutionFlags WorldExecutionFlags)
+	: ExecutionFlags(WorldExecutionFlags)
+{
+	Processors = ObjectPtrWrap(SeedProcessors);
+}
+
 void FMassRuntimePipeline::Reset()
 {
 	Processors.Reset();
diff --git a/Engine/Source/Runtime/MassEntity/Public/MassProcessingTypes.h b/Engine/Source/Runtime/MassEntity/Public/MassProcessingTypes.h
--- a/Engine/Source/Runtime/MassEntity/Public/MassProcessingTypes.h
+++ b/Engine/Source/Runtime/MassEntity/Public/MassProcessingTypes.h
@@ -84,6 +84,7 @@ private:
 public:
 	explicit FMassRuntimePipeline(EProcessorExecutionFlags WorldExecutionFlags = EProcessorExecutionFlags::None);
 	FMassRuntimePipeline(TConstArrayView<TObjectPtr<UMassProcessor>> SeedProcessors, EProcessorExecutionFlags WorldExecutionFlags);
+	FMassRuntimePipeline(TConstArrayView<UMassProcessor*> SeedProcessors, EProcessorExecutionFlags WorldExecutionFlags);
 
 	void Reset();
 	void Initialize(UObject& Owner, const TSharedRef<FMassEntityManager>& EntityManager);
@@ -134,6 +135,7 @@ public:
 
 	int32 Num() const;
 	TConstArrayView<TObjectPtr<UMassProcessor>> GetProcessors() const;
+	TConstArrayView<UMassProcessor*> GetProcessorsView() const { return ObjectPtrDecay(Processors); }
 	TArrayView<TObjectPtr<UMassProcessor>> GetMutableProcessors();
 
 	/** Returns Processors array using move semantics. This operation clears out this FMassRuntimePipeline instance. */
diff --git a/Engine/Source/Runtime/MassEntity/Public/MassProcessor.h b/Engine/Source/Runtime/MassEntity/Public/MassProcessor.h
--- a/Engine/Source/Runtime/MassEntity/Public/MassProcessor.h
+++ b/Engine/Source/Runtime/MassEntity/Public/MassProcessor.h
@@ -294,7 +294,7 @@ public:
 
 	virtual FString GetProcessorName() const override;
 
-	TConstArrayView<TObjectPtr<UMassProcessor>> GetChildProcessors() const;
+	TConstArrayView<UMassProcessor*> GetChildProcessorsView() const;
 
 protected:
 	virtual void ConfigureQueries() override;
@@ -449,7 +449,8 @@ inline bool UMassCompositeProcessor::IsEmpty() const
 	return ChildPipeline.IsEmpty();
 }
 
-inline TConstArrayView<TObjectPtr<UMassProcessor>> UMassCompositeProcessor::GetChildProcessors() const
+inline TConstArrayView<UMassProcessor*> UMassCompositeProcessor::GetChildProcessorsView() const
 {
 	return ChildPipeline.GetProcessors();
 }
+
