[Mass] Lowered the severity of logging additional information when reporting a circular dependency error while building the dependency graph. The initial message is at Error-level, so the extended information doesn't need to log at Warning-level, and it was confusing some automated tests.

#rnx

diff --git a/Engine/Source/Runtime/MassEntity/Private/MassProcessorDependencySolver.cpp b/Engine/Source/Runtime/MassEntity/Private/MassProcessorDependencySolver.cpp
--- a/Engine/Source/Runtime/MassEntity/Private/MassProcessorDependencySolver.cpp
+++ b/Engine/Source/Runtime/MassEntity/Private/MassProcessorDependencySolver.cpp
@@ -1020,7 +1020,7 @@ void FMassProcessorDependencySolver::Solve(TArray<FMassProcessorOrderInfo>& OutR
 				UMassProcessor* Processor = AllNodes[Index].Processor;
 				if (Processor)
 				{
-					UE_LOG(LogMass, Warning, TEXT("\t%s, group: %s, before: %s, after %s")
+					UE_LOG(LogMass, Log, TEXT("\t%s, group: %s, before: %s, after %s")
 						, *Processor->GetName()
 						, *Processor->GetExecutionOrder().ExecuteInGroup.ToString()
 						, *NameViewToString(Processor->GetExecutionOrder().ExecuteBefore)
@@ -1029,10 +1029,10 @@ void FMassProcessorDependencySolver::Solve(TArray<FMassProcessorOrderInfo>& OutR
 				else
 				{
 					// group
-					UE_LOG(LogMass, Warning, TEXT("\tGroup %s"), *AllNodes[Index].Name.ToString());
+					UE_LOG(LogMass, Log, TEXT("\tGroup %s"), *AllNodes[Index].Name.ToString());
 				}
 			}
-			UE_LOG(LogMass, Warning, TEXT("Cutting the chain at an arbitrary location."));
+			UE_LOG(LogMass, Log, TEXT("Cutting the chain at an arbitrary location."));
 
 			// remove first dependency
 			// note that if we're in a cycle handling scenario every node does have some dependencies left
