[Mass] headers cleanup, mainly pushing IWYU

diff --git a/Engine/Plugins/AI/MassAI/Source/MassAIBehavior/Public/MassComponentHitTypes.h b/Engine/Plugins/AI/MassAI/Source/MassAIBehavior/Public/MassComponentHitTypes.h
--- a/Engine/Plugins/AI/MassAI/Source/MassAIBehavior/Public/MassComponentHitTypes.h
+++ b/Engine/Plugins/AI/MassAI/Source/MassAIBehavior/Public/MassComponentHitTypes.h
@@ -3,8 +3,9 @@
 #pragma once
 
 #include "Engine/EngineTypes.h"
-#include "MassEntityTypes.h"
+#include "MassEntityHandle.h"
 #if UE_ENABLE_INCLUDE_ORDER_DEPRECATED_IN_5_6
+#include "MassEntityTypes.h"
 #include "StateTreeTypes.h"
 #endif // UE_ENABLE_INCLUDE_ORDER_DEPRECATED_IN_5_6
 #include "MassComponentHitTypes.generated.h"
diff --git a/Engine/Plugins/AI/MassAI/Source/MassAIBehavior/Public/MassLookAtFragments.h b/Engine/Plugins/AI/MassAI/Source/MassAIBehavior/Public/MassLookAtFragments.h
--- a/Engine/Plugins/AI/MassAI/Source/MassAIBehavior/Public/MassLookAtFragments.h
+++ b/Engine/Plugins/AI/MassAI/Source/MassAIBehavior/Public/MassLookAtFragments.h
@@ -2,6 +2,7 @@
 
 #pragma once
 
+#include "MassEntityHandle.h"
 #include "MassEntityTypes.h"
 #include "MassCommonTypes.h"
 #include "ZoneGraphTypes.h"
diff --git a/Engine/Plugins/AI/MassAI/Source/MassAIBehavior/Public/MassStateTreeExecutionContext.h b/Engine/Plugins/AI/MassAI/Source/MassAIBehavior/Public/MassStateTreeExecutionContext.h
--- a/Engine/Plugins/AI/MassAI/Source/MassAIBehavior/Public/MassStateTreeExecutionContext.h
+++ b/Engine/Plugins/AI/MassAI/Source/MassAIBehavior/Public/MassStateTreeExecutionContext.h
@@ -2,7 +2,10 @@
 
 #pragma once
 
+#if UE_ENABLE_INCLUDE_ORDER_DEPRECATED_IN_5_6
 #include "MassEntityTypes.h"
+#endif // UE_ENABLE_INCLUDE_ORDER_DEPRECATED_IN_5_6
+#include "MassEntityHandle.h"
 #include "StateTreeExecutionContext.h"
 
 struct FMassExecutionContext;
diff --git a/Engine/Plugins/AI/MassAI/Source/MassAIDebug/Public/GameplayDebuggerCategory_Mass.h b/Engine/Plugins/AI/MassAI/Source/MassAIDebug/Public/GameplayDebuggerCategory_Mass.h
--- a/Engine/Plugins/AI/MassAI/Source/MassAIDebug/Public/GameplayDebuggerCategory_Mass.h
+++ b/Engine/Plugins/AI/MassAI/Source/MassAIDebug/Public/GameplayDebuggerCategory_Mass.h
@@ -7,7 +7,10 @@
 
 #if WITH_GAMEPLAY_DEBUGGER && WITH_MASSGAMEPLAY_DEBUG
 
+#include "MassEntityHandle.h"
+#if UE_ENABLE_INCLUDE_ORDER_DEPRECATED_IN_5_6
 #include "MassEntityTypes.h"
+#endif // UE_ENABLE_INCLUDE_ORDER_DEPRECATED_IN_5_6
 #include "HAL/IConsoleManager.h"
 #include "UObject/WeakObjectPtrTemplates.h"
 #include "GameplayDebuggerCategory.h"
diff --git a/Engine/Plugins/AI/MassAI/Source/MassNavigation/Public/MassNavigationFragments.h b/Engine/Plugins/AI/MassAI/Source/MassNavigation/Public/MassNavigationFragments.h
--- a/Engine/Plugins/AI/MassAI/Source/MassNavigation/Public/MassNavigationFragments.h
+++ b/Engine/Plugins/AI/MassAI/Source/MassNavigation/Public/MassNavigationFragments.h
@@ -7,7 +7,11 @@
 #include "MassCommonFragments.h"
 #endif
 #include "AI/Navigation/NavigationElement.h"
+#if UE_ENABLE_INCLUDE_ORDER_DEPRECATED_IN_5_6
+#include "MassCommonFragments.h"
+#endif // UE_ENABLE_INCLUDE_ORDER_DEPRECATED_IN_5_6
 #include "MassCommonTypes.h"
+#include "MassEntityTypes.h"
 #include "MassNavigationSubsystem.h"
 #include "MassNavigationTypes.h"
 #include "MassNavigationFragments.generated.h"
diff --git a/Engine/Plugins/AI/MassAI/Source/MassNavigation/Public/MassNavigationSubsystem.h b/Engine/Plugins/AI/MassAI/Source/MassNavigation/Public/MassNavigationSubsystem.h
--- a/Engine/Plugins/AI/MassAI/Source/MassNavigation/Public/MassNavigationSubsystem.h
+++ b/Engine/Plugins/AI/MassAI/Source/MassNavigation/Public/MassNavigationSubsystem.h
@@ -2,7 +2,10 @@
 
 #pragma once
 
+#include "MassEntityHandle.h"
+#if UE_ENABLE_INCLUDE_ORDER_DEPRECATED_IN_5_6
 #include "MassEntityTypes.h"
+#endif // UE_ENABLE_INCLUDE_ORDER_DEPRECATED_IN_5_6
 #include "HierarchicalHashGrid2D.h"
 #include "MassSubsystemBase.h"
 #include "MassExternalSubsystemTraits.h"
diff --git a/Engine/Plugins/AI/MassAI/Source/MassNavigation/Public/SmoothOrientation/MassSmoothOrientationFragments.h b/Engine/Plugins/AI/MassAI/Source/MassNavigation/Public/SmoothOrientation/MassSmoothOrientationFragments.h
--- a/Engine/Plugins/AI/MassAI/Source/MassNavigation/Public/SmoothOrientation/MassSmoothOrientationFragments.h
+++ b/Engine/Plugins/AI/MassAI/Source/MassNavigation/Public/SmoothOrientation/MassSmoothOrientationFragments.h
@@ -1,7 +1,10 @@
-// Copyright Epic Games, Inc. All Rights Reserved.
+﻿// Copyright Epic Games, Inc. All Rights Reserved.
 
 #pragma once
 
+#if UE_ENABLE_INCLUDE_ORDER_DEPRECATED_IN_5_6
+#include "MassCommonTypes.h"
+#endif // UE_ENABLE_INCLUDE_ORDER_DEPRECATED_IN_5_6
 #include "MassEntityTypes.h"
 #include "MassSmoothOrientationFragments.generated.h"
 
diff --git a/Engine/Plugins/AI/MassAI/Source/MassNavigation/Public/Steering/MassSteeringFragments.h b/Engine/Plugins/AI/MassAI/Source/MassNavigation/Public/Steering/MassSteeringFragments.h
--- a/Engine/Plugins/AI/MassAI/Source/MassNavigation/Public/Steering/MassSteeringFragments.h
+++ b/Engine/Plugins/AI/MassAI/Source/MassNavigation/Public/Steering/MassSteeringFragments.h
@@ -1,4 +1,4 @@
-// Copyright Epic Games, Inc. All Rights Reserved.
+﻿// Copyright Epic Games, Inc. All Rights Reserved.
 
 #pragma once
 
diff --git a/Engine/Plugins/AI/MassAI/Source/MassNavigationEditor/Public/MassNavigationEditorModule.h b/Engine/Plugins/AI/MassAI/Source/MassNavigationEditor/Public/MassNavigationEditorModule.h
--- a/Engine/Plugins/AI/MassAI/Source/MassNavigationEditor/Public/MassNavigationEditorModule.h
+++ b/Engine/Plugins/AI/MassAI/Source/MassNavigationEditor/Public/MassNavigationEditorModule.h
@@ -2,6 +2,10 @@
 
 #pragma  once
 
+#if UE_ENABLE_INCLUDE_ORDER_DEPRECATED_IN_5_6
+#include "CoreMinimal.h"
+#include "AssetTypeCategories.h"
+#endif // UE_ENABLE_INCLUDE_ORDER_DEPRECATED_IN_5_6
 #include "Modules/ModuleManager.h"
 #include "Toolkits/IToolkitHost.h"
 #include "Toolkits/AssetEditorToolkit.h"
@@ -16,9 +20,6 @@ public:
 	virtual void StartupModule() override;
 	virtual void ShutdownModule() override;
 	// End IModuleInterface
-};
 
-#if UE_ENABLE_INCLUDE_ORDER_DEPRECATED_IN_5_6
-#include "CoreMinimal.h"
-#include "AssetTypeCategories.h"
-#endif // UE_ENABLE_INCLUDE_ORDER_DEPRECATED_IN_5_6
\ No newline at end of file
+private:
+};
diff --git a/Engine/Plugins/AI/MassAI/Source/MassZoneGraphNavigation/Public/MassZoneGraphNavigationUtils.h b/Engine/Plugins/AI/MassAI/Source/MassZoneGraphNavigation/Public/MassZoneGraphNavigationUtils.h
--- a/Engine/Plugins/AI/MassAI/Source/MassZoneGraphNavigation/Public/MassZoneGraphNavigationUtils.h
+++ b/Engine/Plugins/AI/MassAI/Source/MassZoneGraphNavigation/Public/MassZoneGraphNavigationUtils.h
@@ -2,7 +2,10 @@
 
 #pragma once
 
+#if UE_ENABLE_INCLUDE_ORDER_DEPRECATED_IN_5_6
 #include "MassEntityTypes.h"
+#endif // UE_ENABLE_INCLUDE_ORDER_DEPRECATED_IN_5_6
+#include "MassEntityHandle.h"
 
 struct FZoneGraphShortPathRequest;
 struct FMassZoneGraphCachedLaneFragment;
diff --git a/Engine/Plugins/AI/MassCrowd/Source/MassCrowd/Public/MassCrowdSubsystem.h b/Engine/Plugins/AI/MassCrowd/Source/MassCrowd/Public/MassCrowdSubsystem.h
--- a/Engine/Plugins/AI/MassCrowd/Source/MassCrowd/Public/MassCrowdSubsystem.h
+++ b/Engine/Plugins/AI/MassCrowd/Source/MassCrowd/Public/MassCrowdSubsystem.h
@@ -2,11 +2,13 @@
 
 #pragma once
 
+#if UE_ENABLE_INCLUDE_ORDER_DEPRECATED_IN_5_6
 #include "MassEntityTypes.h"
+#endif // UE_ENABLE_INCLUDE_ORDER_DEPRECATED_IN_5_6
+#include "MassEntityHandle.h"
 #include "MassCrowdTypes.h"
 #include "ZoneGraphTypes.h"
 #include "MassSubsystemBase.h"
-#include "MassExternalSubsystemTraits.h"
 #include "MassCrowdSubsystem.generated.h"
 
 class UZoneGraphAnnotationSubsystem;
diff --git a/Engine/Plugins/Runtime/InstancedActors/Source/InstancedActors/Public/GameplayDebuggerCategory_InstancedActors.h b/Engine/Plugins/Runtime/InstancedActors/Source/InstancedActors/Public/GameplayDebuggerCategory_InstancedActors.h
--- a/Engine/Plugins/Runtime/InstancedActors/Source/InstancedActors/Public/GameplayDebuggerCategory_InstancedActors.h
+++ b/Engine/Plugins/Runtime/InstancedActors/Source/InstancedActors/Public/GameplayDebuggerCategory_InstancedActors.h
@@ -10,6 +10,7 @@
 
 #include "GameplayDebuggerCategory.h"
 #include "HAL/IConsoleManager.h"
+#include "MassEntityHandle.h"
 
 struct FMassEntityManager;
 
diff --git a/Engine/Plugins/Runtime/InstancedActors/Source/InstancedActors/Public/ServerInstancedActorsSpawnerSubsystem.h b/Engine/Plugins/Runtime/InstancedActors/Source/InstancedActors/Public/ServerInstancedActorsSpawnerSubsystem.h
--- a/Engine/Plugins/Runtime/InstancedActors/Source/InstancedActors/Public/ServerInstancedActorsSpawnerSubsystem.h
+++ b/Engine/Plugins/Runtime/InstancedActors/Source/InstancedActors/Public/ServerInstancedActorsSpawnerSubsystem.h
@@ -7,6 +7,7 @@
 #include "ServerInstancedActorsSpawnerSubsystem.generated.h"
 
 
+struct FMassEntityManager;
 class UInstancedActorsComponent;
 struct FInstancedActorsInstanceHandle;
 struct FActorSpawnParameters;
diff --git a/Engine/Plugins/Runtime/MassGameplay/Source/MassActors/Public/MassActorSpawnerSubsystem.h b/Engine/Plugins/Runtime/MassGameplay/Source/MassActors/Public/MassActorSpawnerSubsystem.h
--- a/Engine/Plugins/Runtime/MassGameplay/Source/MassActors/Public/MassActorSpawnerSubsystem.h
+++ b/Engine/Plugins/Runtime/MassGameplay/Source/MassActors/Public/MassActorSpawnerSubsystem.h
@@ -5,11 +5,16 @@
 #include "StructUtils/InstancedStruct.h"
 #include "IndexedHandle.h"
 #include "Delegates/Delegate.h"
+#if UE_ENABLE_INCLUDE_ORDER_DEPRECATED_IN_5_6
+#include "MassCommonTypes.h"
+#endif // UE_ENABLE_INCLUDE_ORDER_DEPRECATED_IN_5_6
 #include "Templates/SubclassOf.h"
 #include "GameFramework/Actor.h"
 #include "MassSubsystemBase.h"
-#include "MassEntityTypes.h"
+#include "MassEntityHandle.h"
+#include "StructUtils/StructView.h"
 #if UE_ENABLE_INCLUDE_ORDER_DEPRECATED_IN_5_6
+#include "MassEntityTypes.h"
 #include "MassCommonTypes.h"
 #endif // UE_ENABLE_INCLUDE_ORDER_DEPRECATED_IN_5_6
 #include "MassActorSpawnerSubsystem.generated.h"
diff --git a/Engine/Plugins/Runtime/MassGameplay/Source/MassActors/Public/MassActorSubsystem.h b/Engine/Plugins/Runtime/MassGameplay/Source/MassActors/Public/MassActorSubsystem.h
--- a/Engine/Plugins/Runtime/MassGameplay/Source/MassActors/Public/MassActorSubsystem.h
+++ b/Engine/Plugins/Runtime/MassGameplay/Source/MassActors/Public/MassActorSubsystem.h
@@ -2,6 +2,7 @@
 
 #pragma once
 
+#include "MassEntityHandle.h"
 #include "MassEntityTypes.h"
 #include "Engine/ActorInstanceHandle.h"
 #include "MassCommonFragments.h"
diff --git a/Engine/Plugins/Runtime/MassGameplay/Source/MassCommon/Public/MassCommonUtils.h b/Engine/Plugins/Runtime/MassGameplay/Source/MassCommon/Public/MassCommonUtils.h
--- a/Engine/Plugins/Runtime/MassGameplay/Source/MassCommon/Public/MassCommonUtils.h
+++ b/Engine/Plugins/Runtime/MassGameplay/Source/MassCommon/Public/MassCommonUtils.h
@@ -7,8 +7,10 @@
 #endif // UE_ENABLE_INCLUDE_ORDER_DEPRECATED_IN_5_6
 #include "Containers/Array.h"
 #include "Containers/Queue.h"
+#include "MassEntityHandle.h"
+#if UE_ENABLE_INCLUDE_ORDER_DEPRECATED_IN_5_6
 #include "MassEntityTypes.h"
-
+#endif // UE_ENABLE_INCLUDE_ORDER_DEPRECATED_IN_5_6
 
 namespace UE::Mass::Utils
 {
diff --git a/Engine/Plugins/Runtime/MassGameplay/Source/MassEQS/Public/MassEQSTypes.h b/Engine/Plugins/Runtime/MassGameplay/Source/MassEQS/Public/MassEQSTypes.h
--- a/Engine/Plugins/Runtime/MassGameplay/Source/MassEQS/Public/MassEQSTypes.h
+++ b/Engine/Plugins/Runtime/MassGameplay/Source/MassEQS/Public/MassEQSTypes.h
@@ -3,7 +3,10 @@
 #pragma once
 
 #include "IndexedHandle.h"
+#include "MassEntityHandle.h"
+#if UE_ENABLE_INCLUDE_ORDER_DEPRECATED_IN_5_6
 #include "MassEntityTypes.h"
+#endif // UE_ENABLE_INCLUDE_ORDER_DEPRECATED_IN_5_6
 #include "UObject/Interface.h"
 #include "MassEQSTypes.generated.h"
 
diff --git a/Engine/Plugins/Runtime/MassGameplay/Source/MassGameplayDebug/Public/MassGameplayDebugTypes.h b/Engine/Plugins/Runtime/MassGameplay/Source/MassGameplayDebug/Public/MassGameplayDebugTypes.h
--- a/Engine/Plugins/Runtime/MassGameplay/Source/MassGameplayDebug/Public/MassGameplayDebugTypes.h
+++ b/Engine/Plugins/Runtime/MassGameplay/Source/MassGameplayDebug/Public/MassGameplayDebugTypes.h
@@ -2,9 +2,13 @@
 
 #pragma once
 
+#if UE_ENABLE_INCLUDE_ORDER_DEPRECATED_IN_5_6
+#include "CoreMinimal.h"
+#endif // UE_ENABLE_INCLUDE_ORDER_DEPRECATED_IN_5_6
 #include "UObject/ObjectMacros.h"
 #include "Logging/LogMacros.h"
 #include "Engine/DataTable.h"
+#include "MassEntityHandle.h"
 #include "MassEntityTypes.h"
 #include "MassCommonTypes.h"
 #if UE_ENABLE_INCLUDE_ORDER_DEPRECATED_IN_5_6
diff --git a/Engine/Plugins/Runtime/MassGameplay/Source/MassRepresentation/Private/MassVisualizationComponent.cpp b/Engine/Plugins/Runtime/MassGameplay/Source/MassRepresentation/Private/MassVisualizationComponent.cpp
--- a/Engine/Plugins/Runtime/MassGameplay/Source/MassRepresentation/Private/MassVisualizationComponent.cpp
+++ b/Engine/Plugins/Runtime/MassGameplay/Source/MassRepresentation/Private/MassVisualizationComponent.cpp
@@ -3,16 +3,15 @@
 #include "MassVisualizationComponent.h"
 #include "CoreGlobals.h"
 #include "Logging/LogMacros.h"
+#include "MassEntityTypes.h"
 #include "MassVisualizer.h"
 #include "MassRepresentationTypes.h"
 #include "Components/InstancedStaticMeshComponent.h"
-#include "Engine/InstancedStaticMesh.h"
 #include "Engine/CollisionProfile.h"
 #include "RenderUtils.h"
 #include "SceneInterface.h"
 #include "Components/HierarchicalInstancedStaticMeshComponent.h"
 #include "VisualLogger/VisualLogger.h"
-#include "Rendering/NaniteResources.h"
 #include "AI/NavigationSystemBase.h"
 
 
diff --git a/Engine/Plugins/Runtime/MassGameplay/Source/MassRepresentation/Public/MassRepresentationTypes.h b/Engine/Plugins/Runtime/MassGameplay/Source/MassRepresentation/Public/MassRepresentationTypes.h
--- a/Engine/Plugins/Runtime/MassGameplay/Source/MassRepresentation/Public/MassRepresentationTypes.h
+++ b/Engine/Plugins/Runtime/MassGameplay/Source/MassRepresentation/Public/MassRepresentationTypes.h
@@ -3,7 +3,10 @@
 #pragma once
 
 #include "MassLODTypes.h"
+#include "MassEntityHandle.h"
+#if UE_ENABLE_INCLUDE_ORDER_DEPRECATED_IN_5_5
 #include "MassEntityTypes.h"
+#endif // UE_ENABLE_INCLUDE_ORDER_DEPRECATED_IN_5_5
 #include "Engine/DataTable.h"
 #include "Misc/MTAccessDetector.h"
 #include "InstanceDataTypes.h"
diff --git a/Engine/Plugins/Runtime/MassGameplay/Source/MassSignals/Private/MassSignalSubsystem.cpp b/Engine/Plugins/Runtime/MassGameplay/Source/MassSignals/Private/MassSignalSubsystem.cpp
--- a/Engine/Plugins/Runtime/MassGameplay/Source/MassSignals/Private/MassSignalSubsystem.cpp
+++ b/Engine/Plugins/Runtime/MassGameplay/Source/MassSignals/Private/MassSignalSubsystem.cpp
@@ -1,6 +1,7 @@
 // Copyright Epic Games, Inc. All Rights Reserved.
 
 #include "MassSignalSubsystem.h"
+#include "MassEntityManager.h"
 #include "MassExecutionContext.h"
 #include "MassCommandBuffer.h"
 #include "MassSignalTypes.h"
diff --git a/Engine/Plugins/Runtime/MassGameplay/Source/MassSignals/Public/MassSignalSubsystem.h b/Engine/Plugins/Runtime/MassGameplay/Source/MassSignals/Public/MassSignalSubsystem.h
--- a/Engine/Plugins/Runtime/MassGameplay/Source/MassSignals/Public/MassSignalSubsystem.h
+++ b/Engine/Plugins/Runtime/MassGameplay/Source/MassSignals/Public/MassSignalSubsystem.h
@@ -2,12 +2,17 @@
 
 #pragma once
 
+#if UE_ENABLE_INCLUDE_ORDER_DEPRECATED_IN_5_6
 #include "MassEntityManager.h"
+#endif // UE_ENABLE_INCLUDE_ORDER_DEPRECATED_IN_5_6
+#include "MassEntityHandle.h"
 #include "MassSubsystemBase.h"
 #include "Misc/MTAccessDetector.h"
 #include "MassExternalSubsystemTraits.h"
 #include "MassSignalSubsystem.generated.h"
 
+struct FMassExecutionContext;
+
 namespace UE::MassSignal 
 {
 	DECLARE_MULTICAST_DELEGATE_TwoParams(FSignalDelegate, FName /*SignalName*/, TConstArrayView<FMassEntityHandle> /*Entities*/);
diff --git a/Engine/Plugins/Runtime/MassGameplay/Source/MassSmartObjects/Public/MassSmartObjectRequest.h b/Engine/Plugins/Runtime/MassGameplay/Source/MassSmartObjects/Public/MassSmartObjectRequest.h
--- a/Engine/Plugins/Runtime/MassGameplay/Source/MassSmartObjects/Public/MassSmartObjectRequest.h
+++ b/Engine/Plugins/Runtime/MassGameplay/Source/MassSmartObjects/Public/MassSmartObjectRequest.h
@@ -4,8 +4,9 @@
 
 #if UE_ENABLE_INCLUDE_ORDER_DEPRECATED_IN_5_6
 #include "SmartObjectSubsystem.h"
-#endif
+#endif // UE_ENABLE_INCLUDE_ORDER_DEPRECATED_IN_5_6
 #include "Containers/StaticArray.h"
+#include "MassEntityHandle.h"
 #include "MassEntityTypes.h"
 #include "SmartObjectRequestTypes.h"
 #include "ZoneGraphTypes.h"
diff --git a/Engine/Plugins/Runtime/MassGameplay/Source/MassSmartObjects/Public/MassSmartObjectTypes.h b/Engine/Plugins/Runtime/MassGameplay/Source/MassSmartObjects/Public/MassSmartObjectTypes.h
--- a/Engine/Plugins/Runtime/MassGameplay/Source/MassSmartObjects/Public/MassSmartObjectTypes.h
+++ b/Engine/Plugins/Runtime/MassGameplay/Source/MassSmartObjects/Public/MassSmartObjectTypes.h
@@ -2,7 +2,10 @@
 
 #pragma once
 
+#if UE_ENABLE_INCLUDE_ORDER_DEPRECATED_IN_5_6
 #include "MassEntityTypes.h"
+#endif // UE_ENABLE_INCLUDE_ORDER_DEPRECATED_IN_5_6
+#include "MassEntityHandle.h"
 #include "MassSmartObjectTypes.generated.h"
 
 namespace UE::Mass::Signals
diff --git a/Engine/Plugins/Runtime/MassGameplay/Source/MassSpawner/Public/MassEntityConfigAsset.h b/Engine/Plugins/Runtime/MassGameplay/Source/MassSpawner/Public/MassEntityConfigAsset.h
--- a/Engine/Plugins/Runtime/MassGameplay/Source/MassSpawner/Public/MassEntityConfigAsset.h
+++ b/Engine/Plugins/Runtime/MassGameplay/Source/MassSpawner/Public/MassEntityConfigAsset.h
@@ -4,10 +4,12 @@
 
 #if UE_ENABLE_INCLUDE_ORDER_DEPRECATED_IN_5_6
 #include "CoreMinimal.h"
-#include "MassEntityTemplateRegistry.h"
 #endif // UE_ENABLE_INCLUDE_ORDER_DEPRECATED_IN_5_6
 #include "Engine/DataAsset.h"
 #include "MassEntityTemplate.h"
+#if UE_ENABLE_INCLUDE_ORDER_DEPRECATED_IN_5_6
+#include "MassEntityTemplateRegistry.h"
+#endif // UE_ENABLE_INCLUDE_ORDER_DEPRECATED_IN_5_6
 #include "MassEntityConfigAsset.generated.h"
 
 class UMassEntityConfigAsset;
diff --git a/Engine/Plugins/Runtime/MassGameplay/Source/MassSpawner/Public/MassEntityTemplate.h b/Engine/Plugins/Runtime/MassGameplay/Source/MassSpawner/Public/MassEntityTemplate.h
--- a/Engine/Plugins/Runtime/MassGameplay/Source/MassSpawner/Public/MassEntityTemplate.h
+++ b/Engine/Plugins/Runtime/MassGameplay/Source/MassSpawner/Public/MassEntityTemplate.h
@@ -2,7 +2,9 @@
 
 #pragma once
 
+#if UE_ENABLE_INCLUDE_ORDER_DEPRECATED_IN_5_6
 #include "MassEntityManager.h"
+#endif // UE_ENABLE_INCLUDE_ORDER_DEPRECATED_IN_5_6
 #include "MassTranslator.h"
 #include "Templates/SharedPointer.h"
 #include "Misc/Guid.h"
@@ -18,6 +20,7 @@ class UMassEntityTraitBase;
 struct FMassEntityView;
 struct FMassEntityTemplateIDFactory;
 struct FMassEntityTemplate;
+struct FMassEntityManager;
 
 
 //ID of the template an entity is using
diff --git a/Engine/Plugins/Runtime/MassGameplay/Source/MassSpawner/Public/MassEntityTemplateRegistry.h b/Engine/Plugins/Runtime/MassGameplay/Source/MassSpawner/Public/MassEntityTemplateRegistry.h
--- a/Engine/Plugins/Runtime/MassGameplay/Source/MassSpawner/Public/MassEntityTemplateRegistry.h
+++ b/Engine/Plugins/Runtime/MassGameplay/Source/MassSpawner/Public/MassEntityTemplateRegistry.h
@@ -4,9 +4,11 @@
 
 #if UE_ENABLE_INCLUDE_ORDER_DEPRECATED_IN_5_6
 #include "CoreMinimal.h"
-#include "MassCommonTypes.h"
 #endif // UE_ENABLE_INCLUDE_ORDER_DEPRECATED_IN_5_6
 #include "UObject/Object.h"
+#if UE_ENABLE_INCLUDE_ORDER_DEPRECATED_IN_5_6
+#include "MassCommonTypes.h"
+#endif // UE_ENABLE_INCLUDE_ORDER_DEPRECATED_IN_5_6
 #include "MassEntityTemplate.h"
 #include "MassTranslator.h"
 #include "MassEntityTemplateRegistry.generated.h"
diff --git a/Engine/Plugins/Runtime/MassGameplay/Source/MassSpawner/Public/MassEntityZoneGraphSpawnPointsGenerator.h b/Engine/Plugins/Runtime/MassGameplay/Source/MassSpawner/Public/MassEntityZoneGraphSpawnPointsGenerator.h
--- a/Engine/Plugins/Runtime/MassGameplay/Source/MassSpawner/Public/MassEntityZoneGraphSpawnPointsGenerator.h
+++ b/Engine/Plugins/Runtime/MassGameplay/Source/MassSpawner/Public/MassEntityZoneGraphSpawnPointsGenerator.h
@@ -2,7 +2,7 @@
 
 #pragma once 
 
-#if UE_ENABLE_INCLUDE_ORDER_DEPRECATED_IN_5_5
+#if UE_ENABLE_INCLUDE_ORDER_DEPRECATED_IN_5_6
 #include "CoreMinimal.h"
 #endif // UE_ENABLE_INCLUDE_ORDER_DEPRECATED_IN_5_6
 #include "MassEntitySpawnDataGeneratorBase.h"
diff --git a/Engine/Source/Developer/MassEntityTestSuite/Public/MassEntityTestFarmPlot.h b/Engine/Source/Developer/MassEntityTestSuite/Public/MassEntityTestFarmPlot.h
--- a/Engine/Source/Developer/MassEntityTestSuite/Public/MassEntityTestFarmPlot.h
+++ b/Engine/Source/Developer/MassEntityTestSuite/Public/MassEntityTestFarmPlot.h
@@ -2,6 +2,9 @@
 
 #pragma once
 
+#if UE_ENABLE_INCLUDE_ORDER_DEPRECATED_IN_5_6
+#include "CoreMinimal.h"
+#endif // UE_ENABLE_INCLUDE_ORDER_DEPRECATED_IN_5_6
 #include "GameFramework/Actor.h"
 #include "Engine/DataTable.h"
 #include "MassEntityTypes.h"
diff --git a/Engine/Source/Editor/MassEntityDebugger/Private/SMassProcessorsView.cpp b/Engine/Source/Editor/MassEntityDebugger/Private/SMassProcessorsView.cpp
--- a/Engine/Source/Editor/MassEntityDebugger/Private/SMassProcessorsView.cpp
+++ b/Engine/Source/Editor/MassEntityDebugger/Private/SMassProcessorsView.cpp
@@ -1,13 +1,9 @@
 // Copyright Epic Games, Inc. All Rights Reserved.
 
 #include "SMassProcessorsView.h"
-#include "MassDebugger.h"
-#include "SMassDebugger.h"
 #include "SMassProcessor.h"
 #include "MassDebuggerModel.h"
-#include "MassEntitySettings.h"
 #include "Widgets/Layout/SScrollBox.h"
-#include "Widgets/Text/SRichTextBlock.h"
 #include "Widgets/Views/SListView.h"
 #include "Widgets/SBoxPanel.h"
 
diff --git a/Engine/Source/Runtime/MassEntity/Internal/MassArchetypeData.h b/Engine/Source/Runtime/MassEntity/Internal/MassArchetypeData.h
--- a/Engine/Source/Runtime/MassEntity/Internal/MassArchetypeData.h
+++ b/Engine/Source/Runtime/MassEntity/Internal/MassArchetypeData.h
@@ -2,6 +2,7 @@
 
 #pragma once
 
+#include "MassEntityHandle.h"
 #include "MassArchetypeTypes.h"
 #if UE_ENABLE_INCLUDE_ORDER_DEPRECATED_IN_5_6
 #include "MassEntityManager.h"
diff --git a/Engine/Source/Runtime/MassEntity/Private/MassArchetypeTypes.cpp b/Engine/Source/Runtime/MassEntity/Private/MassArchetypeTypes.cpp
--- a/Engine/Source/Runtime/MassEntity/Private/MassArchetypeTypes.cpp
+++ b/Engine/Source/Runtime/MassEntity/Private/MassArchetypeTypes.cpp
@@ -3,7 +3,6 @@
 #include "MassArchetypeTypes.h"
 #include "MassEntityManager.h"
 #include "MassArchetypeData.h"
-#include "MassCommandBuffer.h"
 #include "MassEntityUtils.h"
 
 //-----------------------------------------------------------------------------
diff --git a/Engine/Source/Runtime/MassEntity/Private/MassEntityManagerStorage.cpp b/Engine/Source/Runtime/MassEntity/Private/MassEntityManagerStorage.cpp
--- a/Engine/Source/Runtime/MassEntity/Private/MassEntityManagerStorage.cpp
+++ b/Engine/Source/Runtime/MassEntity/Private/MassEntityManagerStorage.cpp
@@ -1,8 +1,8 @@
 // Copyright Epic Games, Inc. All Rights Reserved.
 
 #include "MassEntityManagerStorage.h"
-
 #include "MassEntityManagerConstants.h"
+#include "MassEntityHandle.h"
 #include "MassEntityTypes.h"
 #include "Templates/SharedPointer.h"
 
diff --git a/Engine/Source/Runtime/MassEntity/Public/MassArchetypeTypes.h b/Engine/Source/Runtime/MassEntity/Public/MassArchetypeTypes.h
--- a/Engine/Source/Runtime/MassEntity/Public/MassArchetypeTypes.h
+++ b/Engine/Source/Runtime/MassEntity/Public/MassArchetypeTypes.h
@@ -13,6 +13,7 @@ struct FMassArchetypeData;
 struct FMassExecutionContext;
 struct FMassFragment;
 struct FMassArchetypeChunkIterator;
+struct FMassEntityHandle;
 struct FMassEntityQuery;
 struct FMassArchetypeEntityCollection;
 struct FMassEntityView;
diff --git a/Engine/Source/Runtime/MassEntity/Public/MassEntityHandle.h b/Engine/Source/Runtime/MassEntity/Public/MassEntityHandle.h
new file mode 100644
--- /dev/null
+++ b/Engine/Source/Runtime/MassEntity/Public/MassEntityHandle.h
@@ -0,0 +1,80 @@
+// Copyright Epic Games, Inc. All Rights Reserved.
+
+#pragma once
+
+#include "MassEntityHandle.generated.h"
+
+/**
+ * A handle to a Mass entity. An entity is used in conjunction with the FMassEntityManager
+ * for the current world and can contain lightweight fragments.
+ */
+USTRUCT()
+struct alignas(8) FMassEntityHandle
+{
+	GENERATED_BODY()
+
+	FMassEntityHandle() = default;
+	FMassEntityHandle(const int32 InIndex, const int32 InSerialNumber)
+		: Index(InIndex), SerialNumber(InSerialNumber)
+	{
+	}
+	
+	UPROPERTY(VisibleAnywhere, Category = "Mass|Debug", Transient)
+	int32 Index = 0;
+	
+	UPROPERTY(VisibleAnywhere, Category = "Mass|Debug", Transient)
+	int32 SerialNumber = 0;
+
+	bool operator==(const FMassEntityHandle Other) const
+	{
+		return Index == Other.Index && SerialNumber == Other.SerialNumber;
+	}
+
+	bool operator!=(const FMassEntityHandle Other) const
+	{
+		return !operator==(Other);
+	}
+
+	/** Has meaning only for sorting purposes */
+	bool operator<(const FMassEntityHandle Other) const { return Index < Other.Index; }
+
+	/** Note that this function is merely checking if Index and SerialNumber are set. There's no way to validate if 
+	 *  these indicate a valid entity in an EntitySubsystem without asking the system. */
+	bool IsSet() const
+	{
+		return Index != 0 && SerialNumber != 0;
+	}
+
+	FORCEINLINE bool IsValid() const
+	{
+		return IsSet();
+	}
+
+	void Reset()
+	{
+		Index = SerialNumber = 0;
+	}
+
+	/** Allows the entity handle to be shared anonymously. */
+	uint64 AsNumber() const { return *reinterpret_cast<const uint64*>(this); } // Relying on the fact that this struct only stores 2 integers and is aligned correctly.
+	/** Reconstruct the entity handle from an anonymously shared integer. */
+	static FMassEntityHandle FromNumber(uint64 Value) 
+	{ 
+		FMassEntityHandle Result;
+		*reinterpret_cast<uint64_t*>(&Result) = Value;
+		return Result;
+	}
+
+	friend uint32 GetTypeHash(const FMassEntityHandle Entity)
+	{
+		return HashCombine(Entity.Index, Entity.SerialNumber);
+	}
+
+	FString DebugGetDescription() const
+	{
+		return FString::Printf(TEXT("i: %d sn: %d"), Index, SerialNumber);
+	}
+};
+
+static_assert(sizeof(FMassEntityHandle) == sizeof(uint64), "Expected FMassEntityHandle to be convertable to a 64-bit integer value, so size needs to be 8 bytes.");
+static_assert(alignof(FMassEntityHandle) == sizeof(uint64), "Expected FMassEntityHandle to be convertable to a 64-bit integer value, so alignment needs to be 8 bytes.");
diff --git a/Engine/Source/Runtime/MassEntity/Public/MassEntityTypes.h b/Engine/Source/Runtime/MassEntity/Public/MassEntityTypes.h
--- a/Engine/Source/Runtime/MassEntity/Public/MassEntityTypes.h
+++ b/Engine/Source/Runtime/MassEntity/Public/MassEntityTypes.h
@@ -6,6 +6,9 @@
 #include "MassProcessingTypes.h"
 #include "StructUtils/StructArrayView.h"
 #include "Subsystems/Subsystem.h"
+#if UE_ENABLE_INCLUDE_ORDER_DEPRECATED_IN_5_6
+#include "MassExternalSubsystemTraits.h"
+#endif // UE_ENABLE_INCLUDE_ORDER_DEPRECATED_IN_5_6
 #include "StructUtils/SharedStruct.h"
 #ifdef WITH_AITESTSUITE
 #include "TestableEnsures.h"
@@ -16,6 +19,7 @@
 #endif 
 #if UE_ENABLE_INCLUDE_ORDER_DEPRECATED_IN_5_6
 #include "MassExternalSubsystemTraits.h"
+#include "MassEntityHandle.h"
 #endif // UE_ENABLE_INCLUDE_ORDER_DEPRECATED_IN_5_6
 #include "MassEntityTypes.generated.h"
 
@@ -68,80 +72,6 @@ struct FMassConstSharedFragment
 	FMassConstSharedFragment() {}
 };
 
-
-// A handle to a lightweight entity.  An entity is used in conjunction with the FMassEntityManager
-// for the current world and can contain lightweight fragments.
-USTRUCT()
-struct alignas(8) FMassEntityHandle
-{
-	GENERATED_BODY()
-
-	FMassEntityHandle() = default;
-	FMassEntityHandle(const int32 InIndex, const int32 InSerialNumber)
-		: Index(InIndex), SerialNumber(InSerialNumber)
-	{
-	}
-	
-	UPROPERTY(VisibleAnywhere, Category = "Mass|Debug", Transient)
-	int32 Index = 0;
-	
-	UPROPERTY(VisibleAnywhere, Category = "Mass|Debug", Transient)
-	int32 SerialNumber = 0;
-
-	bool operator==(const FMassEntityHandle Other) const
-	{
-		return Index == Other.Index && SerialNumber == Other.SerialNumber;
-	}
-
-	bool operator!=(const FMassEntityHandle Other) const
-	{
-		return !operator==(Other);
-	}
-
-	/** Has meaning only for sorting purposes */
-	bool operator<(const FMassEntityHandle Other) const { return Index < Other.Index; }
-
-	/** Note that this function is merely checking if Index and SerialNumber are set. There's no way to validate if 
-	 *  these indicate a valid entity in an EntitySubsystem without asking the system. */
-	bool IsSet() const
-	{
-		return Index != 0 && SerialNumber != 0;
-	}
-
-	FORCEINLINE bool IsValid() const
-	{
-		return IsSet();
-	}
-
-	void Reset()
-	{
-		Index = SerialNumber = 0;
-	}
-
-	/** Allows the entity handle to be shared anonymously. */
-	uint64 AsNumber() const { return *reinterpret_cast<const uint64*>(this); } // Relying on the fact that this struct only stores 2 integers and is aligned correctly.
-	/** Reconstruct the entity handle from an anonymously shared integer. */
-	static FMassEntityHandle FromNumber(uint64 Value) 
-	{ 
-		FMassEntityHandle Result;
-		*reinterpret_cast<uint64_t*>(&Result) = Value;
-		return Result;
-	}
-
-	friend uint32 GetTypeHash(const FMassEntityHandle Entity)
-	{
-		return HashCombine(Entity.Index, Entity.SerialNumber);
-	}
-
-	FString DebugGetDescription() const
-	{
-		return FString::Printf(TEXT("i: %d sn: %d"), Index, SerialNumber);
-	}
-};
-
-static_assert(sizeof(FMassEntityHandle) == sizeof(uint64), "Expected FMassEntityHandle to be convertable to a 64-bit integer value, so size needs to be 8 bytes.");
-static_assert(alignof(FMassEntityHandle) == sizeof(uint64), "Expected FMassEntityHandle to be convertable to a 64-bit integer value, so alignment needs to be 8 bytes.");
-
 DECLARE_STRUCTTYPEBITSET_EXPORTED(MASSENTITY_API, FMassFragmentBitSet, FMassFragment);
 DECLARE_STRUCTTYPEBITSET_EXPORTED(MASSENTITY_API, FMassTagBitSet, FMassTag);
 DECLARE_STRUCTTYPEBITSET_EXPORTED(MASSENTITY_API, FMassChunkFragmentBitSet, FMassChunkFragment);
@@ -805,4 +735,3 @@ struct FMassArchetypeCreationParams
 	/** Name to identify the archetype while debugging*/
 	FName DebugName;
 };
-
diff --git a/Engine/Source/Runtime/MassEntity/Public/MassEntityView.h b/Engine/Source/Runtime/MassEntity/Public/MassEntityView.h
--- a/Engine/Source/Runtime/MassEntity/Public/MassEntityView.h
+++ b/Engine/Source/Runtime/MassEntity/Public/MassEntityView.h
@@ -2,6 +2,7 @@
 
 #pragma once
 
+#include "MassEntityHandle.h"
 #include "MassArchetypeTypes.h"
 #if UE_ENABLE_INCLUDE_ORDER_DEPRECATED_IN_5_6
 #include "StructUtils/InstancedStruct.h"
diff --git a/Engine/Source/Runtime/MassEntity/Public/MassExecutionContext.h b/Engine/Source/Runtime/MassEntity/Public/MassExecutionContext.h
--- a/Engine/Source/Runtime/MassEntity/Public/MassExecutionContext.h
+++ b/Engine/Source/Runtime/MassEntity/Public/MassExecutionContext.h
@@ -2,11 +2,14 @@
 
 #pragma once
 
+
+#include "MassEntityHandle.h"
 #include "MassEntityTypes.h"
 #if UE_ENABLE_INCLUDE_ORDER_DEPRECATED_IN_5_6
 #include "MassExternalSubsystemTraits.h"
-#endif // UE_ENABLE_INCLUDE_ORDER_DEPRECATED_IN_5_6
 #include "MassEntityQuery.h"
+#endif // UE_ENABLE_INCLUDE_ORDER_DEPRECATED_IN_5_6
+#include "MassArchetypeTypes.h"
 #include "MassSubsystemAccess.h"
 
 
@@ -19,6 +22,7 @@
 		, TEXT("Requested fragment type not bound for writing, type %s. Make sure it has been listed as required in ReadWrite mode.") \
 		, View ? *GetNameSafe(View->Requirement.StructType) : TEXT("[Not found]"))
 
+struct FMassEntityQuery;
 
 struct MASSENTITY_API FMassExecutionContext
 {
diff --git a/Engine/Source/Runtime/MassEntity/Public/MassObserverManager.h b/Engine/Source/Runtime/MassEntity/Public/MassObserverManager.h
--- a/Engine/Source/Runtime/MassEntity/Public/MassObserverManager.h
+++ b/Engine/Source/Runtime/MassEntity/Public/MassObserverManager.h
@@ -2,6 +2,8 @@
 
 #pragma once
 
+#include "MassEntityHandle.h"
+#include "MassProcessingTypes.h"
 #include "MassEntityTypes.h"
 #include "MassObserverManager.generated.h"
 
diff --git a/Engine/Source/Runtime/MassEntity/Public/MassProcessor.h b/Engine/Source/Runtime/MassEntity/Public/MassProcessor.h
--- a/Engine/Source/Runtime/MassEntity/Public/MassProcessor.h
+++ b/Engine/Source/Runtime/MassEntity/Public/MassProcessor.h
@@ -4,8 +4,8 @@
 
 #if UE_ENABLE_INCLUDE_ORDER_DEPRECATED_IN_5_6
 #include "CoreMinimal.h"
-#endif // UE_ENABLE_INCLUDE_ORDER_DEPRECATED_IN_5_6
 #include "MassEntityManager.h"
+#endif // UE_ENABLE_INCLUDE_ORDER_DEPRECATED_IN_5_6
 #include "MassProcessingTypes.h"
 #include "Async/TaskGraphInterfaces.h"
 #include "MassCommandBuffer.h"
@@ -16,6 +16,7 @@ struct FMassProcessingPhaseConfig;
 class UMassCompositeProcessor;
 struct FMassDebugger;
 struct FMassEntityQuery;
+struct FMassEntityManager;
 
 enum class EProcessorCompletionStatus : uint8
 {
