[Mass] fix to GC-crashes when switching levels due to the pruned processors still being added by FMassProcessingPhaseManager::AddReferencedObjects, even after FMassProcessingPhaseManager::Deinitialize was called. Added resetting of the pruned processors to FMassProcessingPhaseManager::Deinitialize.

#rnx

diff --git a/Engine/Source/Runtime/MassEntity/Private/MassProcessingPhaseManager.cpp b/Engine/Source/Runtime/MassEntity/Private/MassProcessingPhaseManager.cpp
--- a/Engine/Source/Runtime/MassEntity/Private/MassProcessingPhaseManager.cpp
+++ b/Engine/Source/Runtime/MassEntity/Private/MassProcessingPhaseManager.cpp
@@ -282,6 +282,11 @@ void FMassProcessingPhaseManager::Deinitialize()
 	}
 
 	DynamicProcessors.Reset();
+
+	for (FPhaseGraphBuildState& GraphBuildState : ProcessingGraphBuildStates)
+	{
+		GraphBuildState.LastResult.PrunedProcessors.Reset();
+	}
 }
 
 const FGraphEventRef& FMassProcessingPhaseManager::TriggerPhase(const EMassProcessingPhase Phase, const float DeltaTime, const FGraphEventRef& MyCompletionGraphEvent)
