[Mass] Issue 640908: Compile errors in MassDebuggerModel.h

#rnx

diff --git a/Engine/Source/Editor/MassEntityDebugger/Private/MassDebuggerModel.h b/Engine/Source/Editor/MassEntityDebugger/Private/MassDebuggerModel.h
--- a/Engine/Source/Editor/MassEntityDebugger/Private/MassDebuggerModel.h
+++ b/Engine/Source/Editor/MassEntityDebugger/Private/MassDebuggerModel.h
@@ -150,7 +150,7 @@ struct FMassDebuggerEnvironment
 	bool NeedsValidWorld() const { return bNeedsValidWorld; }
 	
 	TWeakPtr<const FMassEntityManager> EntityManager;
-	TMap<FName, FMassDebugger::FProcessorProviderFunction> ProcessorProviders;
+	TMap<FName, UE::Mass::Debug::FProcessorProviderFunction> ProcessorProviders;
 	TWeakObjectPtr<UWorld> World;
 	bool bNeedsValidWorld = false;
 };
diff --git a/Engine/Source/Runtime/MassEntity/Private/MassDebugger.cpp b/Engine/Source/Runtime/MassEntity/Private/MassDebugger.cpp
--- a/Engine/Source/Runtime/MassEntity/Private/MassDebugger.cpp
+++ b/Engine/Source/Runtime/MassEntity/Private/MassDebugger.cpp
@@ -719,7 +719,7 @@ void FMassDebugger::UnregisterEntityManager(FMassEntityManager& EntityManager)
 	OnEntityManagerDeinitialized.Broadcast(EntityManager);
 }
 
-void FMassDebugger::RegisterProcessorDataProvider(FName ProviderName, const TSharedRef<FMassEntityManager>& EntityManager, const FProcessorProviderFunction& ProviderFunction)
+void FMassDebugger::RegisterProcessorDataProvider(FName ProviderName, const TSharedRef<FMassEntityManager>& EntityManager, const UE::Mass::Debug::FProcessorProviderFunction& ProviderFunction)
 {
 	UE::TScopeLock<UE::FSpinLock> ScopeLock(EntityManagerRegistrationLock);
 	int32 Index = ActiveEnvironments.IndexOfByPredicate([WeakEntityManager = EntityManager->AsWeak()](const FEnvironment& Element) 
diff --git a/Engine/Source/Runtime/MassEntity/Public/MassDebugger.h b/Engine/Source/Runtime/MassEntity/Public/MassDebugger.h
--- a/Engine/Source/Runtime/MassEntity/Public/MassDebugger.h
+++ b/Engine/Source/Runtime/MassEntity/Public/MassDebugger.h
@@ -42,6 +42,8 @@ namespace UE::Mass::Debug
 		/** Total amount of memory needed by a single entity */
 		SIZE_T BytesPerEntity = 0;
 	};
+
+	using FProcessorProviderFunction = TFunction<void(TArray<const UMassProcessor*>&)>;
 } // namespace UE::Mass::Debug
 
 USTRUCT()
@@ -132,12 +134,10 @@ namespace UE::Mass::Debug
 
 struct MASSENTITY_API FMassDebugger
 {
-	using FProcessorProviderFunction = TFunction<void(TArray<const UMassProcessor*>&)>;
-
 	struct FEnvironment
 	{
 		TWeakPtr<const FMassEntityManager> EntityManager;
-		TMap<FName, FProcessorProviderFunction> ProcessorProviders;
+		TMap<FName, UE::Mass::Debug::FProcessorProviderFunction> ProcessorProviders;
 		FMassEntityHandle SelectedEntity;
 
 		explicit FEnvironment(const FMassEntityManager& InEntityManager)
@@ -214,7 +214,7 @@ struct MASSENTITY_API FMassDebugger
 	 * NOTE: there's no UnregisterProcessorDataProvider, the registered providers will automatically get removed along with
 	 * the rest of the data associated with the relevant EntityManager as part of UnregisterEntityManager call.
 	 */
-	static void RegisterProcessorDataProvider(FName ProviderName, const TSharedRef<FMassEntityManager>& EntityManager, const FProcessorProviderFunction& ProviderFunction);
+	static void RegisterProcessorDataProvider(FName ProviderName, const TSharedRef<FMassEntityManager>& EntityManager, const UE::Mass::Debug::FProcessorProviderFunction& ProviderFunction);
 
 	static TConstArrayView<FEnvironment> GetEnvironments() { return ActiveEnvironments; }
 
