[Mass] added a way to batch-add processors to a FMassRuntimePipeline, and used that capability in FMassPhaseProcessorConfigurationHelper::Configure

diff --git a/Engine/Source/Runtime/MassEntity/Private/MassProcessingPhaseManager.cpp b/Engine/Source/Runtime/MassEntity/Private/MassProcessingPhaseManager.cpp
--- a/Engine/Source/Runtime/MassEntity/Private/MassProcessingPhaseManager.cpp
+++ b/Engine/Source/Runtime/MassEntity/Private/MassProcessingPhaseManager.cpp
@@ -153,16 +153,8 @@ void FMassPhaseProcessorConfigurationHelper::Configure(TArrayView<UMassProcessor
 	{
 		SCOPE_CYCLE_COUNTER(STAT_Mass_PhaseConfigurePipelineCreation);
 
-		for (UMassProcessor* PrunedProcessor : InOutOptionalResult.PrunedProcessors)
-		{
-			if (ensure(PrunedProcessor))
-			{
-				// @todo consider adding all PrunedProcessors in one go - there should be no overlap
-				// with PhaseProcessor.GetChildProcessorsView(), not there should be any nulls.
-				TmpPipeline.AppendProcessor(*PrunedProcessor);
-			}
-		}
-
+		TmpPipeline.AppendProcessors(InOutOptionalResult.PrunedProcessors);
+		
 		// some previously added dynamic processors were either in the active processor group,
 		// or were among the pruned processors. At this point we have both groups in TmpPipeline now
 		// so we need to check if any of these processors have been removed since las processing
diff --git a/Engine/Source/Runtime/MassEntity/Private/MassProcessingTypes.cpp b/Engine/Source/Runtime/MassEntity/Private/MassProcessingTypes.cpp
--- a/Engine/Source/Runtime/MassEntity/Private/MassProcessingTypes.cpp
+++ b/Engine/Source/Runtime/MassEntity/Private/MassProcessingTypes.cpp
@@ -232,9 +232,14 @@ void FMassRuntimePipeline::AppendOrOverrideRuntimeProcessorCopies(TConstArrayVie
 	}
 }
 
-void FMassRuntimePipeline::AppendProcessor(UMassProcessor& Processor)
+void FMassRuntimePipeline::AppendProcessor(UMassProcessor& InProcessor)
 {
-	Processors.Add(&Processor);
+	Processors.Add(&InProcessor);
+}
+
+void FMassRuntimePipeline::AppendProcessors(TArrayView<TObjectPtr<UMassProcessor>> InProcessors)
+{
+	Processors.Append(InProcessors);
 }
 
 void FMassRuntimePipeline::AppendUniqueProcessor(UMassProcessor& Processor)
diff --git a/Engine/Source/Runtime/MassEntity/Public/MassProcessingTypes.h b/Engine/Source/Runtime/MassEntity/Public/MassProcessingTypes.h
--- a/Engine/Source/Runtime/MassEntity/Public/MassProcessingTypes.h
+++ b/Engine/Source/Runtime/MassEntity/Public/MassProcessingTypes.h
@@ -117,8 +117,10 @@ public:
 	 */
 	void AppendUniqueRuntimeProcessorCopies(TConstArrayView<const UMassProcessor*> InProcessors, UObject& InOwner, const TSharedRef<FMassEntityManager>& EntityManager);
 
-	/** Adds InProcessor to Processors without any additional checks */
-	void AppendProcessor(UMassProcessor& Processor);
+	/** Adds InProcessor(s) to Processors without any additional checks */
+	void AppendProcessor(UMassProcessor& InProcessor);
+	void AppendProcessors(TArrayView<TObjectPtr<UMassProcessor>> InProcessors);
+
 	void AppendUniqueProcessor(UMassProcessor& Processor);
 
 	/** Creates an instance of ProcessorClass and adds it to Processors without any additional checks */
