PR #12449: Allow Mass phases to run outside the game thread

diff --git a/Engine/Source/Runtime/MassEntity/Private/MassExecutor.cpp b/Engine/Source/Runtime/MassEntity/Private/MassExecutor.cpp
--- a/Engine/Source/Runtime/MassEntity/Private/MassExecutor.cpp
+++ b/Engine/Source/Runtime/MassEntity/Private/MassExecutor.cpp
@@ -158,10 +158,11 @@ void RunProcessorsView(TArrayView<UMassProcessor* const> Processors, FMassProces
 
 struct FMassExecutorDoneTask
 {
-	FMassExecutorDoneTask(const FMassExecutionContext& InExecutionContext, TFunction<void()> InOnDoneNotification, const FString& InDebugName)
+	FMassExecutorDoneTask(const FMassExecutionContext& InExecutionContext, TFunction<void()> InOnDoneNotification, const FString& InDebugName, ENamedThreads::Type InDesiredThread)
 		: ExecutionContext(InExecutionContext)
 		, OnDoneNotification(InOnDoneNotification)
 		, DebugName(InDebugName)
+		, DesiredThread(InDesiredThread)
 	{
 	}
 	static TStatId GetStatId()
@@ -169,7 +170,11 @@ struct FMassExecutorDoneTask
 		RETURN_QUICK_DECLARE_CYCLE_STAT(FMassExecutorDoneTask, STATGROUP_TaskGraphTasks);
 	}
 
-	static ENamedThreads::Type GetDesiredThread() { return ENamedThreads::GameThread; }
+	ENamedThreads::Type GetDesiredThread()
+	{
+		return DesiredThread;
+	}
+
 	static ESubsequentsMode::Type GetSubsequentsMode() { return ESubsequentsMode::TrackSubsequents; }
 	void DoTask(ENamedThreads::Type CurrentThread, const FGraphEventRef& MyCompletionGraphEvent)
 	{
@@ -193,9 +198,10 @@ private:
 	FMassExecutionContext ExecutionContext;
 	TFunction<void()> OnDoneNotification;
 	FString DebugName;
+	ENamedThreads::Type DesiredThread;
 };
 
-FGraphEventRef TriggerParallelTasks(UMassProcessor& Processor, FMassProcessingContext& ProcessingContext, TFunction<void()> OnDoneNotification)
+FGraphEventRef TriggerParallelTasks(UMassProcessor& Processor, FMassProcessingContext& ProcessingContext, TFunction<void()> OnDoneNotification, ENamedThreads::Type CurrentThread)
 {
 	if (!ProcessingContext.EntityManager)
 	{
@@ -224,7 +230,7 @@ FGraphEventRef TriggerParallelTasks(UMassProcessor& Processor, FMassProcessingCo
 	{
 		const FGraphEventArray Prerequisites = { CompletionEvent };
 		CompletionEvent = TGraphTask<FMassExecutorDoneTask>::CreateTask(&Prerequisites)
-			.ConstructAndDispatchWhenReady(ExecutionContext, OnDoneNotification, Processor.GetName());
+			.ConstructAndDispatchWhenReady(ExecutionContext, OnDoneNotification, Processor.GetName(), CurrentThread);
 	}
 
 	return CompletionEvent;
diff --git a/Engine/Source/Runtime/MassEntity/Private/MassProcessingPhaseManager.cpp b/Engine/Source/Runtime/MassEntity/Private/MassProcessingPhaseManager.cpp
--- a/Engine/Source/Runtime/MassEntity/Private/MassProcessingPhaseManager.cpp
+++ b/Engine/Source/Runtime/MassEntity/Private/MassProcessingPhaseManager.cpp
@@ -87,7 +87,8 @@ void FMassProcessingPhase::ExecuteTick(float DeltaTime, ELevelTick TickType, ENa
 			const FGraphEventRef PipelineCompletionEvent = UE::Mass::Executor::TriggerParallelTasks(*PhaseProcessor, Context, [this, DeltaTime]()
 				{
 					OnParallelExecutionDone(DeltaTime);
-				});
+				}
+				, CurrentThread);
 
 			if (PipelineCompletionEvent.IsValid())
 			{
@@ -292,13 +293,14 @@ void FMassProcessingPhaseManager::Deinitialize()
 	}
 }
 
-const FGraphEventRef& FMassProcessingPhaseManager::TriggerPhase(const EMassProcessingPhase Phase, const float DeltaTime, const FGraphEventRef& MyCompletionGraphEvent)
+const FGraphEventRef& FMassProcessingPhaseManager::TriggerPhase(const EMassProcessingPhase Phase, const float DeltaTime
+	, const FGraphEventRef& MyCompletionGraphEvent, ENamedThreads::Type CurrentThread)
 {
 	check(Phase != EMassProcessingPhase::MAX);
 
 	if (bIsAllowedToTick)
 	{
-		ProcessingPhases[(int)Phase].ExecuteTick(DeltaTime, LEVELTICK_All, ENamedThreads::GameThread, MyCompletionGraphEvent);
+		ProcessingPhases[(int)Phase].ExecuteTick(DeltaTime, LEVELTICK_All, CurrentThread, MyCompletionGraphEvent);
 	}
 
 	return MyCompletionGraphEvent;
diff --git a/Engine/Source/Runtime/MassEntity/Public/MassExecutor.h b/Engine/Source/Runtime/MassEntity/Public/MassExecutor.h
--- a/Engine/Source/Runtime/MassEntity/Public/MassExecutor.h
+++ b/Engine/Source/Runtime/MassEntity/Public/MassExecutor.h
@@ -43,7 +43,8 @@ namespace UE::Mass::Executor
 	 *  @param OnDoneNotification will be called after all the processors are done, just after flushing the command buffer.
 	 *    Note that OnDoneNotification will be executed on GameThread.
 	 */
-	MASSENTITY_API FGraphEventRef TriggerParallelTasks(UMassProcessor& Processor, FMassProcessingContext& ProcessingContext, TFunction<void()> OnDoneNotification);
+	MASSENTITY_API FGraphEventRef TriggerParallelTasks(UMassProcessor& Processor, FMassProcessingContext& ProcessingContext, TFunction<void()> OnDoneNotification
+		, ENamedThreads::Type CurrentThread = ENamedThreads::GameThread);
 
 	UE_DEPRECATED(5.5, "This flavor of RunProcessorsView is deprecated. Use the one with TConstArrayView<FMassArchetypeEntityCollection> parameter instead.")
 	MASSENTITY_API void RunProcessorsView(TArrayView<UMassProcessor* const> Processors, FMassProcessingContext& ProcessingContext, const FMassArchetypeEntityCollection* EntityCollection);
diff --git a/Engine/Source/Runtime/MassEntity/Public/MassProcessingPhaseManager.h b/Engine/Source/Runtime/MassEntity/Public/MassProcessingPhaseManager.h
--- a/Engine/Source/Runtime/MassEntity/Public/MassProcessingPhaseManager.h
+++ b/Engine/Source/Runtime/MassEntity/Public/MassProcessingPhaseManager.h
@@ -161,7 +161,8 @@ public:
 	/** Needs to be called before destruction, ideally before owner's BeginDestroy (a FGCObject's limitation) */
 	void Deinitialize();
 
-	const FGraphEventRef& TriggerPhase(const EMassProcessingPhase Phase, const float DeltaTime, const FGraphEventRef& MyCompletionGraphEvent);
+	const FGraphEventRef& TriggerPhase(const EMassProcessingPhase Phase, const float DeltaTime, const FGraphEventRef& MyCompletionGraphEvent
+		, ENamedThreads::Type CurrentThread = ENamedThreads::GameThread);
 
 	/** 
 	 *  Stores EntityManager associated with given world's MassEntitySubsystem and kicks off phase ticking.
