[Mass] Reporting when a processor is being asked for supported archetypes while not having any queries registered

#rnx

diff --git a/Engine/Source/Runtime/MassEntity/Private/MassProcessor.cpp b/Engine/Source/Runtime/MassEntity/Private/MassProcessor.cpp
--- a/Engine/Source/Runtime/MassEntity/Private/MassProcessor.cpp
+++ b/Engine/Source/Runtime/MassEntity/Private/MassProcessor.cpp
@@ -173,6 +173,8 @@ void UMassProcessor::SetShouldAutoRegisterWithGlobalList(const bool bAutoRegiste
 
 void UMassProcessor::GetArchetypesMatchingOwnedQueries(const FMassEntityManager& EntityManager, TArray<FMassArchetypeHandle>& OutArchetype)
 {
+	UE_CLOG(OwnedQueries.Num() == 0, LogMass, Warning, TEXT("%s has no registered queries while being asked for matching archetypes"), *GetName());
+
 	for (FMassEntityQuery* QueryPtr : OwnedQueries)
 	{
 		CA_ASSUME(QueryPtr);
diff --git a/Engine/Source/Runtime/MassEntity/Private/MassProcessorDependencySolver.cpp b/Engine/Source/Runtime/MassEntity/Private/MassProcessorDependencySolver.cpp
--- a/Engine/Source/Runtime/MassEntity/Private/MassProcessorDependencySolver.cpp
+++ b/Engine/Source/Runtime/MassEntity/Private/MassProcessorDependencySolver.cpp
@@ -1129,12 +1129,23 @@ void FMassProcessorDependencySolver::ResolveDependencies(TArray<FMassProcessorOr
 		{
 			CA_ASSUME(Node.Processor); // as implied by Node.IsGroup() == false
 
-			// for each processor-representing node we cache information on which archetypes among the once we've created 
-			// above (see the EntityManager.CreateArchetype call in the previous loop) match this processor. 
-			Node.Processor->GetArchetypesMatchingOwnedQueries(*EntityManager.Get(), Node.ValidArchetypes);
+			const bool bDoQueryBasedPruning = Node.Processor->ShouldAllowQueryBasedPruning(bGameRuntime);
+
+			// we gather archetypes for processors that have queries OR allow query-based pruning.
+			// The main point of this condition is to allow calling GetArchetypesMatchingOwnedQueries
+			// on pruning-supporting processors, while having no queries - that will emmit a warning
+			// that will let the user know their processor is misconfigured.
+			// We do collect archetype information for the processors that never get pruned because we're
+			// using this information for the dependency calculations, regardless of ShouldAllowQueryBasedPruning
+			if (bDoQueryBasedPruning || Node.Processor->GetOwnedQueriesNum())
+			{
+				// for each processor-representing node we cache information on which archetypes among the once we've created 
+				// above (see the EntityManager.CreateArchetype call in the previous loop) match this processor. 
+				Node.Processor->GetArchetypesMatchingOwnedQueries(*EntityManager.Get(), Node.ValidArchetypes);
+			}
 
 			// prune the archetype-less processors
-			if (Node.ValidArchetypes.Num() == 0 && Node.Processor->ShouldAllowQueryBasedPruning(bGameRuntime))
+			if (Node.ValidArchetypes.Num() == 0 && bDoQueryBasedPruning)
 			{
 				UE_LOG(LogMass, Verbose, TEXT("\t%s"), *Node.Processor->GetName());
 
diff --git a/Engine/Source/Runtime/MassEntity/Public/MassProcessor.h b/Engine/Source/Runtime/MassEntity/Public/MassProcessor.h
--- a/Engine/Source/Runtime/MassEntity/Public/MassProcessor.h
+++ b/Engine/Source/Runtime/MassEntity/Public/MassProcessor.h
@@ -142,6 +142,7 @@ public:
 
 	void GetArchetypesMatchingOwnedQueries(const FMassEntityManager& EntityManager, TArray<FMassArchetypeHandle>& OutArchetype);
 	bool DoesAnyArchetypeMatchOwnedQueries(const FMassEntityManager& EntityManager);
+	int32 GetOwnedQueriesNum() const;
 	
 #if CPUPROFILERTRACE_ENABLED
 	FString StatId;
@@ -375,6 +376,11 @@ inline bool UMassProcessor::ShouldAllowMultipleInstances() const
 	return bAllowMultipleInstances;
 }
 
+inline int32 UMassProcessor::GetOwnedQueriesNum() const
+{
+	return OwnedQueries.Num();
+}
+
 inline void UMassProcessor::DebugOutputDescription(FOutputDevice& Ar) const
 {
 	DebugOutputDescription(Ar, 0);
