[Mass] Tiny optimization to FMassPhaseProcessorConfigurationHelper::Configure to avoid iterating through all InOutRemovedDynamicProcessors when the "current" processing graph contains no elements (so there's no need to check if a given removed processor is part of the existing processing graph).

#rnx

diff --git a/Engine/Source/Runtime/MassEntity/Private/MassProcessingPhaseManager.cpp b/Engine/Source/Runtime/MassEntity/Private/MassProcessingPhaseManager.cpp
--- a/Engine/Source/Runtime/MassEntity/Private/MassProcessingPhaseManager.cpp
+++ b/Engine/Source/Runtime/MassEntity/Private/MassProcessingPhaseManager.cpp
@@ -155,21 +155,24 @@ void FMassPhaseProcessorConfigurationHelper::Configure(TArrayView<UMassProcessor
 		SCOPE_CYCLE_COUNTER(STAT_Mass_PhaseConfigurePipelineCreation);
 
 		TmpPipeline.AppendProcessors(InOutOptionalResult.PrunedProcessors);
-		
-		// some previously added dynamic processors were either in the active processor group,
-		// or were among the pruned processors. At this point we have both groups in TmpPipeline now
-		// so we need to check if any of these processors have been removed since las processing
-		// graph recreation
-		for (int32 Index = InOutRemovedDynamicProcessors.Num() - 1; Index >= 0; --Index)
+
+		if (TmpPipeline.Num())
 		{
-			if (const UMassProcessor* RemovedProcessor = InOutRemovedDynamicProcessors[Index].Get())
+			// some previously added dynamic processors were either in the active processor group,
+			// or were among the pruned processors. At this point we have both groups in TmpPipeline now
+			// so we need to check if any of these processors have been removed since las processing
+			// graph recreation
+			for (int32 Index = InOutRemovedDynamicProcessors.Num() - 1; Index >= 0; --Index)
 			{
-				if (TmpPipeline.RemoveProcessor(*RemovedProcessor) == false)
+				if (const UMassProcessor* RemovedProcessor = InOutRemovedDynamicProcessors[Index].Get())
 				{
-					continue;
+					if (TmpPipeline.RemoveProcessor(*RemovedProcessor) == false)
+					{
+						continue;
+					}
 				}
+				InOutRemovedDynamicProcessors.RemoveAtSwap(Index, 1, EAllowShrinking::No);
 			}
-			InOutRemovedDynamicProcessors.RemoveAtSwap(Index, 1, EAllowShrinking::No);
 		}
 
 		for (UMassProcessor* Processor : DynamicProcessors)
