[MassEntity] Minor change to MassArchetypeData's API (private to the MassEntity module)

#rnx

diff --git a/Engine/Source/Runtime/MassEntity/Internal/MassArchetypeData.h b/Engine/Source/Runtime/MassEntity/Internal/MassArchetypeData.h
--- a/Engine/Source/Runtime/MassEntity/Internal/MassArchetypeData.h
+++ b/Engine/Source/Runtime/MassEntity/Internal/MassArchetypeData.h
@@ -225,13 +225,13 @@ public:
 		return CompositionDescriptor.IsEquivalent(OtherCompositionDescriptor);
 	}
 
-	void Initialize(const FMassArchetypeCompositionDescriptor& InCompositionDescriptor, const uint32 ArchetypeDataVersion);
+	void Initialize(const FMassEntityManager& EntityManager, const FMassArchetypeCompositionDescriptor& InCompositionDescriptor, const uint32 ArchetypeDataVersion);
 
 	/** 
 	 * A special way of initializing an archetype resulting in a copy of BaseArchetype's setup with OverrideTags
 	 * replacing original tags of BaseArchetype
 	 */
-	void InitializeWithSimilar(const FMassArchetypeData& BaseArchetype, FMassArchetypeCompositionDescriptor&& NewComposition, const uint32 ArchetypeDataVersion);
+	void InitializeWithSimilar(const FMassEntityManager& EntityManager, const FMassArchetypeData& BaseArchetype, FMassArchetypeCompositionDescriptor&& NewComposition, const uint32 ArchetypeDataVersion);
 
 	void AddEntity(FMassEntityHandle Entity, const FMassArchetypeSharedFragmentValues& InSharedFragmentValues);
 	void RemoveEntity(FMassEntityHandle Entity);
@@ -392,7 +392,7 @@ protected:
 	};
 	void MoveFragmentsToAnotherArchetypeInternal(FMassArchetypeData& TargetArchetype, FTransientChunkLocation Target, const FTransientChunkLocation Source, const int32 ElementsNum);
 	void MoveFragmentsToNewLocationInternal(FTransientChunkLocation Target, const FTransientChunkLocation Source, const int32 NumberToMove);
-	void ConfigureFragments();
+	void ConfigureFragments(const FMassEntityManager& EntityManager);
 
 	FORCEINLINE void* GetFragmentData(const int32 FragmentIndex, uint8* ChunkRawMemory, const int32 IndexWithinChunk) const
 	{
diff --git a/Engine/Source/Runtime/MassEntity/Private/MassArchetypeData.cpp b/Engine/Source/Runtime/MassEntity/Private/MassArchetypeData.cpp
--- a/Engine/Source/Runtime/MassEntity/Private/MassArchetypeData.cpp
+++ b/Engine/Source/Runtime/MassEntity/Private/MassArchetypeData.cpp
@@ -55,7 +55,7 @@ bool FMassArchetypeData::HasFragmentType(const UScriptStruct* FragmentType) cons
 	return (FragmentType && CompositionDescriptor.Fragments.Contains(*FragmentType));
 }
 
-void FMassArchetypeData::Initialize(const FMassArchetypeCompositionDescriptor& InCompositionDescriptor, const uint32 ArchetypeDataVersion)
+void FMassArchetypeData::Initialize(const FMassEntityManager& EntityManager, const FMassArchetypeCompositionDescriptor& InCompositionDescriptor, const uint32 ArchetypeDataVersion)
 {
 	if (!ensureMsgf(Chunks.Num() == 0, TEXT("Trying to re-initialize non-empty Mass Archetype is not supported")))
 	{
@@ -68,7 +68,7 @@ void FMassArchetypeData::Initialize(const FMassArchetypeCompositionDescriptor& I
 
 	CreatedArchetypeDataVersion = ArchetypeDataVersion;
 	CompositionDescriptor.Fragments = InCompositionDescriptor.Fragments;
-	ConfigureFragments();
+	ConfigureFragments(EntityManager);
 
 	// Tags
 	CompositionDescriptor.Tags = InCompositionDescriptor.Tags;
@@ -91,7 +91,7 @@ void FMassArchetypeData::Initialize(const FMassArchetypeCompositionDescriptor& I
 	EntityListOffsetWithinChunk = 0;
 }
 
-void FMassArchetypeData::InitializeWithSimilar(const FMassArchetypeData& BaseArchetype, FMassArchetypeCompositionDescriptor&& NewComposition, const uint32 ArchetypeDataVersion)
+void FMassArchetypeData::InitializeWithSimilar(const FMassEntityManager& EntityManager, const FMassArchetypeData& BaseArchetype, FMassArchetypeCompositionDescriptor&& NewComposition, const uint32 ArchetypeDataVersion)
 {
 	checkf(IsInitialized() == false, TEXT("Trying to InitializeWithSimilar but this archetype has already been initialized"));
 
@@ -101,7 +101,7 @@ void FMassArchetypeData::InitializeWithSimilar(const FMassArchetypeData& BaseArc
 	CompositionDescriptor = MoveTemp(NewComposition);
 	if (CompositionDescriptor.Fragments != BaseArchetype.GetCompositionDescriptor().Fragments)
 	{
-		ConfigureFragments();
+		ConfigureFragments(EntityManager);
 	}
 	else
 	{
@@ -115,7 +115,7 @@ void FMassArchetypeData::InitializeWithSimilar(const FMassArchetypeData& BaseArc
 	EntityListOffsetWithinChunk = 0;
 }
 
-void FMassArchetypeData::ConfigureFragments()
+void FMassArchetypeData::ConfigureFragments(const FMassEntityManager& EntityManager)
 {
 	TArray<const UScriptStruct*, TInlineAllocator<16>> SortedFragmentList;
 	CompositionDescriptor.Fragments.ExportTypes(SortedFragmentList);
diff --git a/Engine/Source/Runtime/MassEntity/Private/MassEntityManager.cpp b/Engine/Source/Runtime/MassEntity/Private/MassEntityManager.cpp
--- a/Engine/Source/Runtime/MassEntity/Private/MassEntityManager.cpp
+++ b/Engine/Source/Runtime/MassEntity/Private/MassEntityManager.cpp
@@ -450,7 +450,7 @@ FMassArchetypeHandle FMassEntityManager::CreateArchetype(const FMassArchetypeCom
 
 		// Create a new archetype
 		FMassArchetypeData* NewArchetype = new FMassArchetypeData(CreationParams);
-		NewArchetype->Initialize(Composition, ArchetypeDataVersion);
+		NewArchetype->Initialize(*this, Composition, ArchetypeDataVersion);
 		ArchetypeDataPtr = HashRow.Add_GetRef(MakeShareable(NewArchetype));
 		AllArchetypes.Add(ArchetypeDataPtr);
 		ensure(AllArchetypes.Num() == ArchetypeDataVersion);
@@ -514,7 +514,7 @@ FMassArchetypeHandle FMassEntityManager::InternalCreateSimilarArchetype(const FM
 
 		// Create a new archetype
 		FMassArchetypeData* NewArchetype = new FMassArchetypeData(FMassArchetypeCreationParams(SourceArchetypeRef));
-		NewArchetype->InitializeWithSimilar(SourceArchetypeRef, MoveTemp(NewComposition), ArchetypeDataVersion);
+		NewArchetype->InitializeWithSimilar(*this, SourceArchetypeRef, MoveTemp(NewComposition), ArchetypeDataVersion);
 		NewArchetype->CopyDebugNamesFrom(SourceArchetypeRef);
 
 		ArchetypeDataPtr = HashRow.Add_GetRef(MakeShareable(NewArchetype));
