[Mass] Extracted Mass "element" types (things that we build archetype composition with) into a dedicated header, MassEntityElementTypes.h

#rnx

diff --git a/Engine/Source/Runtime/MassEntity/Public/MassEntityElementTypes.h b/Engine/Source/Runtime/MassEntity/Public/MassEntityElementTypes.h
new file mode 100644
--- /dev/null
+++ b/Engine/Source/Runtime/MassEntity/Public/MassEntityElementTypes.h
@@ -0,0 +1,44 @@
+// Copyright Epic Games, Inc. All Rights Reserved.
+
+#pragma once
+
+#include "MassEntityElementTypes.generated.h"
+
+// This is the base class for all lightweight fragments
+USTRUCT()
+struct FMassFragment
+{
+	GENERATED_BODY()
+};
+
+// This is the base class for types that will only be tested for presence/absence, i.e. Tags.
+// Subclasses should never contain any member properties.
+USTRUCT()
+struct FMassTag
+{
+	GENERATED_BODY()
+};
+
+USTRUCT()
+struct FMassChunkFragment
+{
+	GENERATED_BODY()
+};
+
+USTRUCT()
+struct FMassSharedFragment
+{
+	GENERATED_BODY()
+};
+
+USTRUCT()
+struct FMassConstSharedFragment
+{
+	GENERATED_BODY()
+};
+
+USTRUCT()
+struct FMassRelationship
+{
+	GENERATED_BODY()
+};
diff --git a/Engine/Source/Runtime/MassEntity/Public/MassEntityTypes.h b/Engine/Source/Runtime/MassEntity/Public/MassEntityTypes.h
--- a/Engine/Source/Runtime/MassEntity/Public/MassEntityTypes.h
+++ b/Engine/Source/Runtime/MassEntity/Public/MassEntityTypes.h
@@ -10,6 +10,7 @@
 #include "MassExternalSubsystemTraits.h"
 #endif // UE_ENABLE_INCLUDE_ORDER_DEPRECATED_IN_5_6
 #include "StructUtils/SharedStruct.h"
+#include "MassEntityElementTypes.h"
 #ifdef WITH_AITESTSUITE
 #include "TestableEnsures.h"
 #else
@@ -29,49 +30,6 @@ MASSENTITY_API DECLARE_LOG_CATEGORY_EXTERN(LogMass, Warning, All);
 DECLARE_STATS_GROUP(TEXT("Mass"), STATGROUP_Mass, STATCAT_Advanced);
 DECLARE_CYCLE_STAT_EXTERN(TEXT("Mass Total Frame Time"), STAT_Mass_Total, STATGROUP_Mass, MASSENTITY_API);
 
-// This is the base class for all lightweight fragments
-USTRUCT()
-struct FMassFragment
-{
-	GENERATED_BODY()
-
-	FMassFragment() {}
-};
-
-// This is the base class for types that will only be tested for presence/absence, i.e. Tags.
-// Subclasses should never contain any member properties.
-USTRUCT()
-struct FMassTag
-{
-	GENERATED_BODY()
-
-	FMassTag() {}
-};
-
-USTRUCT()
-struct FMassChunkFragment
-{
-	GENERATED_BODY()
-
-	FMassChunkFragment() {}
-};
-
-USTRUCT()
-struct FMassSharedFragment
-{
-	GENERATED_BODY()
-
-	FMassSharedFragment() {}
-};
-
-USTRUCT()
-struct FMassConstSharedFragment
-{
-	GENERATED_BODY()
-
-	FMassConstSharedFragment() {}
-};
-
 DECLARE_STRUCTTYPEBITSET_EXPORTED(MASSENTITY_API, FMassFragmentBitSet, FMassFragment);
 DECLARE_STRUCTTYPEBITSET_EXPORTED(MASSENTITY_API, FMassTagBitSet, FMassTag);
 DECLARE_STRUCTTYPEBITSET_EXPORTED(MASSENTITY_API, FMassChunkFragmentBitSet, FMassChunkFragment);
@@ -687,7 +645,7 @@ namespace UE::Mass
 		TArray<FType> FragmentInstances;
 	};
 
-	/**TMultiArray simple-type specialization */
+	/**TMultiArray single-type specialization */
 	template<typename T>
 	struct TMultiArray<T>
 	{
