[Mass] trivial fixes to documentation

diff --git a/Engine/Source/Runtime/MassEntity/Public/MassEntityQuery.h b/Engine/Source/Runtime/MassEntity/Public/MassEntityQuery.h
--- a/Engine/Source/Runtime/MassEntity/Public/MassEntityQuery.h
+++ b/Engine/Source/Runtime/MassEntity/Public/MassEntityQuery.h
@@ -140,7 +140,7 @@ public:
 
 	/** 
 	 * Sums the entity range lengths for each collection in EntityCollections, where the collection's 
-	 * archetype matches the querie's requirements.
+	 * archetype matches the query's requirements.
 	 * @return the number of entities this given query would process if called "now" for EntityCollections
 	 */
 	int32 GetNumMatchingEntities(TConstArrayView<FMassArchetypeEntityCollection> EntityCollections);
@@ -153,8 +153,8 @@ public:
 	bool HasMatchingEntities(FMassEntityManager& InEntityManager);
 
 	/** 
-	 * Sets a chunk filter condition that will applied to each chunk of all valid archetypes. Note 
-	 * that this condition won't be applied when a specific entity colleciton is used (via FMassArchetypeEntityCollection )
+	 * Sets a chunk filter condition that will be applied to each chunk of all valid archetypes. Note 
+	 * that this condition won't be applied when a specific entity collection is used (via FMassArchetypeEntityCollection )
 	 * The value returned by InFunction controls whether to allow execution (true) or block it (false).
 	 */
 	void SetChunkFilter(const FMassChunkConditionFunction& InFunction) { checkf(!HasChunkFilter(), TEXT("Chunk filter needs to be cleared before setting a new one.")); ChunkCondition = InFunction; }
