[Mass] Improved debug information when failing to fetch a subsystem with FMassSubsystemAccess

#rnx

diff --git a/Engine/Source/Runtime/MassEntity/Public/MassSubsystemAccess.h b/Engine/Source/Runtime/MassEntity/Public/MassSubsystemAccess.h
--- a/Engine/Source/Runtime/MassEntity/Public/MassSubsystemAccess.h
+++ b/Engine/Source/Runtime/MassEntity/Public/MassSubsystemAccess.h
@@ -72,7 +72,8 @@ struct MASSENTITY_API FMassSubsystemAccess
 	T* GetMutableSubsystem(const TSubclassOf<USubsystem> SubsystemClass)
 	{
 		const uint32 SystemIndex = FMassExternalSubsystemBitSet::GetTypeIndex(**SubsystemClass);
-		if (ensure(MutableSubsystemsBitSet.IsBitSet(SystemIndex)))
+		if (ensureMsgf(MutableSubsystemsBitSet.IsBitSet(SystemIndex)
+			, TEXT("Not configured to access SubsystemClass %s"), *GetNameSafe(SubsystemClass)))
 		{
 			return GetSubsystemInternal<T>(SystemIndex, SubsystemClass);
 		}
@@ -92,7 +93,8 @@ struct MASSENTITY_API FMassSubsystemAccess
 	const T* GetSubsystem(const TSubclassOf<USubsystem> SubsystemClass)
 	{
 		const uint32 SystemIndex = FMassExternalSubsystemBitSet::GetTypeIndex(**SubsystemClass);
-		if (ensure(ConstSubsystemsBitSet.IsBitSet(SystemIndex) || MutableSubsystemsBitSet.IsBitSet(SystemIndex)))
+		if (ensureMsgf(ConstSubsystemsBitSet.IsBitSet(SystemIndex) || MutableSubsystemsBitSet.IsBitSet(SystemIndex)
+			, TEXT("Not configured to access SubsystemClass %s"), *GetNameSafe(SubsystemClass)))
 		{
 			return GetSubsystemInternal<T>(SystemIndex, SubsystemClass);
 		}
