[Mass] More Mass headers cleanup, enforcing IWYU

#rnx

diff --git a/Engine/Source/Runtime/MassEntity/Private/MassArchetypeTypes.cpp b/Engine/Source/Runtime/MassEntity/Private/MassArchetypeTypes.cpp
--- a/Engine/Source/Runtime/MassEntity/Private/MassArchetypeTypes.cpp
+++ b/Engine/Source/Runtime/MassEntity/Private/MassArchetypeTypes.cpp
@@ -4,6 +4,7 @@
 #include "MassEntityManager.h"
 #include "MassArchetypeData.h"
 #include "MassEntityUtils.h"
+#include "Containers/StridedView.h"
 
 //-----------------------------------------------------------------------------
 // FMassArchetypeHandle
diff --git a/Engine/Source/Runtime/MassEntity/Public/MassArchetypeTypes.h b/Engine/Source/Runtime/MassEntity/Public/MassArchetypeTypes.h
--- a/Engine/Source/Runtime/MassEntity/Public/MassArchetypeTypes.h
+++ b/Engine/Source/Runtime/MassEntity/Public/MassArchetypeTypes.h
@@ -2,10 +2,14 @@
 
 #pragma once
 
+#if UE_ENABLE_INCLUDE_ORDER_DEPRECATED_IN_5_6
 #include "UObject/Class.h"
+#endif // UE_ENABLE_INCLUDE_ORDER_DEPRECATED_IN_5_6
 #include "Containers/ArrayView.h"
+#if UE_ENABLE_INCLUDE_ORDER_DEPRECATED_IN_5_6
 #include "Containers/StridedView.h"
 #include "Containers/UnrealString.h"
+#endif // UE_ENABLE_INCLUDE_ORDER_DEPRECATED_IN_5_6
 #include "MassEntityTypes.h"
 
 struct FMassEntityManager;
diff --git a/Engine/Source/Runtime/MassEntity/Public/MassEntityQuery.h b/Engine/Source/Runtime/MassEntity/Public/MassEntityQuery.h
--- a/Engine/Source/Runtime/MassEntity/Public/MassEntityQuery.h
+++ b/Engine/Source/Runtime/MassEntity/Public/MassEntityQuery.h
@@ -3,8 +3,10 @@
 #pragma once
 
 #include "UObject/Class.h"
+#if UE_ENABLE_INCLUDE_ORDER_DEPRECATED_IN_5_6
 #include "Containers/ArrayView.h"
 #include "Containers/UnrealString.h"
+#endif // UE_ENABLE_INCLUDE_ORDER_DEPRECATED_IN_5_6
 #include "MassEntityTypes.h"
 #include "MassArchetypeTypes.h"
 #if UE_ENABLE_INCLUDE_ORDER_DEPRECATED_IN_5_6
