[Mass] Issue 675512: Added a constructor to FMassEntityQuery::FArchetypeGroupingStep so that GroupSortingSteps.Emplace compiles on android.

#rnx

diff --git a/Engine/Source/Runtime/MassEntity/Public/MassEntityQuery.h b/Engine/Source/Runtime/MassEntity/Public/MassEntityQuery.h
--- a/Engine/Source/Runtime/MassEntity/Public/MassEntityQuery.h
+++ b/Engine/Source/Runtime/MassEntity/Public/MassEntityQuery.h
@@ -220,6 +220,11 @@ private:
 
 	struct FArchetypeGroupingStep
 	{
+		FArchetypeGroupingStep() = default;
+		FArchetypeGroupingStep(UE::Mass::FArchetypeGroupType InGroupType, const TFunction<bool(const UE::Mass::FArchetypeGroupID, const UE::Mass::FArchetypeGroupID)>& InPredicate)
+			: GroupType(InGroupType), Predicate(InPredicate)
+		{
+		}
 		UE::Mass::FArchetypeGroupType GroupType;
 		TFunction<bool(const UE::Mass::FArchetypeGroupID, const UE::Mass::FArchetypeGroupID)> Predicate;
 	};
