[Mass] Marked the deprecated virtual UMassProcessor.ConfigureQueries as `final` to enforce that overriding it is no longer valid.

#rnx

diff --git a/Engine/Source/Runtime/MassEntity/Public/MassProcessor.h b/Engine/Source/Runtime/MassEntity/Public/MassProcessor.h
--- a/Engine/Source/Runtime/MassEntity/Public/MassProcessor.h
+++ b/Engine/Source/Runtime/MassEntity/Public/MassProcessor.h
@@ -235,9 +235,9 @@ private:
 	//-----------------------------------------------------------------------------
 public:
 	UE_DEPRECATED(5.6, "Initialize is deprecated. Override Initialize(UObject&, const TSharedRef<FMassEntityManager>&) instead. If you want to call the function, use CallInitialize.")
-	virtual void Initialize(UObject& Owner);
+	virtual void Initialize(UObject& Owner) final;
 	UE_DEPRECATED(5.6, "This flavor of ConfigureQueries is deprecated. Override ConfigureQueries(const TSharedRef<FMassEntityManager>&) instead.")
-	virtual void ConfigureQueries() {};
+	virtual void ConfigureQueries() final {};
 };
 
 
