[Mass] Minor massage of EntityQuery code and comments related to entity grouping

#rnx

diff --git a/Engine/Source/Runtime/MassEntity/Private/MassEntityQuery.cpp b/Engine/Source/Runtime/MassEntity/Private/MassEntityQuery.cpp
--- a/Engine/Source/Runtime/MassEntity/Private/MassEntityQuery.cpp
+++ b/Engine/Source/Runtime/MassEntity/Private/MassEntityQuery.cpp
@@ -267,6 +267,7 @@ void FMassEntityQuery::ForEachEntityChunk(FMassEntityManager& EntityManager, FMa
 			// it's important to set requirements after caching archetypes due to that call potentially sorting the requirements and the order is relevant here.
 			ExecutionContext.SetFragmentRequirements(*this);
 
+			// note that this checkSlow is here on purpose, for debugging purposes, not data verification purposes.
 			checkSlow(OrderedArchetypeIndices.Num() == ValidArchetypes.Num());
 			for (const int32 ArchetypeIndex : OrderedArchetypeIndices)
 			{
@@ -525,14 +526,10 @@ void FMassEntityQuery::ResetGrouping()
 	IncrementChangeCounter();
 }
 
-void FMassEntityQuery::BuildOrderedArchetypeIndices(const int32 StartingIndex)
+void FMassEntityQuery::BuildOrderedArchetypeIndices(const int32 FirstNewArchetypeIndex)
 {
-	if (StartingIndex == 0)
-	{
-		OrderedArchetypeIndices.Reset();
-	}
 	OrderedArchetypeIndices.SetNumUninitialized(ValidArchetypes.Num());
-	for (int32 ArchetypeIndex = StartingIndex; ArchetypeIndex < ValidArchetypes.Num(); ++ArchetypeIndex)
+	for (int32 ArchetypeIndex = FirstNewArchetypeIndex; ArchetypeIndex < ValidArchetypes.Num(); ++ArchetypeIndex)
 	{
 		OrderedArchetypeIndices[ArchetypeIndex] = ArchetypeIndex;
 	}
@@ -542,6 +539,7 @@ void FMassEntityQuery::SortArchetypes(const int32 FirstNewArchetypeIndex)
 {
 	if (GroupSortingSteps.Num() == 0)
 	{
+		BuildOrderedArchetypeIndices(FirstNewArchetypeIndex);
 		return;
 	}
 
diff --git a/Engine/Source/Runtime/MassEntity/Public/MassEntityManager.h b/Engine/Source/Runtime/MassEntity/Public/MassEntityManager.h
--- a/Engine/Source/Runtime/MassEntity/Public/MassEntityManager.h
+++ b/Engine/Source/Runtime/MassEntity/Public/MassEntityManager.h
@@ -363,7 +363,7 @@ public:
 
 	/**
 	 * Fetches FArchetypeGroupType instance (copy) associated with the given GroupName. A new group type
-	 * is created if GroupName has not been used in the past.f
+	 * is created if GroupName has not been used in the past.
 	 */
 	UE::Mass::FArchetypeGroupType FinOrAddArchetypeGroupType(const FName GroupName);
 
diff --git a/Engine/Source/Runtime/MassEntity/Public/MassEntityQuery.h b/Engine/Source/Runtime/MassEntity/Public/MassEntityQuery.h
--- a/Engine/Source/Runtime/MassEntity/Public/MassEntityQuery.h
+++ b/Engine/Source/Runtime/MassEntity/Public/MassEntityQuery.h
@@ -201,8 +201,16 @@ private:
 		FMassExternalSubsystemBitSet MutableSubsystemsBitSet;
 	};
 
-	void SortArchetypes(const int32 FirstNewArchetype);
-	void BuildOrderedArchetypeIndices(const int32 StartingIndex = 0);
+	/**
+	 * Incrementally sorts all ValidArchetypes to fill OrderedArchetypeIndices with the expected order of archetype processing.
+	 * This function will only ever get called when there are actual sorting steps registered (@see GroupBy)
+	 */
+	void SortArchetypes(const int32 FirstNewArchetypeIndex = 0);
+	/**
+	 * An alternative to SortArchetypes that will get called in the absence of archetype sorting steps to maintain OrderedArchetypeIndices
+	 * and have it reflect the order of ValidArchetypes.
+	 */
+	void BuildOrderedArchetypeIndices(const int32 FirstNewArchetypeIndex = 0);
 
 	/** 
 	 * This function represents a condition that will be called for every chunk to be processed before the actual 
