[Mass] Fixed FMassEntityManager potentially shielding some UObjects from GC after already being marked for destruction with Deinitialize calls. The UObjects can be referenced by shared fragments, and the Entity Manager would mark them as referenced via its AddReferencedObjects

diff --git a/Engine/Source/Runtime/MassEntity/Private/MassEntityManager.cpp b/Engine/Source/Runtime/MassEntity/Private/MassEntityManager.cpp
--- a/Engine/Source/Runtime/MassEntity/Private/MassEntityManager.cpp
+++ b/Engine/Source/Runtime/MassEntity/Private/MassEntityManager.cpp
@@ -182,6 +182,13 @@ void FMassEntityManager::GetResourceSizeEx(FResourceSizeEx& CumulativeResourceSi
 
 void FMassEntityManager::AddReferencedObjects(FReferenceCollector& Collector)
 {
+	if (bInitialized == false)
+	{
+		// this is AddReferencedObjects called after Deinitialize call, which means we don't want to retain any object refs
+		// since this FMassEntityManager instance is going away even if it's kept alive by some stored shared refs at the moment.
+		return;
+	}
+
 	for (FConstSharedStruct& Struct : ConstSharedFragments)
 	{
 		Struct.AddStructReferencedObjects(Collector);
