[Mass] MassEntityTypes.h clean-up, in prep for upcoming changes

#rnx

diff --git a/Engine/Source/Runtime/MassEntity/Private/MassEntityTypes.cpp b/Engine/Source/Runtime/MassEntity/Private/MassEntityTypes.cpp
--- a/Engine/Source/Runtime/MassEntity/Private/MassEntityTypes.cpp
+++ b/Engine/Source/Runtime/MassEntity/Private/MassEntityTypes.cpp
@@ -15,6 +15,83 @@ DEFINE_TYPEBITSET(FMassSharedFragmentBitSet);
 DEFINE_TYPEBITSET(FMassConstSharedFragmentBitSet);
 DEFINE_TYPEBITSET(FMassExternalSubsystemBitSet);
 
+//-----------------------------------------------------------------------------
+// FMassArchetypeCompositionDescriptor
+//-----------------------------------------------------------------------------
+uint32 FMassArchetypeCompositionDescriptor::CalculateHash(const FMassFragmentBitSet& InFragments, const FMassTagBitSet& InTags
+	, const FMassChunkFragmentBitSet& InChunkFragments, const FMassSharedFragmentBitSet& InSharedFragmentBitSet
+	, const FMassConstSharedFragmentBitSet& InConstSharedFragmentBitSet)
+{
+	const uint32 FragmentsHash = GetTypeHash(InFragments);
+	const uint32 TagsHash = GetTypeHash(InTags);
+	const uint32 ChunkFragmentsHash = GetTypeHash(InChunkFragments);
+	const uint32 SharedFragmentsHash = GetTypeHash(InSharedFragmentBitSet);
+	const uint32 ConstSharedFragmentsHash = GetTypeHash(InConstSharedFragmentBitSet);
+
+	return HashCombine(
+			HashCombine(FragmentsHash, TagsHash)
+			, HashCombine(
+				HashCombine(ChunkFragmentsHash, SharedFragmentsHash)
+				, ConstSharedFragmentsHash));
+}
+
+int32 FMassArchetypeCompositionDescriptor::CountStoredTypes() const
+{
+	return Fragments.CountStoredTypes()
+		+ Tags.CountStoredTypes()
+		+ ChunkFragments.CountStoredTypes()
+		+ SharedFragments.CountStoredTypes()
+		+ ConstSharedFragments.CountStoredTypes();
+}
+
+void FMassArchetypeCompositionDescriptor::DebugOutputDescription(FOutputDevice& Ar) const
+{
+#if WITH_MASSENTITY_DEBUG 
+	if (Fragments.IsEmpty()
+		&& Tags.IsEmpty()
+		&& ChunkFragments.IsEmpty())
+	{
+		Ar.Logf(TEXT("Empty"));
+		return;
+	}
+
+	const bool bAutoLineEnd = Ar.GetAutoEmitLineTerminator();
+	Ar.SetAutoEmitLineTerminator(false);
+
+	if (!Fragments.IsEmpty())
+	{
+		Ar.Logf(TEXT("Fragments:\n"));
+		Fragments.DebugGetStringDesc(Ar);
+	}
+
+	if (!Tags.IsEmpty())
+	{
+		Ar.Logf(TEXT("Tags:\n"));
+		Tags.DebugGetStringDesc(Ar);
+	}
+
+	if (!ChunkFragments.IsEmpty())
+	{
+		Ar.Logf(TEXT("ChunkFragments:\n"));
+		ChunkFragments.DebugGetStringDesc(Ar);
+	}
+
+	if (!SharedFragments.IsEmpty())
+	{
+		Ar.Logf(TEXT("SharedFragments:\n"));
+		SharedFragments.DebugGetStringDesc(Ar);
+	}
+
+	if (!ConstSharedFragments.IsEmpty())
+	{
+		Ar.Logf(TEXT("ConstSharedFragments:\n"));
+		ConstSharedFragments.DebugGetStringDesc(Ar);
+	}
+
+	Ar.SetAutoEmitLineTerminator(bAutoLineEnd);
+#endif // WITH_MASSENTITY_DEBUG
+}
+
 //-----------------------------------------------------------------------------
 // FMassArchetypeSharedFragmentValues
 //-----------------------------------------------------------------------------
diff --git a/Engine/Source/Runtime/MassEntity/Public/MassEntityTypes.h b/Engine/Source/Runtime/MassEntity/Public/MassEntityTypes.h
--- a/Engine/Source/Runtime/MassEntity/Public/MassEntityTypes.h
+++ b/Engine/Source/Runtime/MassEntity/Public/MassEntityTypes.h
@@ -101,105 +101,43 @@ struct FMassArchetypeCompositionDescriptor
 
 	bool IsEquivalent(const FMassArchetypeCompositionDescriptor& OtherDescriptor) const
 	{
-		return Fragments.IsEquivalent(OtherDescriptor.Fragments) &&
-			Tags.IsEquivalent(OtherDescriptor.Tags) &&
-			ChunkFragments.IsEquivalent(OtherDescriptor.ChunkFragments) &&
-			SharedFragments.IsEquivalent(OtherDescriptor.SharedFragments) &&
-			ConstSharedFragments.IsEquivalent(OtherDescriptor.ConstSharedFragments);
+		return Fragments.IsEquivalent(OtherDescriptor.Fragments)
+			&& Tags.IsEquivalent(OtherDescriptor.Tags)
+			&& ChunkFragments.IsEquivalent(OtherDescriptor.ChunkFragments)
+			&& SharedFragments.IsEquivalent(OtherDescriptor.SharedFragments)
+			&& ConstSharedFragments.IsEquivalent(OtherDescriptor.ConstSharedFragments);
 	}
 
 	bool IsEmpty() const 
 	{ 
-		return Fragments.IsEmpty() &&
-			Tags.IsEmpty() &&
-			ChunkFragments.IsEmpty() &&
-			SharedFragments.IsEmpty() &&
-			ConstSharedFragments.IsEmpty();
+		return Fragments.IsEmpty()
+			&& Tags.IsEmpty()
+			&& ChunkFragments.IsEmpty()
+			&& SharedFragments.IsEmpty()
+			&& ConstSharedFragments.IsEmpty();
 	}
 
 	bool HasAll(const FMassArchetypeCompositionDescriptor& OtherDescriptor) const
 	{
-		return Fragments.HasAll(OtherDescriptor.Fragments) &&
-			Tags.HasAll(OtherDescriptor.Tags) &&
-			ChunkFragments.HasAll(OtherDescriptor.ChunkFragments) &&
-			SharedFragments.HasAll(OtherDescriptor.SharedFragments) &&
-			ConstSharedFragments.HasAll(OtherDescriptor.ConstSharedFragments);
+		return Fragments.HasAll(OtherDescriptor.Fragments)
+			&& Tags.HasAll(OtherDescriptor.Tags)
+			&& ChunkFragments.HasAll(OtherDescriptor.ChunkFragments)
+			&& SharedFragments.HasAll(OtherDescriptor.SharedFragments)
+			&& ConstSharedFragments.HasAll(OtherDescriptor.ConstSharedFragments);
 	}
 
-	static uint32 CalculateHash(const FMassFragmentBitSet& InFragments, const FMassTagBitSet& InTags
+	MASSENTITY_API static uint32 CalculateHash(const FMassFragmentBitSet& InFragments, const FMassTagBitSet& InTags
 		, const FMassChunkFragmentBitSet& InChunkFragments, const FMassSharedFragmentBitSet& InSharedFragmentBitSet
-		, const FMassConstSharedFragmentBitSet& InConstSharedFragmentBitSet)
-	{
-		const uint32 FragmentsHash = GetTypeHash(InFragments);
-		const uint32 TagsHash = GetTypeHash(InTags);
-		const uint32 ChunkFragmentsHash = GetTypeHash(InChunkFragments);
-		const uint32 SharedFragmentsHash = GetTypeHash(InSharedFragmentBitSet);
-		const uint32 ConstSharedFragmentsHash = GetTypeHash(InConstSharedFragmentBitSet);
-		return HashCombine(HashCombine(HashCombine(HashCombine(FragmentsHash, TagsHash), ChunkFragmentsHash), SharedFragmentsHash), ConstSharedFragmentsHash);
-	}	
+		, const FMassConstSharedFragmentBitSet& InConstSharedFragmentBitSet);
 
 	uint32 CalculateHash() const 
 	{
 		return CalculateHash(Fragments, Tags, ChunkFragments, SharedFragments, ConstSharedFragments);
 	}
 
-	int32 CountStoredTypes() const
-	{
-		return Fragments.CountStoredTypes()
-			+ Tags.CountStoredTypes()
-			+ ChunkFragments.CountStoredTypes()
-			+ SharedFragments.CountStoredTypes()
-			+ ConstSharedFragments.CountStoredTypes();
-	}
-
-	void DebugOutputDescription(FOutputDevice& Ar) const
-	{
-#if WITH_MASSENTITY_DEBUG
-		if (Fragments.IsEmpty()
-			&& Tags.IsEmpty()
-			&& ChunkFragments.IsEmpty())
-		{
-			Ar.Logf(TEXT("Empty"));
-			return;
-		}
-
-		const bool bAutoLineEnd = Ar.GetAutoEmitLineTerminator();
-		Ar.SetAutoEmitLineTerminator(false);
-
-		if (!Fragments.IsEmpty())
-		{
-			Ar.Logf(TEXT("Fragments:\n"));
-			Fragments.DebugGetStringDesc(Ar);
-		}
-
-		if (!Tags.IsEmpty())
-		{
-			Ar.Logf(TEXT("Tags:\n"));
-			Tags.DebugGetStringDesc(Ar);
-		}
-
-		if (!ChunkFragments.IsEmpty())
-		{
-			Ar.Logf(TEXT("ChunkFragments:\n"));
-			ChunkFragments.DebugGetStringDesc(Ar);
-		}
-
-		if (!SharedFragments.IsEmpty())
-		{
-			Ar.Logf(TEXT("SharedFragments:\n"));
-			SharedFragments.DebugGetStringDesc(Ar);
-		}
-
-		if (!ConstSharedFragments.IsEmpty())
-		{
-			Ar.Logf(TEXT("ConstSharedFragments:\n"));
-			ConstSharedFragments.DebugGetStringDesc(Ar);
-		}
+	MASSENTITY_API int32 CountStoredTypes() const;
 
-		Ar.SetAutoEmitLineTerminator(bAutoLineEnd);
-#endif // WITH_MASSENTITY_DEBUG
-
-	}
+	MASSENTITY_API void DebugOutputDescription(FOutputDevice& Ar) const;
 
 	FMassFragmentBitSet Fragments;
 	FMassTagBitSet Tags;
@@ -693,3 +631,4 @@ struct FMassArchetypeCreationParams
 	/** Name to identify the archetype while debugging*/
 	FName DebugName;
 };
+
