[Mass] minor clean up in Mass Archetype Groups code

#rnx

diff --git a/Engine/Source/Runtime/MassEntity/Private/MassArchetypeGroup.cpp b/Engine/Source/Runtime/MassEntity/Private/MassArchetypeGroup.cpp
--- a/Engine/Source/Runtime/MassEntity/Private/MassArchetypeGroup.cpp
+++ b/Engine/Source/Runtime/MassEntity/Private/MassArchetypeGroup.cpp
@@ -59,9 +59,9 @@ namespace UE::Mass
 		const int32 GroupTypeIndex = static_cast<int32>(GroupType);
 		if (IDContainer.IsValidIndex(GroupTypeIndex))
 		{
-			const bool IsLastElement = ((IDContainer.GetMaxIndex() - 1) == GroupTypeIndex);
+			const bool bIsLastElement = ((IDContainer.GetMaxIndex() - 1) == GroupTypeIndex);
 			IDContainer.RemoveAtUninitialized(GroupTypeIndex);
-			if (IsLastElement)
+			if (bIsLastElement)
 			{
 				Shrink();
 			}
@@ -77,7 +77,7 @@ namespace UE::Mass
 			Copy.IDContainer.RemoveAtUninitialized(GroupTypeIndex);
 
 			// if we removed the last element we need to shrink the container too.
-			if (const bool IsLastElement = ((IDContainer.GetMaxIndex() - 1) == GroupTypeIndex))
+			if (const bool bIsLastElement = ((IDContainer.GetMaxIndex() - 1) == GroupTypeIndex))
 			{
 				Copy.Shrink();
 			}
diff --git a/Engine/Source/Runtime/MassEntity/Private/MassEntityManager.cpp b/Engine/Source/Runtime/MassEntity/Private/MassEntityManager.cpp
--- a/Engine/Source/Runtime/MassEntity/Private/MassEntityManager.cpp
+++ b/Engine/Source/Runtime/MassEntity/Private/MassEntityManager.cpp
@@ -999,9 +999,6 @@ void FMassEntityManager::RemoveEntityFromGroupType(FMassEntityHandle EntityHandl
 	{
 		if (CurrentArchetype->IsInGroupOfType(GroupType))
 		{
-			/*UE::Mass::FArchetypeGroups NewGroups = CurrentArchetype->GetGroups();
-			NewGroups.Remove(GroupType);
-			const FMassArchetypeHandle NewArchetypeHandle = InternalCreateSimilarArchetype(CurrentArchetypeHandle.DataPtr, MoveTemp(NewGroups));*/
 			const FMassArchetypeHandle NewArchetypeHandle = InternalCreateSimilarArchetype(CurrentArchetypeHandle.DataPtr, CurrentArchetype->GetGroups().Remove(GroupType));
 
 			CurrentArchetype->MoveEntityToAnotherArchetype(EntityHandle, *NewArchetypeHandle.DataPtr.Get());
