[Mass] tiny change to use move semantics when issuing "processors graph parallel processing done" task

diff --git a/Engine/Source/Runtime/MassEntity/Private/MassExecutor.cpp b/Engine/Source/Runtime/MassEntity/Private/MassExecutor.cpp
--- a/Engine/Source/Runtime/MassEntity/Private/MassExecutor.cpp
+++ b/Engine/Source/Runtime/MassEntity/Private/MassExecutor.cpp
@@ -158,7 +158,7 @@ void RunProcessorsView(TArrayView<UMassProcessor* const> Processors, FMassProces
 
 struct FMassExecutorDoneTask
 {
-	FMassExecutorDoneTask(const FMassExecutionContext& InExecutionContext, TFunction<void()> InOnDoneNotification, const FString& InDebugName, ENamedThreads::Type InDesiredThread)
+	FMassExecutorDoneTask(FMassExecutionContext&& InExecutionContext, TFunction<void()> InOnDoneNotification, const FString& InDebugName, ENamedThreads::Type InDesiredThread)
 		: ExecutionContext(InExecutionContext)
 		, OnDoneNotification(InOnDoneNotification)
 		, DebugName(InDebugName)
@@ -230,7 +230,7 @@ FGraphEventRef TriggerParallelTasks(UMassProcessor& Processor, FMassProcessingCo
 	{
 		const FGraphEventArray Prerequisites = { CompletionEvent };
 		CompletionEvent = TGraphTask<FMassExecutorDoneTask>::CreateTask(&Prerequisites)
-			.ConstructAndDispatchWhenReady(ExecutionContext, OnDoneNotification, Processor.GetName(), CurrentThread);
+			.ConstructAndDispatchWhenReady(MoveTemp(ExecutionContext), OnDoneNotification, Processor.GetName(), CurrentThread);
 	}
 
 	return CompletionEvent;
