[Mass] Minor: switched FMassExecuteFunction from `typedef` to `using`

#rnx

diff --git a/Engine/Source/Runtime/MassEntity/Public/MassArchetypeTypes.h b/Engine/Source/Runtime/MassEntity/Public/MassArchetypeTypes.h
--- a/Engine/Source/Runtime/MassEntity/Public/MassArchetypeTypes.h
+++ b/Engine/Source/Runtime/MassEntity/Public/MassArchetypeTypes.h
@@ -24,8 +24,8 @@ struct FMassEntityView;
 struct FMassDebugger;
 struct FMassArchetypeHelper;
 
-typedef TFunction< void(FMassExecutionContext& /*ExecutionContext*/) > FMassExecuteFunction;
-typedef TFunction< bool(const FMassExecutionContext& /*ExecutionContext*/) > FMassChunkConditionFunction;
+using FMassExecuteFunction = TFunction< void(FMassExecutionContext& /*ExecutionContext*/) >;
+using FMassChunkConditionFunction = TFunction< bool(const FMassExecutionContext& /*ExecutionContext*/) >;
 
 //-----------------------------------------------------------------------------
 // FMassArchetypeHandle
