Allowed parallel Mass queries to be run unbalanced.

Some of the query callbacks in TEDS are non-uniform, meaning that one chunk needs to do very little work while another chunk can require large amounts of work. In worse case scenarios this can cause the one thread to take much longer than other threads if it ends up with multiple chunk that need a lot of work. By introducing the option to balance this, the work can be balanced better between threads and allow more chunks to be processed in parallel but at the cost of extra time spend scheduling jobs.

[RN] Mass can now optionally auto-balance parallel queries. This comes at a slight scheduling overhead, but can improve performance for processors that don't have even performance for all their chunks.

diff --git a/Engine/Plugins/Experimental/EditorDataStorage/Source/TedsCore/Private/Processors/TypedElementProcessorAdaptors.cpp b/Engine/Plugins/Experimental/EditorDataStorage/Source/TedsCore/Private/Processors/TypedElementProcessorAdaptors.cpp
--- a/Engine/Plugins/Experimental/EditorDataStorage/Source/TedsCore/Private/Processors/TypedElementProcessorAdaptors.cpp
+++ b/Engine/Plugins/Experimental/EditorDataStorage/Source/TedsCore/Private/Processors/TypedElementProcessorAdaptors.cpp
@@ -1123,7 +1123,11 @@ UE::Editor::DataStorage::FQueryResult FTypedElementQueryProcessorData::Execute(
 				};
 			if (EnumHasAnyFlags(ExecutionFlags, EDirectQueryExecutionFlags::ParallelizeChunks))
 			{
-				NativeQuery.ParallelForEachEntityChunk(EntityManager, Context, ExecuteFunction);
+				FMassEntityQuery::EParallelExecutionFlags Mode =
+					EnumHasAnyFlags(ExecutionFlags, EDirectQueryExecutionFlags::AutoBalanceParallelChunkProcessing)
+					? FMassEntityQuery::EParallelExecutionFlags::AutoBalance
+					: FMassEntityQuery::EParallelExecutionFlags::Default;
+				NativeQuery.ParallelForEachEntityChunk(EntityManager, Context, ExecuteFunction, Mode);
 			}
 			else
 			{
diff --git a/Engine/Source/Developer/MassEntityTestSuite/Private/MassParallelQueryTest.cpp b/Engine/Source/Developer/MassEntityTestSuite/Private/MassParallelQueryTest.cpp
--- a/Engine/Source/Developer/MassEntityTestSuite/Private/MassParallelQueryTest.cpp
+++ b/Engine/Source/Developer/MassEntityTestSuite/Private/MassParallelQueryTest.cpp
@@ -45,7 +45,7 @@ struct FQueryTest_ParallelForBasic : FEntityTestBase
 					return;
 				}
 				FPlatformProcess::Sleep(0.01f);
-			}, FMassEntityQuery::ForceParallelExecution);
+			}, FMassEntityQuery::EParallelExecutionFlags::Force);
 		
 		AITEST_EQUAL("All created entites should have been processed at this point", EntitiesProcessed.load(), NumEntitiesToCreate);
 		AITEST_TRUE("Some work is expected to be executed in parallel", ConcurrentAccess.load() > 0);
@@ -102,7 +102,7 @@ struct FQueryTest_ParallelForCollection : FEntityTestBase
 					return;
 				}
 				FPlatformProcess::Sleep(0.01f);
-			}, FMassEntityQuery::ForceParallelExecution);
+			}, FMassEntityQuery::EParallelExecutionFlags::Force);
 
 		AITEST_EQUAL("All entities passed for processing should have been counted", EntitiesProcessed.load(), EntitiesToProcess.Num());
 		AITEST_EQUAL("The number of jobs should match expectations", NumJobs, JobsExecuted.load());
@@ -139,7 +139,7 @@ struct FQueryTest_ParallelCommands : FEntityTestBase
 		Query.ParallelForEachEntityChunk(*EntityManager, Context, [](FMassExecutionContext& Context)
 			{
 				Context.Defer().PushCommand<FMassCommandAddTags<FTestFragment_Tag>>(Context.GetEntities());
-			}, FMassEntityQuery::ForceParallelExecution);
+			}, FMassEntityQuery::EParallelExecutionFlags::Force);
 
 		const int32 OriginalArchetypeCountAfterMove = EntityManager->DebugGetArchetypeEntitiesCount(FloatsArchetype);
 		AITEST_EQUAL("Target archetype should be empty after the move", OriginalArchetypeCountAfterMove, 0);
diff --git a/Engine/Source/Runtime/MassEntity/Private/MassEntityQuery.cpp b/Engine/Source/Runtime/MassEntity/Private/MassEntityQuery.cpp
--- a/Engine/Source/Runtime/MassEntity/Private/MassEntityQuery.cpp
+++ b/Engine/Source/Runtime/MassEntity/Private/MassEntityQuery.cpp
@@ -18,7 +18,6 @@
 #include "MassRequirementAccessDetector.h"
 #endif // WITH_MASSENTITY_DEBUG
 
-
 namespace UE::Mass::Tweakables
 {
 	/**
@@ -291,24 +290,39 @@ void FMassEntityQuery::ParallelForEachEntityChunkInCollection(TConstArrayView<FM
 	, FMassEntityManager& EntityManager, FMassExecutionContext& ExecutionContext, const FMassExecuteFunction& ExecuteFunction
 	, const EParallelForMode ParallelMode)
 {
-	if (UE::Mass::Tweakables::bAllowParallelExecution == false && ParallelMode != ForceParallelExecution)
+	EParallelExecutionFlags Flags = EnumHasAnyFlags(ParallelMode, ForceParallelExecution) ? EParallelExecutionFlags::Force : EParallelExecutionFlags::Default;
+	ParallelForEachEntityChunkInCollection(EntityCollections, EntityManager, ExecutionContext, ExecuteFunction, Flags);
+}
+
+void FMassEntityQuery::ParallelForEachEntityChunkInCollection(TConstArrayView<FMassArchetypeEntityCollection> EntityCollections
+	, FMassEntityManager& EntityManager, FMassExecutionContext& ExecutionContext, const FMassExecuteFunction& ExecuteFunction
+	, const EParallelExecutionFlags Flags)
+{
+	if (UE::Mass::Tweakables::bAllowParallelExecution == false && !EnumHasAnyFlags(Flags, EParallelExecutionFlags::Force))
 	{
 		ForEachEntityChunkInCollections(EntityCollections, EntityManager, ExecutionContext, ExecuteFunction);
 		return;
 	}
 
-	ParallelFor(EntityCollections.Num(), [this, &EntityManager, &ExecutionContext, &ExecuteFunction, &EntityCollections, ParallelMode](const int32 JobIndex)
+	ParallelFor(EntityCollections.Num(), [this, &EntityManager, &ExecutionContext, &ExecuteFunction, &EntityCollections, Flags](const int32 JobIndex)
 	{
 		FMassExecutionContext LocalExecutionContext = ExecutionContext; 
 		LocalExecutionContext.SetEntityCollection(EntityCollections[JobIndex]);
-		ParallelForEachEntityChunk(EntityManager, LocalExecutionContext, ExecuteFunction, ParallelMode);
+		ParallelForEachEntityChunk(EntityManager, LocalExecutionContext, ExecuteFunction, Flags);
 	});
 }
 
 void FMassEntityQuery::ParallelForEachEntityChunk(FMassEntityManager& EntityManager, FMassExecutionContext& ExecutionContext
 	, const FMassExecuteFunction& ExecuteFunction, const EParallelForMode ParallelMode)
 {
-	if (UE::Mass::Tweakables::bAllowParallelExecution == false && ParallelMode != ForceParallelExecution)
+	EParallelExecutionFlags Flags = EnumHasAnyFlags(ParallelMode, ForceParallelExecution) ? EParallelExecutionFlags::Force : EParallelExecutionFlags::Default;
+	ParallelForEachEntityChunk(EntityManager, ExecutionContext, ExecuteFunction, Flags);
+}
+
+void FMassEntityQuery::ParallelForEachEntityChunk(FMassEntityManager& EntityManager, FMassExecutionContext& ExecutionContext
+	, const FMassExecuteFunction& ExecuteFunction, const EParallelExecutionFlags Flags)
+{
+	if (UE::Mass::Tweakables::bAllowParallelExecution == false && !EnumHasAnyFlags(Flags, EParallelExecutionFlags::Force))
 	{
 		ForEachEntityChunk(EntityManager, ExecutionContext, ExecuteFunction);
 		return;
@@ -401,6 +415,8 @@ void FMassEntityQuery::ParallelForEachEntityChunk(FMassEntityManager& EntityMana
 
 	if (Jobs.Num())
 	{
+		const EParallelForFlags ParallelForFlags = 
+			EnumHasAnyFlags(Flags, EParallelExecutionFlags::AutoBalance) ? EParallelForFlags::Unbalanced : EParallelForFlags::None;
 		if (bAllowParallelCommands)
 		{
 			struct FTaskContext
@@ -432,7 +448,7 @@ void FMassEntityQuery::ParallelForEachEntityChunk(FMassEntityManager& EntityMana
 						, Jobs[JobIndex].ArchetypeIndex != INDEX_NONE ? ArchetypeFragmentMapping[Jobs[JobIndex].ArchetypeIndex] : FMassQueryRequirementIndicesMapping()
 						, Jobs[JobIndex].EntityRange
 						, ChunkCondition);
-				});
+				}, ParallelForFlags);
 
 			// merge all command buffers
 			for (FTaskContext& CommandContext : TaskContext)
@@ -449,7 +465,7 @@ void FMassEntityQuery::ParallelForEachEntityChunk(FMassEntityManager& EntityMana
 						, Jobs[JobIndex].ArchetypeIndex != INDEX_NONE ? ArchetypeFragmentMapping[Jobs[JobIndex].ArchetypeIndex] : FMassQueryRequirementIndicesMapping()
 						, Jobs[JobIndex].EntityRange
 						, ChunkCondition);
-				});
+				}, ParallelForFlags);
 		}
 	}
 
@@ -641,4 +657,4 @@ FMassEntityQuery::FMassEntityQuery(std::initializer_list<UScriptStruct*>)
 FMassEntityQuery::FMassEntityQuery(TConstArrayView<const UScriptStruct*>)
 	: FMassEntityQuery()
 {
-}
\ No newline at end of file
+}
diff --git a/Engine/Source/Runtime/MassEntity/Public/MassEntityQuery.h b/Engine/Source/Runtime/MassEntity/Public/MassEntityQuery.h
--- a/Engine/Source/Runtime/MassEntity/Public/MassEntityQuery.h
+++ b/Engine/Source/Runtime/MassEntity/Public/MassEntityQuery.h
@@ -66,10 +66,17 @@ private:
     };
 
 public:
-	enum EParallelForMode
+	enum class EParallelExecutionFlags
 	{
-		Default, // use whatever the whole system has been configured for
-		ForceParallelExecution // force
+		// Use whatever the whole system has been configured for.
+		Default = 0,
+		// Force parallel execution of a processor for each chunk even when parallel execution has been disabled.
+		Force = 1 << 0,
+		// The default behavior for parallel execution assigns each chunk to a thread before execution. This implicitly assumes all chunks
+		// take roughly the same amount of time to process. If chunks vary in the time it takes to process this flag can be used to queue
+		// chunks so threads can pick them up as soon as possible. This makes starting the processing of a chunk more expensive but can
+		// result in better overall utilization of threads.
+		AutoBalance = 1 << 1
 	};
 
 	FMassEntityQuery() = default;
@@ -91,14 +98,14 @@ public:
 	 * Attempts to process every chunk of every affected archetype in parallel.
 	 */
 	void ParallelForEachEntityChunk(FMassEntityManager& EntityManager, FMassExecutionContext& ExecutionContext
-		, const FMassExecuteFunction& ExecuteFunction, const EParallelForMode ParallelMode = Default);
+		, const FMassExecuteFunction& ExecuteFunction, const EParallelExecutionFlags Flags = EParallelExecutionFlags::Default);
 
 	void ForEachEntityChunkInCollections(TConstArrayView<FMassArchetypeEntityCollection> EntityCollections, FMassEntityManager& EntityManager
 		, FMassExecutionContext& ExecutionContext, const FMassExecuteFunction& ExecuteFunction);
 
 	void ParallelForEachEntityChunkInCollection(TConstArrayView<FMassArchetypeEntityCollection> EntityCollections
 		, FMassEntityManager& EntityManager, FMassExecutionContext& ExecutionContext, const FMassExecuteFunction& ExecuteFunction
-		, const EParallelForMode ParallelMode);
+		, const EParallelExecutionFlags Flags);
 
 	/** Will gather all archetypes from InEntityManager matching this->Requirements.
 	 *  Note that no work will be done if the cached data is up to date (as tracked by EntitySubsystemHash and 
@@ -272,8 +279,25 @@ public:
 
 	UE_DEPRECATED(5.6, "This type of FMassEntityQuery is no longer supported. Use one of the other constructors instead.")
 	FMassEntityQuery(TConstArrayView<const UScriptStruct*> InitList);
+
+	enum EParallelForMode
+	{
+		Default = static_cast<int32>(EParallelExecutionFlags::Default),
+		ForceParallelExecution = static_cast<int32>(EParallelExecutionFlags::Force),
+	};
+
+	UE_DEPRECATED(5.6, "The version of ParallelForEachEntityChunk using EParallelForMode has been replaced by the version using EParallelExecutionFlags.")
+	void ParallelForEachEntityChunk(FMassEntityManager& EntityManager, FMassExecutionContext& ExecutionContext
+		, const FMassExecuteFunction& ExecuteFunction, const EParallelForMode ParallelMode);
+	
+	UE_DEPRECATED(5.6, "The version of ParallelForEachEntityChunk using EParallelForMode has been replaced by the version using EParallelExecutionFlags.")
+	void ParallelForEachEntityChunkInCollection(TConstArrayView<FMassArchetypeEntityCollection> EntityCollections
+		, FMassEntityManager& EntityManager, FMassExecutionContext& ExecutionContext, const FMassExecuteFunction& ExecuteFunction
+		, const EParallelForMode ParallelMode);
 };
 
+ENUM_CLASS_FLAGS(FMassEntityQuery::EParallelExecutionFlags);
+
 #if WITH_ARCHETYPE_MATCH_OVERRIDE
 template <typename T> requires TArchetypeMatchOverrideConcept<T>
 void FMassEntityQuery::SetArchetypeMatchOverride(const T& Context)
diff --git a/Engine/Source/Runtime/TypedElementFramework/Public/Elements/Interfaces/TypedElementQueryStorageInterfaces.h b/Engine/Source/Runtime/TypedElementFramework/Public/Elements/Interfaces/TypedElementQueryStorageInterfaces.h
--- a/Engine/Source/Runtime/TypedElementFramework/Public/Elements/Interfaces/TypedElementQueryStorageInterfaces.h
+++ b/Engine/Source/Runtime/TypedElementFramework/Public/Elements/Interfaces/TypedElementQueryStorageInterfaces.h
@@ -312,8 +312,14 @@ namespace UE::Editor::DataStorage
 	{
 		Default = 0, //< No settings, use the default behavior.
 		ParallelizeChunks = 1 << 0, //< If set, each chunk is processed on a separate thread.
-		IgnoreActivationCount = 1 << 1, //< If set, a direct call will not check activatable queries if they're set.
-		AllowBoundQueries = 1 << 2, //< Stops checking if queries are bound when set, otherwise bound queries will assert.
+		/*
+		 * Requires ParallelizeChunks to be set.
+		 * If set together with ParallelizeChunks, each chunk will be scheduled individually. This allows work to be better
+		 * distributed among threads, but comes at additional scheduling cost. Add this flag if the processing time per chunk varies.
+		 */
+		AutoBalanceParallelChunkProcessing = 1 << 1,
+		IgnoreActivationCount = 1 << 2, //< If set, a direct call will not check activatable queries if they're set.
+		AllowBoundQueries = 1 << 3, //< Stops checking if queries are bound when set, otherwise bound queries will assert.
 	};
 	ENUM_CLASS_FLAGS(EDirectQueryExecutionFlags);
 
