[Mass] Removed unused local vars from MassArchetypeTypes.cpp

#rnx

diff --git a/Engine/Source/Runtime/MassEntity/Private/MassArchetypeTypes.cpp b/Engine/Source/Runtime/MassEntity/Private/MassArchetypeTypes.cpp
--- a/Engine/Source/Runtime/MassEntity/Private/MassArchetypeTypes.cpp
+++ b/Engine/Source/Runtime/MassEntity/Private/MassArchetypeTypes.cpp
@@ -26,7 +26,6 @@ FMassArchetypeEntityCollection::FMassArchetypeEntityCollection(const FMassArchet
 	}
 
 	const FMassArchetypeData* ArchetypeData = FMassArchetypeHelper::ArchetypeDataFromHandle(InArchetype);
-	const int32 NumEntitiesPerChunk = ArchetypeData ? ArchetypeData->GetNumEntitiesPerChunk() : MAX_int32;
 
 	// InEntities has a real chance of not being sorted by AbsoluteIndex. We gotta fix that to optimize how we process the data 
 	TArray<int32> TrueIndices;
