[Mass] Fixed FMassExecutionContext::ClearExecutionData not clearing out the entity view, thus being able to call Context.GetEntities outside of query processing and still get some data back - which is not expected to happen.

#rnx

diff --git a/Engine/Source/Runtime/MassEntity/Private/MassExecutionContext.cpp b/Engine/Source/Runtime/MassEntity/Private/MassExecutionContext.cpp
--- a/Engine/Source/Runtime/MassEntity/Private/MassExecutionContext.cpp
+++ b/Engine/Source/Runtime/MassEntity/Private/MassExecutionContext.cpp
@@ -31,6 +31,7 @@ void FMassExecutionContext::ClearExecutionData()
 	ConstSharedFragmentViews.Reset();
 	SharedFragmentViews.Reset();
 	CurrentArchetypeCompositionDescriptor.Reset();
+	EntityListView = {};
 	ChunkSerialModificationNumber = -1;
 }
 
