[Mass] Added missing "IsProcessing" checks to MassEntityManager - turns out we were missing quite a lot of those.

This change prompted the deprecation of FMassEntityManager::BatchSetEntityFragmentsValues since it is static an we're unable to check whether Mass processing is in progress. New funciton name is BatchSetEntityFragmentValues
To keep function names consistent deprecated and renamed FMassEntityManager::SetEntityFragmentsValues, new name is SetEntityFragmentValues
All call sites in our code updated.

diff --git a/Engine/Plugins/Runtime/MassGameplay/Source/MassSpawner/Private/MassSpawnerSubsystem.cpp b/Engine/Plugins/Runtime/MassGameplay/Source/MassSpawner/Private/MassSpawnerSubsystem.cpp
--- a/Engine/Plugins/Runtime/MassGameplay/Source/MassSpawner/Private/MassSpawnerSubsystem.cpp
+++ b/Engine/Plugins/Runtime/MassGameplay/Source/MassSpawner/Private/MassSpawnerSubsystem.cpp
@@ -135,7 +135,7 @@ void UMassSpawnerSubsystem::DoSpawning(const FMassEntityTemplate& EntityTemplate
 	// 2. Copy data from FMassEntityTemplate.Fragments.
 	//		a. @todo, could be done as part of creation?
 	TConstArrayView<FInstancedStruct> FragmentInstances = EntityTemplate.GetInitialFragmentValues();
-	EntityManager->BatchSetEntityFragmentsValues(CreationContext->GetEntityCollections(), FragmentInstances);
+	EntityManager->BatchSetEntityFragmentValues(CreationContext->GetEntityCollections(), FragmentInstances);
 	
 	// 3. Run SpawnDataInitializer, if set. This is a special type of processor that operates on the entities to initialize them.
 	// e.g., will run UInstancedActorsInitializerProcessor for Mass InstancedActors
diff --git a/Engine/Source/Runtime/MassEntity/Private/MassEntityBuilder.cpp b/Engine/Source/Runtime/MassEntity/Private/MassEntityBuilder.cpp
--- a/Engine/Source/Runtime/MassEntity/Private/MassEntityBuilder.cpp
+++ b/Engine/Source/Runtime/MassEntity/Private/MassEntityBuilder.cpp
@@ -144,7 +144,7 @@ FMassEntityHandle FEntityBuilder::Commit()
 
 		FMassArchetypeHandle ArchetypeHandle = EntityManager->CreateArchetype(Composition);
 		EntityManager->BuildEntity(EntityHandle, ArchetypeHandle, SharedFragmentValues);
-		EntityManager->SetEntityFragmentsValues(EntityHandle, Fragments);
+		EntityManager->SetEntityFragmentValues(EntityHandle, Fragments);
 	}
 
 	State = EState::Committed;
diff --git a/Engine/Source/Runtime/MassEntity/Private/MassEntityManager.cpp b/Engine/Source/Runtime/MassEntity/Private/MassEntityManager.cpp
--- a/Engine/Source/Runtime/MassEntity/Private/MassEntityManager.cpp
+++ b/Engine/Source/Runtime/MassEntity/Private/MassEntityManager.cpp
@@ -16,6 +16,7 @@
 #include "MassEntityUtils.h"
 #include "MassEntityBuilder.h"
 
+#define CHECK_SYNC_API() checkf(IsProcessing() == false, TEXT("Synchronous API function %hs called during mass processing. Use asynchronous API instead."), __FUNCTION__)
 
 const FMassEntityHandle FMassEntityManager::InvalidEntity;
 
@@ -631,7 +632,7 @@ void FMassEntityManager::DoEntityCompaction(const double TimeAllowed)
 
 FMassEntityHandle FMassEntityManager::CreateEntity(const FMassArchetypeHandle& ArchetypeHandle, const FMassArchetypeSharedFragmentValues& SharedFragmentValues)
 {
-	checkf(IsProcessing() == false, TEXT("Synchronous API function %hs called during mass processing. Use asynchronous API instead."), __FUNCTION__);
+	CHECK_SYNC_API();
 	check(ArchetypeHandle.IsValid());
 
 	const FMassEntityHandle Entity = ReserveEntity();
@@ -644,7 +645,7 @@ FMassEntityHandle FMassEntityManager::CreateEntity(const FMassArchetypeHandle& A
 
 FMassEntityHandle FMassEntityManager::CreateEntity(TConstArrayView<FInstancedStruct> FragmentInstanceList, const FMassArchetypeSharedFragmentValues& SharedFragmentValues, const FMassArchetypeCreationParams& CreationParams)
 {
-	checkf(IsProcessing() == false, TEXT("Synchronous API function %hs called during mass processing. Use asynchronous API instead."), __FUNCTION__);
+	CHECK_SYNC_API();
 	check(FragmentInstanceList.Num() > 0);
 
 	const FMassArchetypeHandle& ArchetypeHandle = CreateArchetype(FMassArchetypeCompositionDescriptor(FragmentInstanceList,
@@ -682,7 +683,7 @@ void FMassEntityManager::ReleaseReservedEntity(FMassEntityHandle Entity)
 
 void FMassEntityManager::BuildEntity(FMassEntityHandle Entity, const FMassArchetypeHandle& ArchetypeHandle, const FMassArchetypeSharedFragmentValues& SharedFragmentValues)
 {
-	checkf(IsProcessing() == false, TEXT("Synchronous API function %hs called during mass processing. Use asynchronous API instead."), __FUNCTION__);
+	CHECK_SYNC_API();
 	checkf(!IsEntityBuilt(Entity), TEXT("Expecting an entity that is not already built"));
 	check(ArchetypeHandle.IsValid());
 
@@ -691,7 +692,7 @@ void FMassEntityManager::BuildEntity(FMassEntityHandle Entity, const FMassArchet
 
 void FMassEntityManager::BuildEntity(FMassEntityHandle Entity, TConstArrayView<FInstancedStruct> FragmentInstanceList, const FMassArchetypeSharedFragmentValues& SharedFragmentValues)
 {
-	checkf(IsProcessing() == false, TEXT("Synchronous API function %hs called during mass processing. Use asynchronous API instead."), __FUNCTION__);
+	CHECK_SYNC_API();
 	check(FragmentInstanceList.Num() > 0);
 	checkf(!IsEntityBuilt(Entity), TEXT("Expecting an entity that is not already built"));
 
@@ -737,7 +738,7 @@ int32 FMassEntityManager::BatchReserveEntities(TArrayView<FMassEntityHandle> InO
 TSharedRef<FMassEntityManager::FEntityCreationContext> FMassEntityManager::BatchBuildEntities(const FMassArchetypeEntityCollectionWithPayload& EncodedEntitiesWithPayload
 	, const FMassFragmentBitSet& FragmentsAffected, const FMassArchetypeSharedFragmentValues& SharedFragmentValues, const FMassArchetypeCreationParams& CreationParams)
 {
-	checkf(IsProcessing() == false, TEXT("Synchronous API function %hs called during mass processing. Use asynchronous API instead."), __FUNCTION__);
+	CHECK_SYNC_API();
 	check(SharedFragmentValues.IsSorted());
 
 	FMassArchetypeCompositionDescriptor Composition(FragmentsAffected, FMassTagBitSet(), FMassChunkFragmentBitSet(), FMassSharedFragmentBitSet(), FMassConstSharedFragmentBitSet());
@@ -757,7 +758,7 @@ TSharedRef<FMassEntityManager::FEntityCreationContext> FMassEntityManager::Batch
 	, FMassArchetypeCompositionDescriptor&& Composition
 	, const FMassArchetypeSharedFragmentValues& SharedFragmentValues, const FMassArchetypeCreationParams& CreationParams)
 {
-	checkf(IsProcessing() == false, TEXT("Synchronous API function %hs called during mass processing. Use asynchronous API instead."), __FUNCTION__);
+	CHECK_SYNC_API();
 
 	TRACE_CPUPROFILER_EVENT_SCOPE(Mass_BatchBuildEntities);
 
@@ -813,11 +814,11 @@ TSharedRef<FMassEntityManager::FEntityCreationContext> FMassEntityManager::Batch
 TSharedRef<FMassEntityManager::FEntityCreationContext> FMassEntityManager::BatchCreateReservedEntities(const FMassArchetypeHandle& ArchetypeHandle
 	, const FMassArchetypeSharedFragmentValues& SharedFragmentValues, TConstArrayView<FMassEntityHandle> ReservedEntities)
 {
-	TRACE_CPUPROFILER_EVENT_SCOPE(Mass_BatchCreateReservedEntities);
-
-	checkf(IsProcessing() == false, TEXT("Synchronous API function %hs called during mass processing. Use asynchronous API instead."), __FUNCTION__);
+	CHECK_SYNC_API();
 	checkf(!ReservedEntities.IsEmpty(), TEXT("No reserved entities given to batch create."));
 
+	TRACE_CPUPROFILER_EVENT_SCOPE(Mass_BatchCreateReservedEntities);
+
 	return InternalBatchCreateReservedEntities(
 		GetOrCreateSuitableArchetype(ArchetypeHandle, SharedFragmentValues.GetSharedFragmentBitSet(), SharedFragmentValues.GetConstSharedFragmentBitSet())
 		, SharedFragmentValues, ReservedEntities);
@@ -826,12 +827,12 @@ TSharedRef<FMassEntityManager::FEntityCreationContext> FMassEntityManager::Batch
 TSharedRef<FMassEntityManager::FEntityCreationContext> FMassEntityManager::BatchCreateEntities(const FMassArchetypeHandle& ArchetypeHandle
 	, const FMassArchetypeSharedFragmentValues& SharedFragmentValues, const int32 Count, TArray<FMassEntityHandle>& InOutEntities)
 {
-	TRACE_CPUPROFILER_EVENT_SCOPE(Mass_BatchCreateEntities);
-
-	checkf(IsProcessing() == false, TEXT("Synchronous API function %hs called during mass processing. Use asynchronous API instead."), __FUNCTION__);
+	CHECK_SYNC_API();
 	testableCheckfReturn(ArchetypeHandle.IsValid(), return GetOrMakeCreationContext()
 		, TEXT("%hs expecting a valid ArchetypeHandle"), __FUNCTION__);
 
+	TRACE_CPUPROFILER_EVENT_SCOPE(Mass_BatchCreateEntities);
+
 	TConstArrayView<FMassEntityHandle> ReservedEntities = BatchReserveEntities(Count, InOutEntities);
 	
 	return InternalBatchCreateReservedEntities(
@@ -863,7 +864,7 @@ TSharedRef<FMassEntityManager::FEntityCreationContext> FMassEntityManager::Inter
 
 void FMassEntityManager::DestroyEntity(FMassEntityHandle Entity)
 {
-	checkf(IsProcessing() == false, TEXT("Synchronous API function %hs called during mass processing. Use asynchronous API instead."), __FUNCTION__);
+	CHECK_SYNC_API();
 	
 	CheckIfEntityIsActive(Entity);
 
@@ -880,11 +881,11 @@ void FMassEntityManager::DestroyEntity(FMassEntityHandle Entity)
 
 void FMassEntityManager::BatchDestroyEntities(TConstArrayView<FMassEntityHandle> InEntities)
 {
-	TRACE_CPUPROFILER_EVENT_SCOPE(Mass_BatchDestroyEntities);
-
-	checkf(IsProcessing() == false, TEXT("Synchronous API function %hs called during mass processing. Use asynchronous API instead."), __FUNCTION__);
+	CHECK_SYNC_API();
 	checkf(IsDuringEntityCreation() == false, TEXT("%hs: Trying to destroy entities while entity creation is under way. This operation is not supported."), __FUNCTION__);
 
+	TRACE_CPUPROFILER_EVENT_SCOPE(Mass_BatchDestroyEntities);
+
 	for (const FMassEntityHandle Entity : InEntities)
 	{
 		if (GetEntityStorageInterface().IsValidIndex(Entity.Index) == false)
@@ -916,11 +917,11 @@ void FMassEntityManager::BatchDestroyEntityChunks(const FMassArchetypeEntityColl
 
 void FMassEntityManager::BatchDestroyEntityChunks(TConstArrayView<FMassArchetypeEntityCollection> Collections)
 {
-	TRACE_CPUPROFILER_EVENT_SCOPE(Mass_BatchDestroyEntityChunks);
-
-	checkf(IsProcessing() == false, TEXT("Synchronous API function %hs called during mass processing. Use asynchronous API instead."), __FUNCTION__);
+	CHECK_SYNC_API();
 	checkf(IsDuringEntityCreation() == false, TEXT("%hs: Trying to destroy entities while entity creation is under way. This operation is not supported."), __FUNCTION__);
 
+	TRACE_CPUPROFILER_EVENT_SCOPE(Mass_BatchDestroyEntityChunks);
+
 	TArray<FMassEntityHandle> EntitiesRemoved;
 	// note that it's important to place the context instance in the same scope as the loop below that updates 
 	// FMassEntityManager.EntityData, otherwise, if there are commands flushed as part of FMassProcessingContext's 
@@ -951,6 +952,8 @@ void FMassEntityManager::BatchDestroyEntityChunks(TConstArrayView<FMassArchetype
 
 void FMassEntityManager::BatchGroupEntities(const UE::Mass::FArchetypeGroupHandle GroupHandle, TConstArrayView<FMassArchetypeEntityCollection> Collections)
 {
+	CHECK_SYNC_API();
+
 	if (GroupHandle.IsValid() == false)
 	{
 		UE_LOG(LogMass, Warning, TEXT("%hs called with an invalid GroupHandle"), __FUNCTION__);
@@ -996,6 +999,8 @@ void FMassEntityManager::BatchGroupEntities(const UE::Mass::FArchetypeGroupHandl
 
 void FMassEntityManager::RemoveEntityFromGroupType(FMassEntityHandle EntityHandle, UE::Mass::FArchetypeGroupType GroupType)
 {
+	CHECK_SYNC_API();
+
 	const FMassArchetypeHandle CurrentArchetypeHandle = GetArchetypeForEntity(EntityHandle);
 	if (FMassArchetypeData* CurrentArchetype = CurrentArchetypeHandle.DataPtr.Get())
 	{
@@ -1046,8 +1051,8 @@ const UE::Mass::FArchetypeGroups& FMassEntityManager::GetGroupsForArchetype(cons
 
 void FMassEntityManager::AddFragmentToEntity(FMassEntityHandle Entity, const UScriptStruct* FragmentType)
 {
+	CHECK_SYNC_API();
 	checkf(FragmentType, TEXT("Null fragment type passed in to %hs"), __FUNCTION__);
-	checkf(IsProcessing() == false, TEXT("Synchronous API function %hs called during mass processing. Use asynchronous API instead."), __FUNCTION__);
 
 	CheckIfEntityIsActive(Entity);
 
@@ -1061,8 +1066,8 @@ void FMassEntityManager::AddFragmentToEntity(FMassEntityHandle Entity, const USc
 
 void FMassEntityManager::AddFragmentToEntity(FMassEntityHandle Entity, const UScriptStruct* FragmentType, const FStructInitializationCallback& Initializer)
 {
+	CHECK_SYNC_API();
 	checkf(FragmentType, TEXT("Null fragment type passed in to %hs"), __FUNCTION__);
-	checkf(IsProcessing() == false, TEXT("Synchronous API function %hs called during mass processing. Use asynchronous API instead."), __FUNCTION__);
 
 	CheckIfEntityIsActive(Entity);
 
@@ -1082,6 +1087,8 @@ void FMassEntityManager::AddFragmentToEntity(FMassEntityHandle Entity, const USc
 
 void FMassEntityManager::AddFragmentListToEntity(FMassEntityHandle Entity, TConstArrayView<const UScriptStruct*> FragmentList)
 {
+	CHECK_SYNC_API();
+
 	CheckIfEntityIsActive(Entity);
 
 	const FMassArchetypeCompositionDescriptor Descriptor(InternalAddFragmentListToEntityChecked(Entity, FMassFragmentBitSet(FragmentList)));
@@ -1094,6 +1101,8 @@ void FMassEntityManager::AddFragmentListToEntity(FMassEntityHandle Entity, TCons
 
 void FMassEntityManager::AddCompositionToEntity_GetDelta(FMassEntityHandle Entity, FMassArchetypeCompositionDescriptor& InDescriptor)
 {
+	CHECK_SYNC_API();
+
 	CheckIfEntityIsActive(Entity);
 
 	FMassArchetypeData* OldArchetype = GetEntityStorageInterface().GetArchetype(Entity.Index);
@@ -1131,6 +1140,8 @@ void FMassEntityManager::AddCompositionToEntity_GetDelta(FMassEntityHandle Entit
 
 void FMassEntityManager::RemoveCompositionFromEntity(FMassEntityHandle Entity, const FMassArchetypeCompositionDescriptor& InDescriptor)
 {
+	CHECK_SYNC_API();
+
 	CheckIfEntityIsActive(Entity);
 
 	if(InDescriptor.IsEmpty() == false)
@@ -1230,7 +1241,7 @@ void FMassEntityManager::InternalAddFragmentListToEntity(FMassEntityHandle Entit
 
 void FMassEntityManager::AddFragmentInstanceListToEntity(FMassEntityHandle Entity, TConstArrayView<FInstancedStruct> FragmentInstanceList)
 {
-	checkf(IsProcessing() == false, TEXT("Synchronous API function %hs called during mass processing. Use asynchronous API instead."), __FUNCTION__);
+	CHECK_SYNC_API();
 
 	CheckIfEntityIsActive(Entity);
 	checkf(FragmentInstanceList.Num() > 0, TEXT("Need to specify at least one fragment instances for this operation"));
@@ -1254,7 +1265,7 @@ void FMassEntityManager::RemoveFragmentFromEntity(FMassEntityHandle Entity, cons
 
 void FMassEntityManager::RemoveFragmentListFromEntity(FMassEntityHandle Entity, TConstArrayView<const UScriptStruct*> FragmentList)
 {
-	checkf(IsProcessing() == false, TEXT("Synchronous API function %hs called during mass processing. Use asynchronous API instead."), __FUNCTION__);
+	CHECK_SYNC_API();
 
 	CheckIfEntityIsActive(Entity);
 	
@@ -1293,7 +1304,7 @@ void FMassEntityManager::RemoveFragmentListFromEntity(FMassEntityHandle Entity,
 
 void FMassEntityManager::SwapTagsForEntity(FMassEntityHandle Entity, const UScriptStruct* OldTagType, const UScriptStruct* NewTagType)
 {
-	checkf(IsProcessing() == false, TEXT("Synchronous API function %hs called during mass processing. Use asynchronous API instead."), __FUNCTION__);
+	CHECK_SYNC_API();
 
 	CheckIfEntityIsActive(Entity);
 
@@ -1321,6 +1332,7 @@ void FMassEntityManager::SwapTagsForEntity(FMassEntityHandle Entity, const UScri
 
 void FMassEntityManager::AddTagToEntity(FMassEntityHandle Entity, const UScriptStruct* TagType)
 {
+	CHECK_SYNC_API();
 	checkf((TagType != nullptr) && TagType->IsChildOf(FMassTag::StaticStruct()), TEXT("%hs works only with tags while '%s' is not one."), __FUNCTION__, *GetPathNameSafe(TagType));
 
 	CheckIfEntityIsActive(Entity);
@@ -1330,7 +1342,6 @@ void FMassEntityManager::AddTagToEntity(FMassEntityHandle Entity, const UScriptS
 
 	if (CurrentArchetype->HasTagType(TagType) == false)
 	{
-		//FMassTagBitSet NewTags = CurrentArchetype->GetTagBitSet() - *TagType;
 		FMassTagBitSet NewTags = CurrentArchetype->GetTagBitSet();
 		NewTags.Add(*TagType);
 		const FMassArchetypeHandle NewArchetypeHandle = InternalCreateSimilarArchetype(CurrentArchetype, NewTags);
@@ -1354,6 +1365,7 @@ void FMassEntityManager::AddTagToEntity(FMassEntityHandle Entity, const UScriptS
 	
 void FMassEntityManager::RemoveTagFromEntity(FMassEntityHandle Entity, const UScriptStruct* TagType)
 {
+	CHECK_SYNC_API();
 	checkf((TagType != nullptr) && TagType->IsChildOf(FMassTag::StaticStruct()), TEXT("%hs works only with tags while '%s' is not one."), __FUNCTION__, *GetPathNameSafe(TagType));
 
 	CheckIfEntityIsActive(Entity);
@@ -1373,7 +1385,6 @@ void FMassEntityManager::RemoveTagFromEntity(FMassEntityHandle Entity, const USc
 			ObserverManager.OnPreCompositionRemoved(Entity, CompositionDelta);
 		}
 		
-		// CurrentArchetype->GetTagBitSet() -  *TagType
 		const FMassTagBitSet NewTagComposition = CurrentArchetype->GetTagBitSet() - TagDelta;
 		const FMassArchetypeHandle NewArchetypeHandle = InternalCreateSimilarArchetype(CurrentArchetype, NewTagComposition);
 		checkSlow(NewArchetypeHandle.IsValid());
@@ -1386,6 +1397,8 @@ void FMassEntityManager::RemoveTagFromEntity(FMassEntityHandle Entity, const USc
 
 bool FMassEntityManager::AddConstSharedFragmentToEntity(const FMassEntityHandle Entity, const FConstSharedStruct& InConstSharedFragment)
 {
+	CHECK_SYNC_API();
+
 	if (!ensureMsgf(InConstSharedFragment.IsValid(), TEXT("%hs parameter Fragment is expected to be valid"), __FUNCTION__))
 	{
 		return false;
@@ -1434,6 +1447,8 @@ bool FMassEntityManager::AddConstSharedFragmentToEntity(const FMassEntityHandle
 
 bool FMassEntityManager::RemoveConstSharedFragmentFromEntity(const FMassEntityHandle Entity, const UScriptStruct& ConstSharedFragmentType)
 {
+	CHECK_SYNC_API();
+
 	if (!ensureMsgf(ConstSharedFragmentType.IsChildOf(FMassConstSharedFragment::StaticStruct()), TEXT("%hs parameter ConstSharedFragmentType is expected to be a FMassConstSharedFragment"), __FUNCTION__))
 	{
 		return false;
@@ -1475,6 +1490,8 @@ bool FMassEntityManager::RemoveConstSharedFragmentFromEntity(const FMassEntityHa
 
 bool FMassEntityManager::AddSharedFragmentToEntity(const FMassEntityHandle Entity, const FSharedStruct& InSharedFragment)
 {
+	CHECK_SYNC_API();
+
 	if (!ensureMsgf(InSharedFragment.IsValid(), TEXT("%hs parameter Fragment is expected to be valid"), __FUNCTION__))
 	{
 		return false;
@@ -1523,6 +1540,8 @@ bool FMassEntityManager::AddSharedFragmentToEntity(const FMassEntityHandle Entit
 
 bool FMassEntityManager::RemoveSharedFragmentFromEntity(const FMassEntityHandle Entity, const UScriptStruct& SharedFragmentType)
 {
+	CHECK_SYNC_API();
+
 	if (!ensureMsgf(SharedFragmentType.IsChildOf(FMassSharedFragment::StaticStruct()), TEXT("%hs parameter SharedFragmentType is expected to be a FMassSharedFragment"), __FUNCTION__))
 	{
 		return false;
@@ -1564,6 +1583,8 @@ bool FMassEntityManager::RemoveSharedFragmentFromEntity(const FMassEntityHandle
 
 void FMassEntityManager::BatchChangeTagsForEntities(TConstArrayView<FMassArchetypeEntityCollection> EntityCollections, const FMassTagBitSet& TagsToAdd, const FMassTagBitSet& TagsToRemove)
 {
+	CHECK_SYNC_API();
+
 	TRACE_CPUPROFILER_EVENT_SCOPE(Mass_BatchChangeTagsForEntities);
 
 	const FScopedCreationContextOperations CreationContextOperations(*this);
@@ -1618,6 +1639,8 @@ void FMassEntityManager::BatchChangeTagsForEntities(TConstArrayView<FMassArchety
 
 void FMassEntityManager::BatchChangeFragmentCompositionForEntities(TConstArrayView<FMassArchetypeEntityCollection> EntityCollections, const FMassFragmentBitSet& FragmentsToAdd, const FMassFragmentBitSet& FragmentsToRemove)
 {
+	CHECK_SYNC_API();
+
 	TRACE_CPUPROFILER_EVENT_SCOPE(Mass_BatchChangeFragmentCompositionForEntities);
 
 	FScopedCreationContextOperations CreationContextOperations(*this);
@@ -1677,6 +1700,8 @@ void FMassEntityManager::BatchChangeFragmentCompositionForEntities(TConstArrayVi
 
 void FMassEntityManager::BatchAddFragmentInstancesForEntities(TConstArrayView<FMassArchetypeEntityCollectionWithPayload> EntityCollections, const FMassFragmentBitSet& FragmentsAffected)
 {
+	CHECK_SYNC_API();
+
 	TRACE_CPUPROFILER_EVENT_SCOPE(Mass_BatchAddFragmentInstancesForEntities);
 
 	// here's the scenario:
@@ -1756,6 +1781,8 @@ void FMassEntityManager::BatchAddFragmentInstancesForEntities(TConstArrayView<FM
 void FMassEntityManager::BatchAddSharedFragmentsForEntities(TConstArrayView<FMassArchetypeEntityCollection> EntityCollections
 	, const FMassArchetypeSharedFragmentValues& AddedFragmentValues)
 {
+	CHECK_SYNC_API();
+
 	TRACE_CPUPROFILER_EVENT_SCOPE(Mass_BatchAddConstSharedFragmentForEntities);
 
 	for (const FMassArchetypeEntityCollection& Collection : EntityCollections)
@@ -1791,6 +1818,8 @@ void FMassEntityManager::BatchAddSharedFragmentsForEntities(TConstArrayView<FMas
 
 void FMassEntityManager::MoveEntityToAnotherArchetype(FMassEntityHandle Entity, FMassArchetypeHandle NewArchetypeHandle)
 {
+	CHECK_SYNC_API();
+
 	CheckIfEntityIsActive(Entity);
 
 	FMassArchetypeData& NewArchetype = FMassArchetypeHelper::ArchetypeDataFromHandleChecked(NewArchetypeHandle);
@@ -1802,7 +1831,7 @@ void FMassEntityManager::MoveEntityToAnotherArchetype(FMassEntityHandle Entity,
 	GetEntityStorageInterface().SetArchetypeFromShared(Entity.Index, NewArchetypeHandle.DataPtr);
 }
 
-void FMassEntityManager::SetEntityFragmentsValues(FMassEntityHandle Entity, TArrayView<const FInstancedStruct> FragmentInstanceList)
+void FMassEntityManager::SetEntityFragmentValues(FMassEntityHandle Entity, TArrayView<const FInstancedStruct> FragmentInstanceList)
 {
 	CheckIfEntityIsActive(Entity);
 
@@ -1811,16 +1840,18 @@ void FMassEntityManager::SetEntityFragmentsValues(FMassEntityHandle Entity, TArr
 	CurrentArchetype->SetFragmentsData(Entity, FragmentInstanceList);
 }
 
-void FMassEntityManager::BatchSetEntityFragmentsValues(const FMassArchetypeEntityCollection& SparseEntities, TArrayView<const FInstancedStruct> FragmentInstanceList)
+void FMassEntityManager::BatchSetEntityFragmentValues(const FMassArchetypeEntityCollection& SparseEntities, TArrayView<const FInstancedStruct> FragmentInstanceList)
 {
 	if (FragmentInstanceList.Num())
 	{
-		BatchSetEntityFragmentsValues(MakeArrayView(&SparseEntities, 1), FragmentInstanceList);
+		BatchSetEntityFragmentValues(MakeArrayView(&SparseEntities, 1), FragmentInstanceList);
 	}
 }
 
-void FMassEntityManager::BatchSetEntityFragmentsValues(TConstArrayView<FMassArchetypeEntityCollection> EntityCollections, TArrayView<const FInstancedStruct> FragmentInstanceList)
+void FMassEntityManager::BatchSetEntityFragmentValues(TConstArrayView<FMassArchetypeEntityCollection> EntityCollections, TArrayView<const FInstancedStruct> FragmentInstanceList)
 {
+	CHECK_SYNC_API();
+
 	if (FragmentInstanceList.IsEmpty())
 	{
 		return;
@@ -1828,7 +1859,7 @@ void FMassEntityManager::BatchSetEntityFragmentsValues(TConstArrayView<FMassArch
 
 	for (const FMassArchetypeEntityCollection& SparseEntities : EntityCollections)
 	{
-		TRACE_CPUPROFILER_EVENT_SCOPE(Mass_BatchSetEntityFragmentsValues);
+		TRACE_CPUPROFILER_EVENT_SCOPE(Mass_BatchSetEntityFragmentValues);
 
 		FMassArchetypeData* Archetype = SparseEntities.GetArchetype().DataPtr.Get();
 		check(Archetype);
@@ -2235,3 +2266,20 @@ FMassArchetypeHandle FMassEntityManager::InternalCreateSimilarArchetype(const FM
 {
 	return InternalCreateSimilarArchetype(SourceArchetypeRef, Forward<FMassArchetypeCompositionDescriptor>(NewComposition), {});
 }
+
+void FMassEntityManager::SetEntityFragmentsValues(FMassEntityHandle Entity, TArrayView<const FInstancedStruct> FragmentInstanceList)
+{
+	SetEntityFragmentValues(Entity, FragmentInstanceList);
+}
+
+void FMassEntityManager::BatchSetEntityFragmentsValues(const FMassArchetypeEntityCollection& SparseEntities, TArrayView<const FInstancedStruct> FragmentInstanceList)
+{
+	ensureMsgf(false, TEXT("The static BatchSetEntityFragmentsValues is not expected to be called anymore. There's no way to deduce the EntityManager instance related to the call"));
+}
+
+void FMassEntityManager::BatchSetEntityFragmentsValues(TConstArrayView<FMassArchetypeEntityCollection> EntityCollections, TArrayView<const FInstancedStruct> FragmentInstanceList)
+{
+	ensureMsgf(false, TEXT("The static BatchSetEntityFragmentsValues is not expected to be called anymore. There's no way to deduce the EntityManager instance related to the call"));
+}
+
+#undef CHECK_SYNC_API
\ No newline at end of file
diff --git a/Engine/Source/Runtime/MassEntity/Public/MassEntityManager.h b/Engine/Source/Runtime/MassEntity/Public/MassEntityManager.h
--- a/Engine/Source/Runtime/MassEntity/Public/MassEntityManager.h
+++ b/Engine/Source/Runtime/MassEntity/Public/MassEntityManager.h
@@ -479,14 +479,14 @@ public:
 
 	/** Copies values from FragmentInstanceList over to Entity's fragment. Caller is responsible for ensuring that 
 	 *  the given entity does have given fragments. Failing this assumption will cause a check-fail.*/
-	void SetEntityFragmentsValues(FMassEntityHandle Entity, TArrayView<const FInstancedStruct> FragmentInstanceList);
+	void SetEntityFragmentValues(FMassEntityHandle Entity, TArrayView<const FInstancedStruct> FragmentInstanceList);
 
 	/** Copies values from FragmentInstanceList over to fragments of given entities collection. The caller is responsible 
 	 *  for ensuring that the given entity archetype (FMassArchetypeEntityCollection .Archetype) does have given fragments. 
 	 *  Failing this assumption will cause a check-fail. */
-	static void BatchSetEntityFragmentsValues(const FMassArchetypeEntityCollection& SparseEntities, TArrayView<const FInstancedStruct> FragmentInstanceList);
+	void BatchSetEntityFragmentValues(const FMassArchetypeEntityCollection& SparseEntities, TArrayView<const FInstancedStruct> FragmentInstanceList);
 
-	static void BatchSetEntityFragmentsValues(TConstArrayView<FMassArchetypeEntityCollection> EntityCollections, TArrayView<const FInstancedStruct> FragmentInstanceList);
+	void BatchSetEntityFragmentValues(TConstArrayView<FMassArchetypeEntityCollection> EntityCollections, TArrayView<const FInstancedStruct> FragmentInstanceList);
 
 	// Return true if it is an valid built entity
 	bool IsEntityActive(FMassEntityHandle Entity) const 
@@ -1248,4 +1248,13 @@ private:
 	//-----------------------------------------------------------------------------
 	UE_DEPRECATED(5.6, "This flavor of InternalCreateSimilarArchetype is deprecated due to the introduction of archetype grouping. Use InternalCreateSimilarArchetype with a FArchetypeGroups parameter instead")
 	FMassArchetypeHandle InternalCreateSimilarArchetype(const FMassArchetypeData& SourceArchetypeRef, FMassArchetypeCompositionDescriptor&& NewComposition);
+
+	UE_DEPRECATED(5.6, "SetEntityFragmentsValues is deprecated. Use SetEntityFragmentValues instead (note the slight change in name).")
+	void SetEntityFragmentsValues(FMassEntityHandle Entity, TArrayView<const FInstancedStruct> FragmentInstanceList);
+
+	UE_DEPRECATED(5.6, "Static BatchSetEntityFragmentsValues is deprecated. Use EntityManager's member function BatchSetEntityFragmentValues (note the slight change in name).")
+	static void BatchSetEntityFragmentsValues(const FMassArchetypeEntityCollection& SparseEntities, TArrayView<const FInstancedStruct> FragmentInstanceList);
+
+	UE_DEPRECATED(5.6, "Static BatchSetEntityFragmentsValues is deprecated. Use EntityManager's member function BatchSetEntityFragmentValues (note the slight change in name).")
+	static void BatchSetEntityFragmentsValues(TConstArrayView<FMassArchetypeEntityCollection> EntityCollections, TArrayView<const FInstancedStruct> FragmentInstanceList);
 };
