[Mass] Deprecated a bunch of MassEntityQuery functions that take the MassEntityManager parameter - now obsolete due to the query having access to CachedEntityManager now

Additionally, modified MassExecutionContext to store the query currently being executed, mainly for safety checks but also for facilitating changes requiring access to the current query instance inside any FMassExecuteFunction

diff --git a/Engine/Plugins/AI/MassAI/Source/MassAIBehavior/Private/MassLookAtProcessors.cpp b/Engine/Plugins/AI/MassAI/Source/MassAIBehavior/Private/MassLookAtProcessors.cpp
--- a/Engine/Plugins/AI/MassAI/Source/MassAIBehavior/Private/MassLookAtProcessors.cpp
+++ b/Engine/Plugins/AI/MassAI/Source/MassAIBehavior/Private/MassLookAtProcessors.cpp
@@ -123,7 +123,7 @@ void UMassLookAtProcessor::Execute(FMassEntityManager& EntityManager, FMassExecu
 
 	const double CurrentTime = GetWorld()->GetTimeSeconds();
 
-	EntityQuery_Conditional.ForEachEntityChunk(EntityManager, Context, [this, &EntityManager, CurrentTime](FMassExecutionContext& Context)
+	EntityQuery_Conditional.ForEachEntityChunk(Context, [this, &EntityManager, CurrentTime](FMassExecutionContext& Context)
 		{
 			const UMassNavigationSubsystem& MassNavSystem = Context.GetSubsystemChecked<UMassNavigationSubsystem>();
 			const UZoneGraphSubsystem& ZoneGraphSubsystem = Context.GetSubsystemChecked<UZoneGraphSubsystem>();
diff --git a/Engine/Plugins/AI/MassAI/Source/MassAIBehavior/Private/MassStateTreeProcessors.cpp b/Engine/Plugins/AI/MassAI/Source/MassAIBehavior/Private/MassStateTreeProcessors.cpp
--- a/Engine/Plugins/AI/MassAI/Source/MassAIBehavior/Private/MassStateTreeProcessors.cpp
+++ b/Engine/Plugins/AI/MassAI/Source/MassAIBehavior/Private/MassStateTreeProcessors.cpp
@@ -98,7 +98,7 @@ void UMassStateTreeFragmentDestructor::Execute(FMassEntityManager& EntityManager
 		return;
 	}
 
-	EntityQuery.ForEachEntityChunk(EntityManager, Context,
+	EntityQuery.ForEachEntityChunk(Context,
 		[](FMassExecutionContext& Context)
 		{
 			UMassStateTreeSubsystem& MassStateTreeSubsystem = Context.GetMutableSubsystemChecked<UMassStateTreeSubsystem>();
@@ -162,7 +162,7 @@ void UMassStateTreeActivationProcessor::Execute(FMassEntityManager& EntityManage
 	TArray<FMassEntityHandle> EntitiesToSignal;
 	int32 ActivationCounts[EMassLOD::Max] {0,0,0,0};
 	
-	EntityQuery.ForEachEntityChunk(EntityManager, Context,
+	EntityQuery.ForEachEntityChunk(Context,
 		[&EntitiesToSignal, &ActivationCounts, MaxActivationsPerLOD = BehaviorSettings->MaxActivationsPerLOD, TimeInSeconds](FMassExecutionContext& Context)
 		{
 			UMassStateTreeSubsystem& MassStateTreeSubsystem = Context.GetMutableSubsystemChecked<UMassStateTreeSubsystem>();
@@ -337,11 +337,11 @@ void UMassStateTreeProcessor::SignalEntities(FMassEntityManager& EntityManager,
 
 	if (bProcessEntitiesInParallel)
 	{
-		EntityQuery.ParallelForEachEntityChunk(EntityManager, Context, TickFunc);
+		EntityQuery.ParallelForEachEntityChunk(Context, TickFunc);
 	}
 	else
 	{
-		EntityQuery.ForEachEntityChunk(EntityManager, Context, TickFunc);
+		EntityQuery.ForEachEntityChunk(Context, TickFunc);
 	}
 
 	if (EntitiesToSignal.Num())
diff --git a/Engine/Plugins/AI/MassAI/Source/MassAIBehavior/Private/MassZoneGraphAnnotationProcessors.cpp b/Engine/Plugins/AI/MassAI/Source/MassAIBehavior/Private/MassZoneGraphAnnotationProcessors.cpp
--- a/Engine/Plugins/AI/MassAI/Source/MassAIBehavior/Private/MassZoneGraphAnnotationProcessors.cpp
+++ b/Engine/Plugins/AI/MassAI/Source/MassAIBehavior/Private/MassZoneGraphAnnotationProcessors.cpp
@@ -32,7 +32,7 @@ void UMassZoneGraphAnnotationTagsInitializer::ConfigureQueries(const TSharedRef<
 
 void UMassZoneGraphAnnotationTagsInitializer::Execute(FMassEntityManager& EntityManager, FMassExecutionContext& Context)
 {
-	EntityQuery.ForEachEntityChunk(EntityManager, Context, [this](FMassExecutionContext& Context)
+	EntityQuery.ForEachEntityChunk(Context, [this](FMassExecutionContext& Context)
 	{
 		const UZoneGraphAnnotationSubsystem& ZoneGraphAnnotationSubsystem = Context.GetSubsystemChecked<UZoneGraphAnnotationSubsystem>();
 
@@ -97,7 +97,7 @@ void UMassZoneGraphAnnotationTagUpdateProcessor::Execute(FMassEntityManager& Ent
 	// Calling super will update the signals, and call SignalEntities() below.
 	Super::Execute(EntityManager, Context);
 
-	EntityQuery.ForEachEntityChunk(EntityManager, Context, [this](FMassExecutionContext& Context)
+	EntityQuery.ForEachEntityChunk(Context, [this](FMassExecutionContext& Context)
 	{
 		// Periodically update tags.
 		if (!FMassZoneGraphAnnotationVariableTickChunkFragment::UpdateChunk(Context))
@@ -148,7 +148,7 @@ void UMassZoneGraphAnnotationTagUpdateProcessor::UpdateAnnotationTags(UZoneGraph
 
 void UMassZoneGraphAnnotationTagUpdateProcessor::SignalEntities(FMassEntityManager& EntityManager, FMassExecutionContext& Context, FMassSignalNameLookup&)
 {
-	EntityQuery.ForEachEntityChunk(EntityManager, Context, [this](FMassExecutionContext& Context)
+	EntityQuery.ForEachEntityChunk(Context, [this](FMassExecutionContext& Context)
 	{
 		UZoneGraphAnnotationSubsystem& ZoneGraphAnnotationSubsystem = Context.GetMutableSubsystemChecked<UZoneGraphAnnotationSubsystem>();
 
diff --git a/Engine/Plugins/AI/MassAI/Source/MassAIDebug/Private/GameplayDebuggerCategory_Mass.cpp b/Engine/Plugins/AI/MassAI/Source/MassAIDebug/Private/GameplayDebuggerCategory_Mass.cpp
--- a/Engine/Plugins/AI/MassAI/Source/MassAIDebug/Private/GameplayDebuggerCategory_Mass.cpp
+++ b/Engine/Plugins/AI/MassAI/Source/MassAIDebug/Private/GameplayDebuggerCategory_Mass.cpp
@@ -238,7 +238,7 @@ void FGameplayDebuggerCategory_Mass::PickEntity(const FVector& ViewLocation, con
 		FMassExecutionContext ExecutionContext(EntityManager);
 		FMassEntityQuery Query(EntityManager.AsShared());
 		Query.AddRequirement<FTransformFragment>(EMassFragmentAccess::ReadOnly);
-		Query.ForEachEntityChunk(EntityManager, ExecutionContext, [&Entities, &Locations](FMassExecutionContext& Context)
+		Query.ForEachEntityChunk(ExecutionContext, [&Entities, &Locations](FMassExecutionContext& Context)
 		{
 			Entities.Append(Context.GetEntities().GetData(), Context.GetEntities().Num());
 			TConstArrayView<FTransformFragment> InLocations = Context.GetFragmentView<FTransformFragment>();
@@ -518,7 +518,7 @@ void FGameplayDebuggerCategory_Mass::CollectData(APlayerController* OwnerPC, AAc
 		{
 			FMassExecutionContext Context(EntityManager, 0.0f);
 		
-			EntityQuery.ForEachEntityChunk(EntityManager, Context, [this, MassStateTreeSubsystem, SmartObjectSubsystem, ViewLocation, ViewDirection, CurrentTime](FMassExecutionContext& Context)
+			EntityQuery.ForEachEntityChunk(Context, [this, MassStateTreeSubsystem, SmartObjectSubsystem, ViewLocation, ViewDirection, CurrentTime](FMassExecutionContext& Context)
 			{
 				FMassEntityManager& EntityManager = Context.GetEntityManagerChecked();
 
@@ -765,7 +765,7 @@ void FGameplayDebuggerCategory_Mass::CollectData(APlayerController* OwnerPC, AAc
 			EntityColliderQuery.AddRequirement<FMassAvoidanceColliderFragment>(EMassFragmentAccess::ReadOnly);
 			EntityColliderQuery.AddRequirement<FTransformFragment>(EMassFragmentAccess::ReadOnly);
 			FMassExecutionContext Context(EntityManager, 0.f);
-			EntityColliderQuery.ForEachEntityChunk(EntityManager, Context, [this, ViewLocation, ViewDirection](const FMassExecutionContext& Context)
+			EntityColliderQuery.ForEachEntityChunk(Context, [this, ViewLocation, ViewDirection](const FMassExecutionContext& Context)
 			{
 				const int32 NumEntities = Context.GetNumEntities();
 				const TConstArrayView<FTransformFragment> TransformList = Context.GetFragmentView<FTransformFragment>();
diff --git a/Engine/Plugins/AI/MassAI/Source/MassAIDebug/Private/MassDebugStateTreeProcessor.cpp b/Engine/Plugins/AI/MassAI/Source/MassAIDebug/Private/MassDebugStateTreeProcessor.cpp
--- a/Engine/Plugins/AI/MassAI/Source/MassAIDebug/Private/MassDebugStateTreeProcessor.cpp
+++ b/Engine/Plugins/AI/MassAI/Source/MassAIDebug/Private/MassDebugStateTreeProcessor.cpp
@@ -55,7 +55,7 @@ void UMassDebugStateTreeProcessor::Execute(FMassEntityManager& EntityManager, FM
 	}
 	
 	QUICK_SCOPE_CYCLE_COUNTER(UMassDebugStateTreeProcessor_Run);	
-	EntityQuery.ForEachEntityChunk(EntityManager, Context, [this, Debugger, MassStateTreeSubsystem](FMassExecutionContext& Context)
+	EntityQuery.ForEachEntityChunk(Context, [this, Debugger, MassStateTreeSubsystem](FMassExecutionContext& Context)
 	{
 		const FMassEntityHandle SelectedEntity = Debugger->GetSelectedEntity();
 		const int32 NumEntities = Context.GetNumEntities();
diff --git a/Engine/Plugins/AI/MassAI/Source/MassNavMeshNavigation/Private/MassNavMeshNavigationBoundaryProcessor.cpp b/Engine/Plugins/AI/MassAI/Source/MassNavMeshNavigation/Private/MassNavMeshNavigationBoundaryProcessor.cpp
--- a/Engine/Plugins/AI/MassAI/Source/MassNavMeshNavigation/Private/MassNavMeshNavigationBoundaryProcessor.cpp
+++ b/Engine/Plugins/AI/MassAI/Source/MassNavMeshNavigation/Private/MassNavMeshNavigationBoundaryProcessor.cpp
@@ -1,4 +1,4 @@
-﻿// Copyright Epic Games, Inc. All Rights Reserved.
+// Copyright Epic Games, Inc. All Rights Reserved.
 
 #include "MassNavMeshNavigationBoundaryProcessor.h"
 
@@ -33,7 +33,7 @@ void UMassNavMeshNavigationBoundaryProcessor::Execute(FMassEntityManager& Entity
 {
 	QUICK_SCOPE_CYCLE_COUNTER(UMassNavMeshNavigationBoundaryProcessor);
 
-	EntityQuery.ForEachEntityChunk(EntityManager, Context, [this](FMassExecutionContext& Context)
+	EntityQuery.ForEachEntityChunk(Context, [this](FMassExecutionContext& Context)
 	{
 		const int32 NumEntities = Context.GetNumEntities();
 
diff --git a/Engine/Plugins/AI/MassAI/Source/MassNavMeshNavigation/Private/MassNavMeshNavigationProcessors.cpp b/Engine/Plugins/AI/MassAI/Source/MassNavMeshNavigation/Private/MassNavMeshNavigationProcessors.cpp
--- a/Engine/Plugins/AI/MassAI/Source/MassNavMeshNavigation/Private/MassNavMeshNavigationProcessors.cpp
+++ b/Engine/Plugins/AI/MassAI/Source/MassNavMeshNavigation/Private/MassNavMeshNavigationProcessors.cpp
@@ -1,4 +1,4 @@
-﻿// Copyright Epic Games, Inc. All Rights Reserved.
+// Copyright Epic Games, Inc. All Rights Reserved.
 
 #include "MassNavMeshNavigationProcessors.h"
 #include "Avoidance/MassAvoidanceFragments.h"
@@ -47,7 +47,7 @@ void UMassNavMeshPathFollowProcessor::Execute(FMassEntityManager& EntityManager,
 	
 	TArray<FMassEntityHandle> EntitiesToSignalPathDone;
 
-	EntityQuery_Conditional.ForEachEntityChunk(EntityManager, Context, [this, &EntitiesToSignalPathDone](FMassExecutionContext& Context)
+	EntityQuery_Conditional.ForEachEntityChunk(Context, [this, &EntitiesToSignalPathDone](FMassExecutionContext& Context)
 	{
 		const int32 NumEntities = Context.GetNumEntities();
 		const TArrayView<FMassNavMeshShortPathFragment> ShortPathList = Context.GetMutableFragmentView<FMassNavMeshShortPathFragment>();
diff --git a/Engine/Plugins/AI/MassAI/Source/MassNavigation/Private/MassAvoidanceProcessors.cpp b/Engine/Plugins/AI/MassAI/Source/MassNavigation/Private/MassAvoidanceProcessors.cpp
--- a/Engine/Plugins/AI/MassAI/Source/MassNavigation/Private/MassAvoidanceProcessors.cpp
+++ b/Engine/Plugins/AI/MassAI/Source/MassNavigation/Private/MassAvoidanceProcessors.cpp
@@ -387,7 +387,7 @@ void UMassMovingAvoidanceProcessor::Execute(FMassEntityManager& EntityManager, F
 		return;
 	}
 
-	EntityQuery.ForEachEntityChunk(EntityManager, Context, [this, &EntityManager](FMassExecutionContext& Context)
+	EntityQuery.ForEachEntityChunk(Context, [this, &EntityManager](FMassExecutionContext& Context)
 	{
 		const float DeltaTime = Context.GetDeltaTimeSeconds();
 		const double CurrentTime = World->GetTimeSeconds();
@@ -1081,7 +1081,7 @@ void UMassStandingAvoidanceProcessor::Execute(FMassEntityManager& EntityManager,
 	}
 
 	// Avoidance while standing
-	EntityQuery.ForEachEntityChunk(EntityManager, Context, [this, &EntityManager](FMassExecutionContext& Context)
+	EntityQuery.ForEachEntityChunk(Context, [this, &EntityManager](FMassExecutionContext& Context)
 	{
 		const int32 NumEntities = Context.GetNumEntities();
 		const float DeltaTime = Context.GetDeltaTimeSeconds();
diff --git a/Engine/Plugins/AI/MassAI/Source/MassNavigation/Private/MassNavigationProcessors.cpp b/Engine/Plugins/AI/MassAI/Source/MassNavigation/Private/MassNavigationProcessors.cpp
--- a/Engine/Plugins/AI/MassAI/Source/MassNavigation/Private/MassNavigationProcessors.cpp
+++ b/Engine/Plugins/AI/MassAI/Source/MassNavigation/Private/MassNavigationProcessors.cpp
@@ -41,7 +41,7 @@ void UMassOffLODNavigationProcessor::ConfigureQueries(const TSharedRef<FMassEnti
 void UMassOffLODNavigationProcessor::Execute(FMassEntityManager& EntityManager,
 													FMassExecutionContext& Context)
 {
-	EntityQuery_Conditional.ForEachEntityChunk(EntityManager, Context, [this](FMassExecutionContext& Context)
+	EntityQuery_Conditional.ForEachEntityChunk(Context, [this](FMassExecutionContext& Context)
 	{
 #if WITH_MASSGAMEPLAY_DEBUG
 		if (UE::MassMovement::bFreezeMovement)
@@ -88,7 +88,7 @@ void UMassNavigationSmoothHeightProcessor::ConfigureQueries(const TSharedRef<FMa
 void UMassNavigationSmoothHeightProcessor::Execute(FMassEntityManager& EntityManager,
 													FMassExecutionContext& Context)
 {
-	EntityQuery.ForEachEntityChunk(EntityManager, Context, [this](FMassExecutionContext& Context)
+	EntityQuery.ForEachEntityChunk(Context, [this](FMassExecutionContext& Context)
 	{
 #if WITH_MASSGAMEPLAY_DEBUG
 		if (UE::MassMovement::bFreezeMovement)
@@ -138,7 +138,7 @@ void UMassMoveTargetFragmentInitializer::ConfigureQueries(const TSharedRef<FMass
 
 void UMassMoveTargetFragmentInitializer::Execute(FMassEntityManager& EntityManager, FMassExecutionContext& Context)
 {
-	InitializerQuery.ForEachEntityChunk(EntityManager, Context, [](FMassExecutionContext& Context)
+	InitializerQuery.ForEachEntityChunk(Context, [](FMassExecutionContext& Context)
 	{
 		const int32 NumEntities = Context.GetNumEntities();
 		const TArrayView<FMassMoveTargetFragment> MoveTargetList = Context.GetMutableFragmentView<FMassMoveTargetFragment>();
@@ -196,7 +196,7 @@ void UMassNavigationObstacleGridProcessor::ConfigureQueries(const TSharedRef<FMa
 void UMassNavigationObstacleGridProcessor::Execute(FMassEntityManager& EntityManager, FMassExecutionContext& Context)
 {
 	// can't be ParallelFor due to MovementSubsystem->GetGridMutable().Move not being thread-safe
-	AddToGridEntityQuery.ForEachEntityChunk(EntityManager, Context, [this, &EntityManager](FMassExecutionContext& Context)
+	AddToGridEntityQuery.ForEachEntityChunk(Context, [this, &EntityManager](FMassExecutionContext& Context)
 	{
 		FNavigationObstacleHashGrid2D& HashGrid = Context.GetMutableSubsystemChecked<UMassNavigationSubsystem>().GetObstacleGridMutable();
 		const int32 NumEntities = Context.GetNumEntities();
@@ -223,7 +223,7 @@ void UMassNavigationObstacleGridProcessor::Execute(FMassEntityManager& EntityMan
 		}
 	});
 
-	UpdateGridEntityQuery.ForEachEntityChunk(EntityManager, Context, [this, &EntityManager](FMassExecutionContext& Context)
+	UpdateGridEntityQuery.ForEachEntityChunk(Context, [this, &EntityManager](FMassExecutionContext& Context)
 	{
 		FNavigationObstacleHashGrid2D& HashGrid = Context.GetMutableSubsystemChecked<UMassNavigationSubsystem>().GetObstacleGridMutable();
 		const int32 NumEntities = Context.GetNumEntities();
@@ -257,7 +257,7 @@ void UMassNavigationObstacleGridProcessor::Execute(FMassEntityManager& EntityMan
 		}
 	});
 
-	RemoveFromGridEntityQuery.ForEachEntityChunk(EntityManager, Context, [this, &EntityManager](FMassExecutionContext& Context)
+	RemoveFromGridEntityQuery.ForEachEntityChunk(Context, [this, &EntityManager](FMassExecutionContext& Context)
 	{
 		FNavigationObstacleHashGrid2D& HashGrid = Context.GetMutableSubsystemChecked<UMassNavigationSubsystem>().GetObstacleGridMutable();
 		const int32 NumEntities = Context.GetNumEntities();
@@ -294,7 +294,7 @@ void UMassNavigationObstacleRemoverProcessor::ConfigureQueries(const TSharedRef<
 
 void UMassNavigationObstacleRemoverProcessor::Execute(FMassEntityManager& EntityManager, FMassExecutionContext& Context)
 {
-	EntityQuery.ForEachEntityChunk(EntityManager, Context, [this](FMassExecutionContext& Context)
+	EntityQuery.ForEachEntityChunk(Context, [this](FMassExecutionContext& Context)
 	{
 		FNavigationObstacleHashGrid2D& HashGrid = Context.GetMutableSubsystemChecked<UMassNavigationSubsystem>().GetObstacleGridMutable();
 		const int32 NumEntities = Context.GetNumEntities();
diff --git a/Engine/Plugins/AI/MassAI/Source/MassNavigation/Private/MassSmoothOrientationProcessors.cpp b/Engine/Plugins/AI/MassAI/Source/MassNavigation/Private/MassSmoothOrientationProcessors.cpp
--- a/Engine/Plugins/AI/MassAI/Source/MassNavigation/Private/MassSmoothOrientationProcessors.cpp
+++ b/Engine/Plugins/AI/MassAI/Source/MassNavigation/Private/MassSmoothOrientationProcessors.cpp
@@ -49,7 +49,7 @@ void UMassSmoothOrientationProcessor::Execute(FMassEntityManager& EntityManager,
 	{
 		QUICK_SCOPE_CYCLE_COUNTER(HighRes);
 
-		HighResEntityQuery.ForEachEntityChunk(EntityManager, Context, [this, DeltaTime](FMassExecutionContext& Context)
+		HighResEntityQuery.ForEachEntityChunk(Context, [this, DeltaTime](FMassExecutionContext& Context)
 		{
 			const int32 NumEntities = Context.GetNumEntities();
 
@@ -118,7 +118,7 @@ void UMassSmoothOrientationProcessor::Execute(FMassEntityManager& EntityManager,
 	{
 		QUICK_SCOPE_CYCLE_COUNTER(LowRes);
 
-		LowResEntityQuery_Conditional.ForEachEntityChunk(EntityManager, Context, [this](FMassExecutionContext& Context)
+		LowResEntityQuery_Conditional.ForEachEntityChunk(Context, [this](FMassExecutionContext& Context)
 		{
 			const int32 NumEntities = Context.GetNumEntities();
 
diff --git a/Engine/Plugins/AI/MassAI/Source/MassNavigation/Private/MassSteeringProcessors.cpp b/Engine/Plugins/AI/MassAI/Source/MassNavigation/Private/MassSteeringProcessors.cpp
--- a/Engine/Plugins/AI/MassAI/Source/MassNavigation/Private/MassSteeringProcessors.cpp
+++ b/Engine/Plugins/AI/MassAI/Source/MassNavigation/Private/MassSteeringProcessors.cpp
@@ -78,7 +78,7 @@ void UMassSteerToMoveTargetProcessor::ConfigureQueries(const TSharedRef<FMassEnt
 
 void UMassSteerToMoveTargetProcessor::Execute(FMassEntityManager& EntityManager, FMassExecutionContext& Context)
 {
-	EntityQuery.ForEachEntityChunk(EntityManager, Context, [this](FMassExecutionContext& Context)
+	EntityQuery.ForEachEntityChunk(Context, [this](FMassExecutionContext& Context)
 	{
 		const int32 NumEntities = Context.GetNumEntities();
 		const TConstArrayView<FTransformFragment> TransformList = Context.GetFragmentView<FTransformFragment>();
diff --git a/Engine/Plugins/AI/MassAI/Source/MassZoneGraphNavigation/Private/MassZoneGraphNavigationProcessors.cpp b/Engine/Plugins/AI/MassAI/Source/MassZoneGraphNavigation/Private/MassZoneGraphNavigationProcessors.cpp
--- a/Engine/Plugins/AI/MassAI/Source/MassZoneGraphNavigation/Private/MassZoneGraphNavigationProcessors.cpp
+++ b/Engine/Plugins/AI/MassAI/Source/MassZoneGraphNavigation/Private/MassZoneGraphNavigationProcessors.cpp
@@ -57,7 +57,7 @@ void UMassZoneGraphLocationInitializer::ConfigureQueries(const TSharedRef<FMassE
 
 void UMassZoneGraphLocationInitializer::Execute(FMassEntityManager& EntityManager, FMassExecutionContext& Context)
 {
-	EntityQuery.ForEachEntityChunk(EntityManager, Context, [this](FMassExecutionContext& Context)
+	EntityQuery.ForEachEntityChunk(Context, [this](FMassExecutionContext& Context)
 	{
 		const UZoneGraphSubsystem& ZoneGraphSubsystem = Context.GetSubsystemChecked<UZoneGraphSubsystem>();
 
@@ -151,7 +151,7 @@ void UMassZoneGraphPathFollowProcessor::Execute(FMassEntityManager& EntityManage
 	TArray<FMassEntityHandle> EntitiesToSignalPathDone;
 	TArray<FMassEntityHandle> EntitiesToSignalLaneChanged;
 
-	EntityQuery_Conditional.ForEachEntityChunk(EntityManager, Context, [this, &EntitiesToSignalPathDone, &EntitiesToSignalLaneChanged](FMassExecutionContext& Context)
+	EntityQuery_Conditional.ForEachEntityChunk(Context, [this, &EntitiesToSignalPathDone, &EntitiesToSignalLaneChanged](FMassExecutionContext& Context)
 	{
 		const UZoneGraphSubsystem& ZoneGraphSubsystem = Context.GetSubsystemChecked<UZoneGraphSubsystem>();
 
@@ -450,7 +450,7 @@ void UMassZoneGraphLaneCacheBoundaryProcessor::Execute(FMassEntityManager& Entit
 		return;
 	}
 
-	EntityQuery.ForEachEntityChunk(EntityManager, Context, [this, World](FMassExecutionContext& Context)
+	EntityQuery.ForEachEntityChunk(Context, [this, World](FMassExecutionContext& Context)
 	{
 		const int32 NumEntities = Context.GetNumEntities();
 
diff --git a/Engine/Plugins/AI/MassCrowd/Source/MassCrowd/Private/MassCrowdNavigationProcessor.cpp b/Engine/Plugins/AI/MassCrowd/Source/MassCrowd/Private/MassCrowdNavigationProcessor.cpp
--- a/Engine/Plugins/AI/MassCrowd/Source/MassCrowd/Private/MassCrowdNavigationProcessor.cpp
+++ b/Engine/Plugins/AI/MassCrowd/Source/MassCrowd/Private/MassCrowdNavigationProcessor.cpp
@@ -41,7 +41,7 @@ void UMassCrowdLaneTrackingSignalProcessor::InitializeInternal(UObject& Owner, c
 
 void UMassCrowdLaneTrackingSignalProcessor::SignalEntities(FMassEntityManager& EntityManager, FMassExecutionContext& Context, FMassSignalNameLookup&)
 {
-	EntityQuery.ForEachEntityChunk(EntityManager, Context, [this](FMassExecutionContext& Context)
+	EntityQuery.ForEachEntityChunk(Context, [this](FMassExecutionContext& Context)
 	{
 		UMassCrowdSubsystem& MassCrowdSubsystem = Context.GetMutableSubsystemChecked<UMassCrowdSubsystem>();
 		const int32 NumEntities = Context.GetNumEntities();
@@ -81,7 +81,7 @@ void UMassCrowdLaneTrackingDestructor::ConfigureQueries(const TSharedRef<FMassEn
 
 void UMassCrowdLaneTrackingDestructor::Execute(FMassEntityManager& EntityManager, FMassExecutionContext& Context)
 {
-	EntityQuery.ForEachEntityChunk(EntityManager, Context, [this](FMassExecutionContext& Context)
+	EntityQuery.ForEachEntityChunk(Context, [this](FMassExecutionContext& Context)
 	{
 		UMassCrowdSubsystem& MassCrowdSubsystem = Context.GetMutableSubsystemChecked<UMassCrowdSubsystem>();
 		const int32 NumEntities = Context.GetNumEntities();
@@ -135,7 +135,7 @@ void UMassCrowdDynamicObstacleProcessor::Execute(FMassEntityManager& EntityManag
 	const UMassCrowdSettings* CrowdSettings = GetDefault<UMassCrowdSettings>();
 	checkf(CrowdSettings, TEXT("Settings default object is always expected to be valid"));
 
-	EntityQuery_Conditional.ForEachEntityChunk(EntityManager, Context, [this, World, CrowdSettings](FMassExecutionContext& Context)
+	EntityQuery_Conditional.ForEachEntityChunk(Context, [this, World, CrowdSettings](FMassExecutionContext& Context)
 	{
 		const int32 NumEntities = Context.GetNumEntities();
 
@@ -244,7 +244,7 @@ void UMassCrowdDynamicObstacleInitializer::Execute(FMassEntityManager& EntityMan
 {
 	UWorld* World = EntityManager.GetWorld();
 	
-	EntityQuery.ForEachEntityChunk(EntityManager, Context, [World](FMassExecutionContext& Context)
+	EntityQuery.ForEachEntityChunk(Context, [World](FMassExecutionContext& Context)
 	{
 		const int32 NumEntities = Context.GetNumEntities();
 		const TConstArrayView<FTransformFragment> LocationList = Context.GetFragmentView<FTransformFragment>();
@@ -291,7 +291,7 @@ void UMassCrowdDynamicObstacleDeinitializer::ConfigureQueries(const TSharedRef<F
 
 void UMassCrowdDynamicObstacleDeinitializer::Execute(FMassEntityManager& EntityManager, FMassExecutionContext& Context)
 {
-	EntityQuery.ForEachEntityChunk(EntityManager, Context, [this](FMassExecutionContext& Context)
+	EntityQuery.ForEachEntityChunk(Context, [this](FMassExecutionContext& Context)
 	{
 		const int32 NumEntities = Context.GetNumEntities();
 		const TArrayView<FMassCrowdObstacleFragment> ObstacleDataList = Context.GetMutableFragmentView<FMassCrowdObstacleFragment>();
diff --git a/Engine/Plugins/AI/MassCrowd/Source/MassCrowd/Private/MassCrowdServerRepresentationLODProcessor.cpp b/Engine/Plugins/AI/MassCrowd/Source/MassCrowd/Private/MassCrowdServerRepresentationLODProcessor.cpp
--- a/Engine/Plugins/AI/MassCrowd/Source/MassCrowd/Private/MassCrowdServerRepresentationLODProcessor.cpp
+++ b/Engine/Plugins/AI/MassCrowd/Source/MassCrowd/Private/MassCrowdServerRepresentationLODProcessor.cpp
@@ -61,7 +61,7 @@ void UMassCrowdServerRepresentationLODProcessor::Execute(FMassEntityManager& Ent
 	{
 		TRACE_CPUPROFILER_EVENT_SCOPE(CalculateLOD)
 		
-		EntityQuery.ForEachEntityChunk(EntityManager, Context, [this](FMassExecutionContext& Context)
+		EntityQuery.ForEachEntityChunk(Context, [this](FMassExecutionContext& Context)
 		{
 			const TConstArrayView<FMassViewerInfoFragment> ViewersInfoList = Context.GetFragmentView<FMassViewerInfoFragment>();
 			const TArrayView<FMassRepresentationLODFragment> RepresentationLODFragments = Context.GetMutableFragmentView<FMassRepresentationLODFragment>();
@@ -74,7 +74,7 @@ void UMassCrowdServerRepresentationLODProcessor::Execute(FMassEntityManager& Ent
 		
 		if (LODCalculator.AdjustDistancesFromCount())
 		{
-			EntityQuery.ForEachEntityChunk(EntityManager, Context, [this](FMassExecutionContext& Context)
+			EntityQuery.ForEachEntityChunk(Context, [this](FMassExecutionContext& Context)
 			{
 				const TConstArrayView<FMassViewerInfoFragment> ViewersInfoList = Context.GetFragmentView<FMassViewerInfoFragment>();
 				const TArrayView<FMassRepresentationLODFragment> RepresentationLODFragments = Context.GetMutableFragmentView<FMassRepresentationLODFragment>();
@@ -89,7 +89,7 @@ void UMassCrowdServerRepresentationLODProcessor::Execute(FMassEntityManager& Ent
 	{
 		TRACE_CPUPROFILER_EVENT_SCOPE(DebugDisplayLOD)
 		UWorld* World = EntityManager.GetWorld();
-		EntityQuery.ForEachEntityChunk(EntityManager, Context, [this, World](FMassExecutionContext& Context)
+		EntityQuery.ForEachEntityChunk(Context, [this, World](FMassExecutionContext& Context)
 		{
 			const TConstArrayView<FTransformFragment> LocationList = Context.GetFragmentView<FTransformFragment>();
 			const TConstArrayView<FMassRepresentationLODFragment> RepresentationLODFragments = Context.GetFragmentView<FMassRepresentationLODFragment>();
diff --git a/Engine/Plugins/AI/MassCrowd/Source/MassCrowd/Private/MassCrowdVisualizationLODProcessor.cpp b/Engine/Plugins/AI/MassCrowd/Source/MassCrowd/Private/MassCrowdVisualizationLODProcessor.cpp
--- a/Engine/Plugins/AI/MassCrowd/Source/MassCrowd/Private/MassCrowdVisualizationLODProcessor.cpp
+++ b/Engine/Plugins/AI/MassCrowd/Source/MassCrowd/Private/MassCrowdVisualizationLODProcessor.cpp
@@ -65,7 +65,7 @@ void UMassCrowdVisualizationLODProcessor::Execute(FMassEntityManager& EntityMana
 	{
 		TRACE_CPUPROFILER_EVENT_SCOPE(DebugDisplayLOD)
 
-		DebugEntityQuery.ForEachEntityChunk(EntityManager, Context, [World](FMassExecutionContext& Context)
+		DebugEntityQuery.ForEachEntityChunk(Context, [World](FMassExecutionContext& Context)
 		{
 			FMassVisualizationLODSharedFragment& LODSharedFragment = Context.GetMutableSharedFragment<FMassVisualizationLODSharedFragment>();
 			const TConstArrayView<FTransformFragment> LocationList = Context.GetFragmentView<FTransformFragment>();
@@ -79,7 +79,7 @@ void UMassCrowdVisualizationLODProcessor::Execute(FMassEntityManager& EntityMana
 	{
 		TRACE_CPUPROFILER_EVENT_SCOPE(ShowISMUnderSpecifiedRange)
 
-		DebugEntityQuery.ForEachEntityChunk(EntityManager, Context, [World](const FMassExecutionContext& Context)
+		DebugEntityQuery.ForEachEntityChunk(Context, [World](const FMassExecutionContext& Context)
 		{
 			const TConstArrayView<FTransformFragment> LocationList = Context.GetFragmentView<FTransformFragment>();
 			const TConstArrayView<FMassRepresentationFragment> RepresentationFragmentList = Context.GetFragmentView<FMassRepresentationFragment>();
diff --git a/Engine/Plugins/AI/MassCrowd/Source/MassCrowd/Private/MassCrowdVisualizationProcessor.cpp b/Engine/Plugins/AI/MassCrowd/Source/MassCrowd/Private/MassCrowdVisualizationProcessor.cpp
--- a/Engine/Plugins/AI/MassCrowd/Source/MassCrowd/Private/MassCrowdVisualizationProcessor.cpp
+++ b/Engine/Plugins/AI/MassCrowd/Source/MassCrowd/Private/MassCrowdVisualizationProcessor.cpp
@@ -80,7 +80,7 @@ void UMassDebugCrowdVisualizationProcessor::Execute(FMassEntityManager& EntityMa
 
 		TRACE_CPUPROFILER_EVENT_SCOPE(DebugDisplayVisualType)
 
-		EntityQuery.ForEachEntityChunk(EntityManager, Context, [this](const FMassExecutionContext& Context)
+		EntityQuery.ForEachEntityChunk(Context, [this](const FMassExecutionContext& Context)
 		{
 			const TConstArrayView<FMassRepresentationFragment> VisualizationList = Context.GetFragmentView<FMassRepresentationFragment>();
 			const TConstArrayView<FMassActorFragment> ActorList = Context.GetFragmentView<FMassActorFragment>();
diff --git a/Engine/Plugins/Experimental/EditorDataStorage/Source/TedsCore/Private/Processors/TypedElementProcessorAdaptors.cpp b/Engine/Plugins/Experimental/EditorDataStorage/Source/TedsCore/Private/Processors/TypedElementProcessorAdaptors.cpp
--- a/Engine/Plugins/Experimental/EditorDataStorage/Source/TedsCore/Private/Processors/TypedElementProcessorAdaptors.cpp
+++ b/Engine/Plugins/Experimental/EditorDataStorage/Source/TedsCore/Private/Processors/TypedElementProcessorAdaptors.cpp
@@ -899,7 +899,7 @@ protected:
 	{
 		if (Description.Callback.ActivationCount > 0)
 		{
-			NativeQuery.ForEachEntityChunk(Context.GetEntityManagerChecked(), Context,
+			NativeQuery.ForEachEntityChunk(Context,
 				[&Callback, &QueryStore, &Environment, &Description](FMassExecutionContext& ExecutionContext)
 				{
 					if (FTypedElementQueryProcessorData::PrepareCachedDependenciesOnQuery(Description, ExecutionContext))
@@ -1127,11 +1127,11 @@ UE::Editor::DataStorage::FQueryResult FTypedElementQueryProcessorData::Execute(
 					EnumHasAnyFlags(ExecutionFlags, EDirectQueryExecutionFlags::AutoBalanceParallelChunkProcessing)
 					? FMassEntityQuery::EParallelExecutionFlags::AutoBalance
 					: FMassEntityQuery::EParallelExecutionFlags::Default;
-				NativeQuery.ParallelForEachEntityChunk(EntityManager, Context, ExecuteFunction, Mode);
+				NativeQuery.ParallelForEachEntityChunk(Context, ExecuteFunction, Mode);
 			}
 			else
 			{
-				NativeQuery.ForEachEntityChunk(EntityManager, Context, ExecuteFunction);
+				NativeQuery.ForEachEntityChunk(Context, ExecuteFunction);
 			}
 		}
 	}
@@ -1164,7 +1164,7 @@ UE::Editor::DataStorage::FQueryResult FTypedElementQueryProcessorData::Execute(
 		Context.SetDeferredCommandBuffer(ParentContext.GetSharedDeferredCommandBuffer());
 		Context.SetFlushDeferredCommands(false);
 
-		NativeQuery.ForEachEntityChunk(EntityManager, Context,
+		NativeQuery.ForEachEntityChunk(Context,
 			[&Result, &Callback, &Description, &Environment](FMassExecutionContext& Context)
 			{
 				// No need to cache any subsystem dependencies as these are not accessible from a subquery.
@@ -1203,7 +1203,7 @@ UE::Editor::DataStorage::FQueryResult FTypedElementQueryProcessorData::Execute(
 		Context.SetDeferredCommandBuffer(ParentContext.GetSharedDeferredCommandBuffer());
 		Context.SetFlushDeferredCommands(false);
 
-		NativeQuery.ForEachEntityChunk(EntityManager, Context,
+		NativeQuery.ForEachEntityChunk(Context,
 			[&Result, &Callback, &Description, &Environment, &EntityManager, RowHandle](FMassExecutionContext& Context)
 			{
 				// No need to cache any subsystem dependencies as these are not accessible from a subquery.
@@ -1239,11 +1239,11 @@ void FTypedElementQueryProcessorData::Execute(FMassEntityManager& EntityManager,
 		
 		if (StoredQuery->Description.Callback.ExecutionMode != EExecutionMode::ThreadedChunks)
 		{
-			NativeQuery.ForEachEntityChunk(EntityManager, Context, ExceteFunction);
+			NativeQuery.ForEachEntityChunk(Context, ExceteFunction);
 		}
 		else
 		{
-			NativeQuery.ParallelForEachEntityChunk(EntityManager, Context, ExceteFunction);
+			NativeQuery.ParallelForEachEntityChunk(Context, ExceteFunction);
 		}
 	}
 }
diff --git a/Engine/Plugins/Experimental/EditorDataStorage/Source/TedsCore/Private/Queries/TypedElementExtendedQueryStore.cpp b/Engine/Plugins/Experimental/EditorDataStorage/Source/TedsCore/Private/Queries/TypedElementExtendedQueryStore.cpp
--- a/Engine/Plugins/Experimental/EditorDataStorage/Source/TedsCore/Private/Queries/TypedElementExtendedQueryStore.cpp
+++ b/Engine/Plugins/Experimental/EditorDataStorage/Source/TedsCore/Private/Queries/TypedElementExtendedQueryStore.cpp
@@ -337,7 +337,7 @@ namespace UE::Editor::DataStorage
 			case ActionType::Select:
 				// Fall through: There's nothing to callback to, so only return the total count.
 			case ActionType::Count:
-				Result.Count = QueryData->NativeQuery.GetNumMatchingEntities(EntityManager);
+				Result.Count = QueryData->NativeQuery.GetNumMatchingEntities();
 				Result.Completed = CompletionType::Fully;
 				break;
 			default:
@@ -391,7 +391,7 @@ namespace UE::Editor::DataStorage
 				break;
 			case FQueryDescription::EActionType::Count:
 				// Only the count is requested so no need to trigger the callback.
-				Result.Count = QueryData->NativeQuery.GetNumMatchingEntities(EntityManager);
+				Result.Count = QueryData->NativeQuery.GetNumMatchingEntities();
 				Result.Completed = FQueryResult::ECompletion::Fully;
 				break;
 			default:
diff --git a/Engine/Plugins/Runtime/InstancedActors/Source/InstancedActors/Private/InstancedActorsData.cpp b/Engine/Plugins/Runtime/InstancedActors/Source/InstancedActors/Private/InstancedActorsData.cpp
--- a/Engine/Plugins/Runtime/InstancedActors/Source/InstancedActors/Private/InstancedActorsData.cpp
+++ b/Engine/Plugins/Runtime/InstancedActors/Source/InstancedActors/Private/InstancedActorsData.cpp
@@ -456,7 +456,7 @@ void UInstancedActorsData::DespawnEntities()
 	FMassExecutionContext ExecutionContext(MassEntityManager);
 	for (FMassArchetypeEntityCollection& Collection : EntityCollectionsToDestroy)
 	{
-		InstancedActorLocationQuery.ForEachEntityChunk(Collection, MassEntityManager, ExecutionContext, [this, &WorldToLocalTranslation, &ManagerTransform](FMassExecutionContext& Context)
+		InstancedActorLocationQuery.ForEachEntityChunk(Collection, ExecutionContext, [this, &WorldToLocalTranslation, &ManagerTransform](FMassExecutionContext& Context)
 			{
 				int32 NumEntities = Context.GetNumEntities();
 				TConstArrayView<FInstancedActorsFragment> InstancedActorFragments = Context.GetFragmentView<FInstancedActorsFragment>();
diff --git a/Engine/Plugins/Runtime/InstancedActors/Source/InstancedActors/Private/InstancedActorsDebugProcessor.cpp b/Engine/Plugins/Runtime/InstancedActors/Source/InstancedActors/Private/InstancedActorsDebugProcessor.cpp
--- a/Engine/Plugins/Runtime/InstancedActors/Source/InstancedActors/Private/InstancedActorsDebugProcessor.cpp
+++ b/Engine/Plugins/Runtime/InstancedActors/Source/InstancedActors/Private/InstancedActorsDebugProcessor.cpp
@@ -123,7 +123,7 @@ namespace UE::InstancedActors::Debug
 					DetailedLODEntityQuery.AddSharedRequirement<FMassRepresentationSubsystemSharedFragment>(EMassFragmentAccess::ReadOnly);
 				
 					FMassExecutionContext Context = EntityManager.CreateExecutionContext(0);
-					DetailedLODEntityQuery.ForEachEntityChunk(EntityManager, Context, [World = InWorld](FMassExecutionContext& Context)
+					DetailedLODEntityQuery.ForEachEntityChunk(Context, [World = InWorld](FMassExecutionContext& Context)
 						{
 							//const FMassEntityHandle SelectedEntity = Debugger->GetSelectedEntity();
 							const int32 NumEntities = Context.GetNumEntities();
@@ -292,11 +292,11 @@ void UInstancedActorsDebugProcessor::Execute(FMassEntityManager& EntityManager,
 	
 	if (UE::InstancedActors::Debug::bDebugDrawAllEntities)
 	{
-		DebugAllEntityQuery.ForEachEntityChunk(EntityManager, Context, ExecFunction);
+		DebugAllEntityQuery.ForEachEntityChunk(Context, ExecFunction);
 	}
 	else
 	{
-		DetailedLODEntityQuery.ForEachEntityChunk(EntityManager, Context, ExecFunction);
+		DetailedLODEntityQuery.ForEachEntityChunk(Context, ExecFunction);
 	}
 }
 UE_ENABLE_OPTIMIZATION_SHIP
diff --git a/Engine/Plugins/Runtime/InstancedActors/Source/InstancedActors/Private/InstancedActorsInitializerProcessor.cpp b/Engine/Plugins/Runtime/InstancedActors/Source/InstancedActors/Private/InstancedActorsInitializerProcessor.cpp
--- a/Engine/Plugins/Runtime/InstancedActors/Source/InstancedActors/Private/InstancedActorsInitializerProcessor.cpp
+++ b/Engine/Plugins/Runtime/InstancedActors/Source/InstancedActors/Private/InstancedActorsInitializerProcessor.cpp
@@ -127,7 +127,7 @@ void UInstancedActorsInitializerProcessor::Execute(FMassEntityManager& EntityMan
 
 	TArray<FMassEntityHandle> EntitiesToSignal;
 
-	EntityQuery.ForEachEntityChunk(EntityManager, Context, [InstanceData, &NumInitializedEntities, &NextInstanceIndex, &ManagerLocation, &ManagerTransform, bApplyManagerTranslationOnly, bFilterInstanceTransforms, &EntitiesToSignal](FMassExecutionContext& Context)
+	EntityQuery.ForEachEntityChunk(Context, [InstanceData, &NumInitializedEntities, &NextInstanceIndex, &ManagerLocation, &ManagerTransform, bApplyManagerTranslationOnly, bFilterInstanceTransforms, &EntitiesToSignal](FMassExecutionContext& Context)
 	{
 		if (bApplyManagerTranslationOnly)
 		{
diff --git a/Engine/Plugins/Runtime/InstancedActors/Source/InstancedActors/Private/InstancedActorsManager.cpp b/Engine/Plugins/Runtime/InstancedActors/Source/InstancedActors/Private/InstancedActorsManager.cpp
--- a/Engine/Plugins/Runtime/InstancedActors/Source/InstancedActors/Private/InstancedActorsManager.cpp
+++ b/Engine/Plugins/Runtime/InstancedActors/Source/InstancedActors/Private/InstancedActorsManager.cpp
@@ -864,7 +864,7 @@ bool AInstancedActorsManager::ForEachInstance(FInstanceOperationFunc Operation,
 			UE::Mass::Utils::CreateEntityCollections(*MassEntityManager, InstanceData->Entities, FMassArchetypeEntityCollection::NoDuplicates, EntityCollections);
 
 			FMassExecutionContext ExecutionContext(*MassEntityManager);
-			InstancedActorLocationQuery.ForEachEntityChunkInCollections(EntityCollections, *MassEntityManager, ExecutionContext, [InstanceData, &IterationContext, &InstanceHandle, &Operation, &bContinue](FMassExecutionContext& Context)
+			InstancedActorLocationQuery.ForEachEntityChunkInCollections(EntityCollections, ExecutionContext, [InstanceData, &IterationContext, &InstanceHandle, &Operation, &bContinue](FMassExecutionContext& Context)
 				{
 					if (!bContinue)
 					{
diff --git a/Engine/Plugins/Runtime/InstancedActors/Source/InstancedActors/Private/InstancedActorsStationaryLODBatchProcessor.cpp b/Engine/Plugins/Runtime/InstancedActors/Source/InstancedActors/Private/InstancedActorsStationaryLODBatchProcessor.cpp
--- a/Engine/Plugins/Runtime/InstancedActors/Source/InstancedActors/Private/InstancedActorsStationaryLODBatchProcessor.cpp
+++ b/Engine/Plugins/Runtime/InstancedActors/Source/InstancedActors/Private/InstancedActorsStationaryLODBatchProcessor.cpp
@@ -363,7 +363,7 @@ void UInstancedActorsStationaryLODBatchProcessor::Execute(FMassEntityManager& En
 							TArray<FMassArchetypeEntityCollection> EntityCollections;
 							UE::Mass::Utils::CreateEntityCollections(EntityManager, InstanceData->Entities, FMassArchetypeEntityCollection::NoDuplicates, EntityCollections);
 
-							LODChangingEntityQuery->ForEachEntityChunkInCollections(EntityCollections, EntityManager, Context, [NewLOD](FMassExecutionContext& Context)
+							LODChangingEntityQuery->ForEachEntityChunkInCollections(EntityCollections, Context, [NewLOD](FMassExecutionContext& Context)
 								{
 									const TArrayView<FMassRepresentationLODFragment> RepresentationLODFragments = Context.GetMutableFragmentView<FMassRepresentationLODFragment>();
 									for (FMassRepresentationLODFragment& LODFragment : RepresentationLODFragments)
@@ -444,7 +444,7 @@ void UInstancedActorsStationaryLODBatchProcessor::Execute(FMassEntityManager& En
 				}
 
 				// Ensure that detailed representation update occurs for explicitly dirtied instanced actor entities with non-detailed BulkLOD 
-				DirtyVisualizationEntityQuery.ForEachEntityChunkInCollections(DirtyEntityCollections, EntityManager, Context, [](FMassExecutionContext& Context)
+				DirtyVisualizationEntityQuery.ForEachEntityChunkInCollections(DirtyEntityCollections, Context, [](FMassExecutionContext& Context)
 					{
 						// It's possible that we've only just switched to Non-Detailed this frame, the tag removal to prevent regular processing wouldn't have
 						// occurred yet and we would have performed a representation update this frame already.
diff --git a/Engine/Plugins/Runtime/InstancedActors/Source/InstancedActors/Private/InstancedActorsVisualizationSwitcherProcessor.cpp b/Engine/Plugins/Runtime/InstancedActors/Source/InstancedActors/Private/InstancedActorsVisualizationSwitcherProcessor.cpp
--- a/Engine/Plugins/Runtime/InstancedActors/Source/InstancedActors/Private/InstancedActorsVisualizationSwitcherProcessor.cpp
+++ b/Engine/Plugins/Runtime/InstancedActors/Source/InstancedActors/Private/InstancedActorsVisualizationSwitcherProcessor.cpp
@@ -32,7 +32,7 @@ void UInstancedActorsVisualizationSwitcherProcessor::ConfigureQueries(const TSha
 
 void UInstancedActorsVisualizationSwitcherProcessor::Execute(FMassEntityManager& EntityManager, FMassExecutionContext& Context)
 {
-	EntityQuery.ForEachEntityChunk(EntityManager, Context, [](FMassExecutionContext& Context)
+	EntityQuery.ForEachEntityChunk(Context, [](FMassExecutionContext& Context)
 	{
 		UMassRepresentationSubsystem* RepresentationSubsystem = Context.GetMutableSharedFragment<FMassRepresentationSubsystemSharedFragment>().RepresentationSubsystem;
 		check(RepresentationSubsystem);
diff --git a/Engine/Plugins/Runtime/MassGameplay/Source/MassActors/Private/Translators/MassCapsuleComponentTranslators.cpp b/Engine/Plugins/Runtime/MassGameplay/Source/MassActors/Private/Translators/MassCapsuleComponentTranslators.cpp
--- a/Engine/Plugins/Runtime/MassGameplay/Source/MassActors/Private/Translators/MassCapsuleComponentTranslators.cpp
+++ b/Engine/Plugins/Runtime/MassGameplay/Source/MassActors/Private/Translators/MassCapsuleComponentTranslators.cpp
@@ -27,7 +27,7 @@ void UMassCapsuleTransformToMassTranslator::ConfigureQueries(const TSharedRef<FM
 
 void UMassCapsuleTransformToMassTranslator::Execute(FMassEntityManager& EntityManager, FMassExecutionContext& Context)
 {
-	EntityQuery.ForEachEntityChunk(EntityManager, Context, [this](FMassExecutionContext& Context)
+	EntityQuery.ForEachEntityChunk(Context, [this](FMassExecutionContext& Context)
 		{
 			const TConstArrayView<FCapsuleComponentWrapperFragment> CapsuleComponentList = Context.GetFragmentView<FCapsuleComponentWrapperFragment>();
 			const TArrayView<FTransformFragment> LocationList = Context.GetMutableFragmentView<FTransformFragment>();
@@ -64,7 +64,7 @@ void UMassTransformToActorCapsuleTranslator::ConfigureQueries(const TSharedRef<F
 
 void UMassTransformToActorCapsuleTranslator::Execute(FMassEntityManager& EntityManager, FMassExecutionContext& Context)
 {
-	EntityQuery.ForEachEntityChunk(EntityManager, Context, [this](FMassExecutionContext& Context)
+	EntityQuery.ForEachEntityChunk(Context, [this](FMassExecutionContext& Context)
 		{
 			const TArrayView<FCapsuleComponentWrapperFragment> CapsuleComponentList = Context.GetMutableFragmentView<FCapsuleComponentWrapperFragment>();
 			const TConstArrayView<FTransformFragment> LocationList = Context.GetFragmentView<FTransformFragment>();
diff --git a/Engine/Plugins/Runtime/MassGameplay/Source/MassActors/Private/Translators/MassCharacterMovementTranslators.cpp b/Engine/Plugins/Runtime/MassGameplay/Source/MassActors/Private/Translators/MassCharacterMovementTranslators.cpp
--- a/Engine/Plugins/Runtime/MassGameplay/Source/MassActors/Private/Translators/MassCharacterMovementTranslators.cpp
+++ b/Engine/Plugins/Runtime/MassGameplay/Source/MassActors/Private/Translators/MassCharacterMovementTranslators.cpp
@@ -32,7 +32,7 @@ void UMassCharacterMovementToMassTranslator::ConfigureQueries(const TSharedRef<F
 
 void UMassCharacterMovementToMassTranslator::Execute(FMassEntityManager& EntityManager, FMassExecutionContext& Context)
 {
-	EntityQuery.ForEachEntityChunk(EntityManager, Context, [this](FMassExecutionContext& Context)
+	EntityQuery.ForEachEntityChunk(Context, [this](FMassExecutionContext& Context)
 	{
 		const TConstArrayView<FCharacterMovementComponentWrapperFragment> ComponentList = Context.GetFragmentView<FCharacterMovementComponentWrapperFragment>();
 		const TArrayView<FTransformFragment> LocationList = Context.GetMutableFragmentView<FTransformFragment>();
@@ -74,7 +74,7 @@ void UMassCharacterMovementToActorTranslator::ConfigureQueries(const TSharedRef<
 
 void UMassCharacterMovementToActorTranslator::Execute(FMassEntityManager& EntityManager, FMassExecutionContext& Context)
 {
-	EntityQuery.ForEachEntityChunk(EntityManager, Context, [this](FMassExecutionContext& Context)
+	EntityQuery.ForEachEntityChunk(Context, [this](FMassExecutionContext& Context)
 	{
 		const TArrayView<FCharacterMovementComponentWrapperFragment> ComponentList = Context.GetMutableFragmentView<FCharacterMovementComponentWrapperFragment>();
 		const TConstArrayView<FMassVelocityFragment> VelocityList = Context.GetFragmentView<FMassVelocityFragment>();
@@ -112,7 +112,7 @@ void UMassCharacterOrientationToMassTranslator::ConfigureQueries(const TSharedRe
 
 void UMassCharacterOrientationToMassTranslator::Execute(FMassEntityManager& EntityManager, FMassExecutionContext& Context)
 {
-	EntityQuery.ForEachEntityChunk(EntityManager, Context, [this](FMassExecutionContext& Context)
+	EntityQuery.ForEachEntityChunk(Context, [this](FMassExecutionContext& Context)
 	{
 		const TConstArrayView<FCharacterMovementComponentWrapperFragment> ComponentList = Context.GetFragmentView<FCharacterMovementComponentWrapperFragment>();
 		const TArrayView<FTransformFragment> LocationList = Context.GetMutableFragmentView<FTransformFragment>();
@@ -155,7 +155,7 @@ void UMassCharacterOrientationToActorTranslator::ConfigureQueries(const TSharedR
 
 void UMassCharacterOrientationToActorTranslator::Execute(FMassEntityManager& EntityManager, FMassExecutionContext& Context)
 {
-	EntityQuery.ForEachEntityChunk(EntityManager, Context, [this](FMassExecutionContext& Context)
+	EntityQuery.ForEachEntityChunk(Context, [this](FMassExecutionContext& Context)
 	{
 		const TArrayView<FCharacterMovementComponentWrapperFragment> ComponentList = Context.GetMutableFragmentView<FCharacterMovementComponentWrapperFragment>();
 		const TConstArrayView<FTransformFragment> TransformList = Context.GetFragmentView<FTransformFragment>();
diff --git a/Engine/Plugins/Runtime/MassGameplay/Source/MassActors/Private/Translators/MassSceneComponentLocationTranslator.cpp b/Engine/Plugins/Runtime/MassGameplay/Source/MassActors/Private/Translators/MassSceneComponentLocationTranslator.cpp
--- a/Engine/Plugins/Runtime/MassGameplay/Source/MassActors/Private/Translators/MassSceneComponentLocationTranslator.cpp
+++ b/Engine/Plugins/Runtime/MassGameplay/Source/MassActors/Private/Translators/MassSceneComponentLocationTranslator.cpp
@@ -25,7 +25,7 @@ void UMassSceneComponentLocationToMassTranslator::ConfigureQueries(const TShared
 
 void UMassSceneComponentLocationToMassTranslator::Execute(FMassEntityManager& EntityManager, FMassExecutionContext& Context)
 {
-	EntityQuery.ForEachEntityChunk(EntityManager, Context, [this](FMassExecutionContext& Context)
+	EntityQuery.ForEachEntityChunk(Context, [this](FMassExecutionContext& Context)
 	{
 		const TConstArrayView<FMassSceneComponentWrapperFragment> ComponentList = Context.GetFragmentView<FMassSceneComponentWrapperFragment>();
 		const TArrayView<FTransformFragment> LocationList = Context.GetMutableFragmentView<FTransformFragment>();
@@ -64,7 +64,7 @@ void UMassSceneComponentLocationToActorTranslator::ConfigureQueries(const TShare
 
 void UMassSceneComponentLocationToActorTranslator::Execute(FMassEntityManager& EntityManager, FMassExecutionContext& Context)
 {
-	EntityQuery.ForEachEntityChunk(EntityManager, Context, [this](FMassExecutionContext& Context)
+	EntityQuery.ForEachEntityChunk(Context, [this](FMassExecutionContext& Context)
 		{
 			const TConstArrayView<FMassSceneComponentWrapperFragment> ComponentList = Context.GetFragmentView<FMassSceneComponentWrapperFragment>();
 			const TArrayView<FTransformFragment> LocationList = Context.GetMutableFragmentView<FTransformFragment>();
diff --git a/Engine/Plugins/Runtime/MassGameplay/Source/MassEQS/Private/Processors/MassEnvQueryGeneratorProcessor_MassEntityHandles.cpp b/Engine/Plugins/Runtime/MassGameplay/Source/MassEQS/Private/Processors/MassEnvQueryGeneratorProcessor_MassEntityHandles.cpp
--- a/Engine/Plugins/Runtime/MassGameplay/Source/MassEQS/Private/Processors/MassEnvQueryGeneratorProcessor_MassEntityHandles.cpp
+++ b/Engine/Plugins/Runtime/MassGameplay/Source/MassEQS/Private/Processors/MassEnvQueryGeneratorProcessor_MassEntityHandles.cpp
@@ -43,7 +43,7 @@ void UMassEnvQueryGeneratorProcessor_MassEntityHandles::Execute(FMassEntityManag
 	TArray<FMassEnvQueryEntityInfo> Items = {};
 	if (SearchRadius <= 0)
 	{
-		EntityQuery.ForEachEntityChunk(EntityManager, Context, [&GeneratorData, &Items, &SearchRadius, &ContextPositions](FMassExecutionContext& ChunkContext)
+		EntityQuery.ForEachEntityChunk(Context, [&GeneratorData, &Items, &SearchRadius, &ContextPositions](FMassExecutionContext& ChunkContext)
 		{
 			const TConstArrayView<FTransformFragment> TransformFragmentList = ChunkContext.GetFragmentView<FTransformFragment>();
 			const int32 NumEntities = ChunkContext.GetNumEntities();
@@ -57,7 +57,7 @@ void UMassEnvQueryGeneratorProcessor_MassEntityHandles::Execute(FMassEntityManag
 	else
 	{
 		FVector::FReal SearchRadiusSqr = FMath::Square(SearchRadius);
-		EntityQuery.ForEachEntityChunk(EntityManager, Context, [&GeneratorData, &Items, &SearchRadiusSqr, &ContextPositions](FMassExecutionContext& ChunkContext)
+		EntityQuery.ForEachEntityChunk(Context, [&GeneratorData, &Items, &SearchRadiusSqr, &ContextPositions](FMassExecutionContext& ChunkContext)
 		{
 			const TConstArrayView<FTransformFragment> TransformFragmentList = ChunkContext.GetFragmentView<FTransformFragment>();
 			const int32 NumEntities = ChunkContext.GetNumEntities();
diff --git a/Engine/Plugins/Runtime/MassGameplay/Source/MassEQS/Private/Processors/MassEnvQueryTestProcessor_MassEntityTags.cpp b/Engine/Plugins/Runtime/MassGameplay/Source/MassEQS/Private/Processors/MassEnvQueryTestProcessor_MassEntityTags.cpp
--- a/Engine/Plugins/Runtime/MassGameplay/Source/MassEQS/Private/Processors/MassEnvQueryTestProcessor_MassEntityTags.cpp
+++ b/Engine/Plugins/Runtime/MassGameplay/Source/MassEQS/Private/Processors/MassEnvQueryTestProcessor_MassEntityTags.cpp
@@ -85,7 +85,7 @@ void UMassEnvQueryTestProcessor_MassEntityTags::Execute(FMassEntityManager& Enti
 	TArray<FMassArchetypeEntityCollection> EntityCollectionsToTest;
 	UE::Mass::Utils::CreateEntityCollections(EntityManager, TestDataUniquePtr->EntityHandles, FMassArchetypeEntityCollection::NoDuplicates, EntityCollectionsToTest);
 
-	EntityQuery.ForEachEntityChunkInCollections(EntityCollectionsToTest, EntityManager, ExecutionContext
+	EntityQuery.ForEachEntityChunkInCollections(EntityCollectionsToTest, ExecutionContext
 		, [&ScoreMap, &ScoringFunction, &Tags](FMassExecutionContext& Context)
 		{
 			bool ChunkSuccess = ScoringFunction(Tags, Context);
diff --git a/Engine/Plugins/Runtime/MassGameplay/Source/MassGameplayDebug/Private/AssignDebugVisProcessor.cpp b/Engine/Plugins/Runtime/MassGameplay/Source/MassGameplayDebug/Private/AssignDebugVisProcessor.cpp
--- a/Engine/Plugins/Runtime/MassGameplay/Source/MassGameplayDebug/Private/AssignDebugVisProcessor.cpp
+++ b/Engine/Plugins/Runtime/MassGameplay/Source/MassGameplayDebug/Private/AssignDebugVisProcessor.cpp
@@ -40,7 +40,7 @@ void UAssignDebugVisProcessor::Execute(FMassEntityManager& EntityManager, FMassE
 	// note that this function will create the "visual components" only it they're missing or out of sync. 
 	Debugger.GetVisualizationComponent()->ConditionallyConstructVisualComponent();
 
-	EntityQuery.ForEachEntityChunk(EntityManager, Context, [this](FMassExecutionContext& Context)
+	EntityQuery.ForEachEntityChunk(Context, [this](FMassExecutionContext& Context)
 	{
 		UMassDebuggerSubsystem& Debugger = Context.GetMutableSubsystemChecked<UMassDebuggerSubsystem>();
 		UMassDebugVisualizationComponent* Visualizer = Debugger.GetVisualizationComponent();
diff --git a/Engine/Plugins/Runtime/MassGameplay/Source/MassGameplayDebug/Private/DebugVisLocationProcessor.cpp b/Engine/Plugins/Runtime/MassGameplay/Source/MassGameplayDebug/Private/DebugVisLocationProcessor.cpp
--- a/Engine/Plugins/Runtime/MassGameplay/Source/MassGameplayDebug/Private/DebugVisLocationProcessor.cpp
+++ b/Engine/Plugins/Runtime/MassGameplay/Source/MassGameplayDebug/Private/DebugVisLocationProcessor.cpp
@@ -32,7 +32,7 @@ void UDebugVisLocationProcessor::Execute(FMassEntityManager& EntityManager, FMas
 #if WITH_EDITORONLY_DATA
 	QUICK_SCOPE_CYCLE_COUNTER(DebugVisLocationProcessor_Run);
 
-	EntityQuery.ForEachEntityChunk(EntityManager, Context, [this](FMassExecutionContext& Context)
+	EntityQuery.ForEachEntityChunk(Context, [this](FMassExecutionContext& Context)
 	{
 		UMassDebuggerSubsystem& Debugger = Context.GetMutableSubsystemChecked<UMassDebuggerSubsystem>();
 		UMassDebugVisualizationComponent* Visualizer = Debugger.GetVisualizationComponent();
@@ -97,7 +97,7 @@ void UMassProcessor_UpdateDebugVis::Execute(FMassEntityManager& EntityManager, F
 	UMassDebuggerSubsystem& Debugger = Context.GetMutableSubsystemChecked<UMassDebuggerSubsystem>();
 	Debugger.ResetDebugShapes();
 
-	EntityQuery.ForEachEntityChunk(EntityManager, Context, [&Debugger](FMassExecutionContext& Context)
+	EntityQuery.ForEachEntityChunk(Context, [&Debugger](FMassExecutionContext& Context)
 		{
 			const int32 NumEntities = Context.GetNumEntities();
 			const TConstArrayView<FTransformFragment> LocationList = Context.GetFragmentView<FTransformFragment>();
diff --git a/Engine/Plugins/Runtime/MassGameplay/Source/MassLOD/Private/MassLODCollectorProcessor.cpp b/Engine/Plugins/Runtime/MassGameplay/Source/MassLOD/Private/MassLODCollectorProcessor.cpp
--- a/Engine/Plugins/Runtime/MassGameplay/Source/MassLOD/Private/MassLODCollectorProcessor.cpp
+++ b/Engine/Plugins/Runtime/MassGameplay/Source/MassLOD/Private/MassLODCollectorProcessor.cpp
@@ -68,14 +68,14 @@ void UMassLODCollectorProcessor::ExecuteInternal(FMassEntityManager& EntityManag
 {
 	{
 		TRACE_CPUPROFILER_EVENT_SCOPE(Close);
-		EntityQuery_VisibleRangeAndOnLOD.ForEachEntityChunk(EntityManager, Context, [this](FMassExecutionContext& Context) { CollectLODForChunk<bLocalViewersOnly>(Context); });
-		EntityQuery_VisibleRangeOnly.ForEachEntityChunk(EntityManager, Context, [this](FMassExecutionContext& Context) { CollectLODForChunk<bLocalViewersOnly>(Context); });
-		EntityQuery_OnLODOnly.ForEachEntityChunk(EntityManager, Context, [this](FMassExecutionContext& Context) { CollectLODForChunk<bLocalViewersOnly>(Context); });
+		EntityQuery_VisibleRangeAndOnLOD.ForEachEntityChunk(Context, [this](FMassExecutionContext& Context) { CollectLODForChunk<bLocalViewersOnly>(Context); });
+		EntityQuery_VisibleRangeOnly.ForEachEntityChunk(Context, [this](FMassExecutionContext& Context) { CollectLODForChunk<bLocalViewersOnly>(Context); });
+		EntityQuery_OnLODOnly.ForEachEntityChunk(Context, [this](FMassExecutionContext& Context) { CollectLODForChunk<bLocalViewersOnly>(Context); });
 	}
 
 	{
 		TRACE_CPUPROFILER_EVENT_SCOPE(Far);
-		EntityQuery_NotVisibleRangeAndOffLOD.ForEachEntityChunk(EntityManager, Context, [this](FMassExecutionContext& Context) { CollectLODForChunk<bLocalViewersOnly>(Context); });
+		EntityQuery_NotVisibleRangeAndOffLOD.ForEachEntityChunk(Context, [this](FMassExecutionContext& Context) { CollectLODForChunk<bLocalViewersOnly>(Context); });
 	}
 }
 
diff --git a/Engine/Plugins/Runtime/MassGameplay/Source/MassLOD/Private/MassLODDistanceCollectorProcessor.cpp b/Engine/Plugins/Runtime/MassGameplay/Source/MassLOD/Private/MassLODDistanceCollectorProcessor.cpp
--- a/Engine/Plugins/Runtime/MassGameplay/Source/MassLOD/Private/MassLODDistanceCollectorProcessor.cpp
+++ b/Engine/Plugins/Runtime/MassGameplay/Source/MassLOD/Private/MassLODDistanceCollectorProcessor.cpp
@@ -74,21 +74,21 @@ void UMassLODDistanceCollectorProcessor::ExecuteInternal(FMassEntityManager& Ent
 		TRACE_CPUPROFILER_EVENT_SCOPE(LODDistanceCollector_Close);
 		{
 			TRACE_CPUPROFILER_EVENT_SCOPE(LODDistanceCollector_Close_RelevantRangeAndOnLOD);
-			EntityQuery_RelevantRangeAndOnLOD.ForEachEntityChunk(EntityManager, Context, [this](FMassExecutionContext& Context) { CollectLODForChunk<bLocalViewersOnly>(Context); });
+			EntityQuery_RelevantRangeAndOnLOD.ForEachEntityChunk(Context, [this](FMassExecutionContext& Context) { CollectLODForChunk<bLocalViewersOnly>(Context); });
 		}
 		{
 			TRACE_CPUPROFILER_EVENT_SCOPE(LODDistanceCollector_Close_OnLODOnly);
-			EntityQuery_RelevantRangeOnly.ForEachEntityChunk(EntityManager, Context, [this](FMassExecutionContext& Context) { CollectLODForChunk<bLocalViewersOnly>(Context); });
+			EntityQuery_RelevantRangeOnly.ForEachEntityChunk(Context, [this](FMassExecutionContext& Context) { CollectLODForChunk<bLocalViewersOnly>(Context); });
 		}
 		{
 			TRACE_CPUPROFILER_EVENT_SCOPE(LODDistanceCollector_Close_OnLODOnly);
-			EntityQuery_OnLODOnly.ForEachEntityChunk(EntityManager, Context, [this](FMassExecutionContext& Context) { CollectLODForChunk<bLocalViewersOnly>(Context); });
+			EntityQuery_OnLODOnly.ForEachEntityChunk(Context, [this](FMassExecutionContext& Context) { CollectLODForChunk<bLocalViewersOnly>(Context); });
 		}
 	}
 
 	{
 		TRACE_CPUPROFILER_EVENT_SCOPE(LODDistanceCollector_Far);
-		EntityQuery_NotRelevantRangeAndOffLOD.ForEachEntityChunk(EntityManager, Context, [this](FMassExecutionContext& Context) { CollectLODForChunk<bLocalViewersOnly>(Context); });
+		EntityQuery_NotRelevantRangeAndOffLOD.ForEachEntityChunk(Context, [this](FMassExecutionContext& Context) { CollectLODForChunk<bLocalViewersOnly>(Context); });
 	}
 }
 
diff --git a/Engine/Plugins/Runtime/MassGameplay/Source/MassLOD/Private/MassSimulationLOD.cpp b/Engine/Plugins/Runtime/MassGameplay/Source/MassLOD/Private/MassSimulationLOD.cpp
--- a/Engine/Plugins/Runtime/MassGameplay/Source/MassLOD/Private/MassSimulationLOD.cpp
+++ b/Engine/Plugins/Runtime/MassGameplay/Source/MassLOD/Private/MassSimulationLOD.cpp
@@ -131,7 +131,7 @@ void UMassSimulationLODProcessor::Execute(FMassEntityManager& EntityManager, FMa
 
 	{
 		TRACE_CPUPROFILER_EVENT_SCOPE(CalculateLOD);
-		EntityQueryCalculateLOD.ForEachEntityChunk(EntityManager, Context, [](FMassExecutionContext& Context)
+		EntityQueryCalculateLOD.ForEachEntityChunk(Context, [](FMassExecutionContext& Context)
 		{
 			FMassSimulationLODSharedFragment& LODSharedFragment = Context.GetMutableSharedFragment<FMassSimulationLODSharedFragment>();
 			TConstArrayView<FMassViewerInfoFragment> ViewersInfoList = Context.GetFragmentView<FMassViewerInfoFragment>();
@@ -148,7 +148,7 @@ void UMassSimulationLODProcessor::Execute(FMassEntityManager& EntityManager, FMa
 			LODSharedFragment.bHasAdjustedDistancesFromCount = LODSharedFragment.LODCalculator.AdjustDistancesFromCount();
 		});
 
-		EntityQueryAdjustDistances.ForEachEntityChunk(EntityManager, Context, [](FMassExecutionContext& Context)
+		EntityQueryAdjustDistances.ForEachEntityChunk(Context, [](FMassExecutionContext& Context)
 		{
 			FMassSimulationLODSharedFragment& LODSharedFragment = Context.GetMutableSharedFragment<FMassSimulationLODSharedFragment>();
 			TConstArrayView<FMassViewerInfoFragment> ViewersInfoList = Context.GetFragmentView<FMassViewerInfoFragment>();
@@ -162,7 +162,7 @@ void UMassSimulationLODProcessor::Execute(FMassEntityManager& EntityManager, FMa
 		TRACE_CPUPROFILER_EVENT_SCOPE(VariableTickRates)
 		check(World);
 		const double Time = World->GetTimeSeconds();
-		EntityQueryVariableTick.ForEachEntityChunk(EntityManager, Context, [Time](FMassExecutionContext& Context)
+		EntityQueryVariableTick.ForEachEntityChunk(Context, [Time](FMassExecutionContext& Context)
 		{
 			FMassSimulationVariableTickSharedFragment& TickRateSharedFragment = Context.GetMutableSharedFragment<FMassSimulationVariableTickSharedFragment>();
 			TConstArrayView<FMassSimulationLODFragment> SimulationLODFragments = Context.GetFragmentView<FMassSimulationLODFragment>();
@@ -175,7 +175,7 @@ void UMassSimulationLODProcessor::Execute(FMassEntityManager& EntityManager, FMa
 	{
 		TRACE_CPUPROFILER_EVENT_SCOPE(SetLODTags)
 		check(World);
-		EntityQuerySetLODTag.ForEachEntityChunk(EntityManager, Context, [](FMassExecutionContext& Context)
+		EntityQuerySetLODTag.ForEachEntityChunk(Context, [](FMassExecutionContext& Context)
 		{
 			TConstArrayView<FMassSimulationLODFragment> SimulationLODFragments = Context.GetFragmentView<FMassSimulationLODFragment>();
 			const int32 NumEntities = Context.GetNumEntities();
@@ -196,7 +196,7 @@ void UMassSimulationLODProcessor::Execute(FMassEntityManager& EntityManager, FMa
 	if (UE::MassLOD::bDebugSimulationLOD)
 	{
 		TRACE_CPUPROFILER_EVENT_SCOPE(DebugDisplayLOD);
-		EntityQuery.ForEachEntityChunk(EntityManager, Context, [World](FMassExecutionContext& Context)
+		EntityQuery.ForEachEntityChunk(Context, [World](FMassExecutionContext& Context)
 		{
 			FMassSimulationLODSharedFragment& LODSharedFragment = Context.GetMutableSharedFragment<FMassSimulationLODSharedFragment>();
 			TConstArrayView<FTransformFragment> LocationList = Context.GetFragmentView<FTransformFragment>();
diff --git a/Engine/Plugins/Runtime/MassGameplay/Source/MassMovement/Private/MassMovementProcessors.cpp b/Engine/Plugins/Runtime/MassGameplay/Source/MassMovement/Private/MassMovementProcessors.cpp
--- a/Engine/Plugins/Runtime/MassGameplay/Source/MassMovement/Private/MassMovementProcessors.cpp
+++ b/Engine/Plugins/Runtime/MassGameplay/Source/MassMovement/Private/MassMovementProcessors.cpp
@@ -41,7 +41,7 @@ void UMassApplyMovementProcessor::Execute(FMassEntityManager& EntityManager,
 
 	QUICK_SCOPE_CYCLE_COUNTER(HighRes);
 
-	EntityQuery.ForEachEntityChunk(EntityManager, Context, [this, DeltaTime](FMassExecutionContext& Context)
+	EntityQuery.ForEachEntityChunk(Context, [this, DeltaTime](FMassExecutionContext& Context)
 	{
 		const int32 NumEntities = Context.GetNumEntities();
 
diff --git a/Engine/Plugins/Runtime/MassGameplay/Source/MassMovement/Private/MassSimpleMovementTrait.cpp b/Engine/Plugins/Runtime/MassGameplay/Source/MassMovement/Private/MassSimpleMovementTrait.cpp
--- a/Engine/Plugins/Runtime/MassGameplay/Source/MassMovement/Private/MassSimpleMovementTrait.cpp
+++ b/Engine/Plugins/Runtime/MassGameplay/Source/MassMovement/Private/MassSimpleMovementTrait.cpp
@@ -43,7 +43,7 @@ void UMassSimpleMovementProcessor::ConfigureQueries(const TSharedRef<FMassEntity
 
 void UMassSimpleMovementProcessor::Execute(FMassEntityManager& EntityManager, FMassExecutionContext& Context)
 {
-	EntityQuery.ForEachEntityChunk(EntityManager, Context, ([this](FMassExecutionContext& Context)
+	EntityQuery.ForEachEntityChunk(Context, ([this](FMassExecutionContext& Context)
 		{
 			const TConstArrayView<FMassVelocityFragment> VelocitiesList = Context.GetFragmentView<FMassVelocityFragment>();
 			const TArrayView<FTransformFragment> TransformsList = Context.GetMutableFragmentView<FTransformFragment>();
diff --git a/Engine/Plugins/Runtime/MassGameplay/Source/MassMovement/Private/MassVelocityRandomizerTrait.cpp b/Engine/Plugins/Runtime/MassGameplay/Source/MassMovement/Private/MassVelocityRandomizerTrait.cpp
--- a/Engine/Plugins/Runtime/MassGameplay/Source/MassMovement/Private/MassVelocityRandomizerTrait.cpp
+++ b/Engine/Plugins/Runtime/MassGameplay/Source/MassMovement/Private/MassVelocityRandomizerTrait.cpp
@@ -39,7 +39,7 @@ void UMassRandomVelocityInitializer::ConfigureQueries(const TSharedRef<FMassEnti
 void UMassRandomVelocityInitializer::Execute(FMassEntityManager& EntityManager, FMassExecutionContext& Context)
 {
 	// note: the author is aware that the vectors produced below are not distributed uniformly, but it's good enough
-	EntityQuery.ForEachEntityChunk(EntityManager, Context, ([this](FMassExecutionContext& Context)
+	EntityQuery.ForEachEntityChunk(Context, ([this](FMassExecutionContext& Context)
 		{
 			const TArrayView<FMassVelocityFragment> VelocitiesList = Context.GetMutableFragmentView<FMassVelocityFragment>();
 			for (FMassVelocityFragment& VelocityFragment : VelocitiesList)
diff --git a/Engine/Plugins/Runtime/MassGameplay/Source/MassReplication/Private/MassReplicationFragments.cpp b/Engine/Plugins/Runtime/MassGameplay/Source/MassReplication/Private/MassReplicationFragments.cpp
--- a/Engine/Plugins/Runtime/MassGameplay/Source/MassReplication/Private/MassReplicationFragments.cpp
+++ b/Engine/Plugins/Runtime/MassGameplay/Source/MassReplication/Private/MassReplicationFragments.cpp
@@ -33,7 +33,7 @@ void UMassNetworkIDFragmentInitializer::Execute(FMassEntityManager& EntityManage
 	if (NetMode != NM_Client)
 	{
 #if UE_REPLICATION_COMPILE_SERVER_CODE
-		EntityQuery.ForEachEntityChunk(EntityManager, Context, [](FMassExecutionContext& Context)
+		EntityQuery.ForEachEntityChunk(Context, [](FMassExecutionContext& Context)
 			{
 				UMassReplicationSubsystem& ReplicationSubsystem = Context.GetMutableSubsystemChecked<UMassReplicationSubsystem>();
 
diff --git a/Engine/Plugins/Runtime/MassGameplay/Source/MassReplication/Private/MassReplicationGridProcessor.cpp b/Engine/Plugins/Runtime/MassGameplay/Source/MassReplication/Private/MassReplicationGridProcessor.cpp
--- a/Engine/Plugins/Runtime/MassGameplay/Source/MassReplication/Private/MassReplicationGridProcessor.cpp
+++ b/Engine/Plugins/Runtime/MassGameplay/Source/MassReplication/Private/MassReplicationGridProcessor.cpp
@@ -53,7 +53,7 @@ void UMassReplicationGridProcessor::ConfigureQueries(const TSharedRef<FMassEntit
 
 void UMassReplicationGridProcessor::Execute(FMassEntityManager& EntityManager, FMassExecutionContext& Context)
 {
-	AddToGridEntityQuery.ForEachEntityChunk(EntityManager, Context, [](FMassExecutionContext& Context)
+	AddToGridEntityQuery.ForEachEntityChunk(Context, [](FMassExecutionContext& Context)
 	{
 		UMassReplicationSubsystem& ReplicationSubsystem = Context.GetMutableSubsystemChecked<UMassReplicationSubsystem>();
 		FReplicationHashGrid2D& ReplicationGrid = ReplicationSubsystem.GetGridMutable();
@@ -78,7 +78,7 @@ void UMassReplicationGridProcessor::Execute(FMassEntityManager& EntityManager, F
 		}
 	});
 
-	UpdateGridEntityQuery.ForEachEntityChunk(EntityManager, Context, [](FMassExecutionContext& Context)
+	UpdateGridEntityQuery.ForEachEntityChunk(Context, [](FMassExecutionContext& Context)
 	{
 		UMassReplicationSubsystem& ReplicationSubsystem = Context.GetMutableSubsystemChecked<UMassReplicationSubsystem>();
 		FReplicationHashGrid2D& ReplicationGrid = ReplicationSubsystem.GetGridMutable();
@@ -110,7 +110,7 @@ void UMassReplicationGridProcessor::Execute(FMassEntityManager& EntityManager, F
 		}
 	});
 
-	RemoveFromGridEntityQuery.ForEachEntityChunk(EntityManager, Context, [](FMassExecutionContext& Context)
+	RemoveFromGridEntityQuery.ForEachEntityChunk(Context, [](FMassExecutionContext& Context)
 	{
 		UMassReplicationSubsystem& ReplicationSubsystem = Context.GetMutableSubsystemChecked<UMassReplicationSubsystem>();
 		FReplicationHashGrid2D& ReplicationGrid = ReplicationSubsystem.GetGridMutable();
@@ -147,7 +147,7 @@ void UMassReplicationGridRemoverProcessor::ConfigureQueries(const TSharedRef<FMa
 
 void UMassReplicationGridRemoverProcessor::Execute(FMassEntityManager& EntityManager, FMassExecutionContext& Context)
 {
-	EntityQuery.ForEachEntityChunk(EntityManager, Context, [](FMassExecutionContext& Context)
+	EntityQuery.ForEachEntityChunk(Context, [](FMassExecutionContext& Context)
 	{
 		UMassReplicationSubsystem& ReplicationSubsystem = Context.GetMutableSubsystemChecked<UMassReplicationSubsystem>();
 		FReplicationHashGrid2D& ReplicationGrid = ReplicationSubsystem.GetGridMutable();
diff --git a/Engine/Plugins/Runtime/MassGameplay/Source/MassReplication/Private/MassReplicationProcessor.cpp b/Engine/Plugins/Runtime/MassGameplay/Source/MassReplication/Private/MassReplicationProcessor.cpp
--- a/Engine/Plugins/Runtime/MassGameplay/Source/MassReplication/Private/MassReplicationProcessor.cpp
+++ b/Engine/Plugins/Runtime/MassGameplay/Source/MassReplication/Private/MassReplicationProcessor.cpp
@@ -217,7 +217,7 @@ void UMassReplicationProcessor::Execute(FMassEntityManager& EntityManager, FMass
 			ReplicationSubsystem->GetGrid().Query(Bounds, EntitiesInRange);
 		}
 
-		EntityQuery.CacheArchetypes(EntityManager);
+		EntityQuery.CacheArchetypes();
 		if (EntityQuery.GetArchetypes().Num() > 0)
 		{
 			// EntitySet stores array of entities per specified archetype, may contain duplicates.
@@ -270,7 +270,7 @@ void UMassReplicationProcessor::Execute(FMassEntityManager& EntityManager, FMass
 
 				{
 					QUICK_SCOPE_CYCLE_COUNTER(UMassReplicationProcessor_SyncToMass);
-					SyncClientData.ForEachEntityChunk(EntityManager, Context, [&ClientReplicationInfo](FMassExecutionContext& Context)
+					SyncClientData.ForEachEntityChunk(Context, [&ClientReplicationInfo](FMassExecutionContext& Context)
 					{
 						const TArrayView<FMassReplicationLODFragment> ViewerLODList = Context.GetMutableFragmentView<FMassReplicationLODFragment>();
 						TArrayView<FMassReplicatedAgentFragment> ReplicatedAgentList = Context.GetMutableFragmentView<FMassReplicatedAgentFragment>();
@@ -298,7 +298,7 @@ void UMassReplicationProcessor::Execute(FMassEntityManager& EntityManager, FMass
 
 				{
 					QUICK_SCOPE_CYCLE_COUNTER(UMassReplicationProcessor_LODCollection);
-					CollectViewerInfoQuery.ForEachEntityChunk(EntityManager, Context, [](FMassExecutionContext& Context)
+					CollectViewerInfoQuery.ForEachEntityChunk(Context, [](FMassExecutionContext& Context)
 					{
 						const TConstArrayView<FTransformFragment> LocationList = Context.GetFragmentView<FTransformFragment>();
 						const TArrayView<FMassReplicationViewerInfoFragment> ViewersInfoList = Context.GetMutableFragmentView<FMassReplicationViewerInfoFragment>();
@@ -309,7 +309,7 @@ void UMassReplicationProcessor::Execute(FMassEntityManager& EntityManager, FMass
 
 				{
 					QUICK_SCOPE_CYCLE_COUNTER(UMassReplicationProcessor_LODCaculation);
-					CalculateLODQuery.ForEachEntityChunk(EntityManager, Context, [](FMassExecutionContext& Context)
+					CalculateLODQuery.ForEachEntityChunk(Context, [](FMassExecutionContext& Context)
 					{
 						const TConstArrayView<FMassReplicationViewerInfoFragment> ViewersInfoList = Context.GetFragmentView<FMassReplicationViewerInfoFragment>();
 						const TArrayView<FMassReplicationLODFragment> ViewerLODList = Context.GetMutableFragmentView<FMassReplicationLODFragment>();
@@ -338,7 +338,7 @@ void UMassReplicationProcessor::Execute(FMassEntityManager& EntityManager, FMass
 
 				{
 					QUICK_SCOPE_CYCLE_COUNTER(UMassReplicationProcessor_LODAdjustLODFromCount);
-					AdjustLODDistancesQuery.ForEachEntityChunk(EntityManager, Context, [](FMassExecutionContext& Context)
+					AdjustLODDistancesQuery.ForEachEntityChunk(Context, [](FMassExecutionContext& Context)
 					{
 						const TConstArrayView<FMassReplicationViewerInfoFragment> ViewersInfoList = Context.GetFragmentView<FMassReplicationViewerInfoFragment>();
 						const TArrayView<FMassReplicationLODFragment> ViewerLODList = Context.GetMutableFragmentView<FMassReplicationLODFragment>();
@@ -354,7 +354,7 @@ void UMassReplicationProcessor::Execute(FMassEntityManager& EntityManager, FMass
 					{
 						RepSharedFragment.CurrentClientHandle = ClientHandle;
 
-						RepSharedFragment.EntityQuery.ForEachEntityChunk(EntityManager, Context, [&ReplicationContext, &RepSharedFragment](FMassExecutionContext& Context)
+						RepSharedFragment.EntityQuery.ForEachEntityChunk(Context, [&ReplicationContext, &RepSharedFragment](FMassExecutionContext& Context)
 						{
 							RepSharedFragment.CachedReplicator->ProcessClientReplication(Context, ReplicationContext);
 						});
@@ -363,7 +363,7 @@ void UMassReplicationProcessor::Execute(FMassEntityManager& EntityManager, FMass
 
 				{
 					QUICK_SCOPE_CYCLE_COUNTER(UMassReplicationProcessor_SyncFromMass);
-					SyncClientData.ForEachEntityChunk(EntityManager, Context, [&ClientReplicationInfo](FMassExecutionContext& Context)
+					SyncClientData.ForEachEntityChunk(Context, [&ClientReplicationInfo](FMassExecutionContext& Context)
 					{
 						TArrayView<FMassReplicatedAgentFragment> ReplicatedAgentList = Context.GetMutableFragmentView<FMassReplicatedAgentFragment>();
 
@@ -383,7 +383,7 @@ void UMassReplicationProcessor::Execute(FMassEntityManager& EntityManager, FMass
 				{
 					EntityManager.ForEachSharedFragment<FMassReplicationSharedFragment>([World, &EntityManager, &Context](FMassReplicationSharedFragment& RepSharedFragment)
 					{
-						RepSharedFragment.EntityQuery.ForEachEntityChunk(EntityManager, Context, [World, &RepSharedFragment](FMassExecutionContext& Context)
+						RepSharedFragment.EntityQuery.ForEachEntityChunk(Context, [World, &RepSharedFragment](FMassExecutionContext& Context)
 						{
 							const TConstArrayView<FTransformFragment> TransformList = Context.GetFragmentView<FTransformFragment>();
 							const TConstArrayView<FMassReplicationLODFragment> ViewerLODList = Context.GetFragmentView<FMassReplicationLODFragment>();
diff --git a/Engine/Plugins/Runtime/MassGameplay/Source/MassReplication/Public/MassClientBubbleHandler.h b/Engine/Plugins/Runtime/MassGameplay/Source/MassReplication/Public/MassClientBubbleHandler.h
--- a/Engine/Plugins/Runtime/MassGameplay/Source/MassReplication/Public/MassClientBubbleHandler.h
+++ b/Engine/Plugins/Runtime/MassGameplay/Source/MassReplication/Public/MassClientBubbleHandler.h
@@ -543,7 +543,7 @@ void TClientBubbleHandlerBase<AgentArrayItem>::PostReplicatedAddEntitiesHelper(c
 		int32 AgentsSpawnIdx = 0;
 
 		Query.ForEachEntityChunk(FMassArchetypeEntityCollection(ArchetypeHandle, Entities, FMassArchetypeEntityCollection::NoDuplicates)
-								, EntityManager, ExecContext, [&AgentsSpawn, &AgentsSpawnIdx, this, ReplicationSubsystem, &ExecContext, &CacheFragmentViewsForSpawnQuery, &SetSpawnedEntityData, &EntityManager](FMassExecutionContext& Context)
+								, ExecContext, [&AgentsSpawn, &AgentsSpawnIdx, this, ReplicationSubsystem, &ExecContext, &CacheFragmentViewsForSpawnQuery, &SetSpawnedEntityData, &EntityManager](FMassExecutionContext& Context)
 			{
 				CacheFragmentViewsForSpawnQuery(ExecContext);
 
diff --git a/Engine/Plugins/Runtime/MassGameplay/Source/MassRepresentation/Private/MassDistanceLODProcessor.cpp b/Engine/Plugins/Runtime/MassGameplay/Source/MassRepresentation/Private/MassDistanceLODProcessor.cpp
--- a/Engine/Plugins/Runtime/MassGameplay/Source/MassRepresentation/Private/MassDistanceLODProcessor.cpp
+++ b/Engine/Plugins/Runtime/MassGameplay/Source/MassRepresentation/Private/MassDistanceLODProcessor.cpp
@@ -46,7 +46,7 @@ void UMassDistanceLODProcessor::Execute(FMassEntityManager& EntityManager, FMass
 {
 	if (bForceOFFLOD)
 	{
-		CloseEntityQuery.ForEachEntityChunk(EntityManager, Context, [this](FMassExecutionContext& Context)
+		CloseEntityQuery.ForEachEntityChunk(Context, [this](FMassExecutionContext& Context)
 		{
 			FMassDistanceLODSharedFragment& LODSharedFragment = Context.GetMutableSharedFragment<FMassDistanceLODSharedFragment>();
 			TArrayView<FMassRepresentationLODFragment> RepresentationLODList = Context.GetMutableFragmentView<FMassRepresentationLODFragment>();
@@ -78,8 +78,8 @@ void UMassDistanceLODProcessor::Execute(FMassEntityManager& EntityManager, FMass
 			TConstArrayView<FMassViewerInfoFragment> ViewerInfoList = Context.GetFragmentView<FMassViewerInfoFragment>();
 			LODSharedFragment.LODCalculator.CalculateLOD(Context, ViewerInfoList, RepresentationLODList);
 		};
-		CloseEntityQuery.ForEachEntityChunk(EntityManager, Context, CalculateLOD);
-		FarEntityQuery.ForEachEntityChunk(EntityManager, Context, CalculateLOD);
+		CloseEntityQuery.ForEachEntityChunk(Context, CalculateLOD);
+		FarEntityQuery.ForEachEntityChunk(Context, CalculateLOD);
 	}
 
 #if WITH_MASSGAMEPLAY_DEBUG
@@ -88,7 +88,7 @@ void UMassDistanceLODProcessor::Execute(FMassEntityManager& EntityManager, FMass
 	{
 		TRACE_CPUPROFILER_EVENT_SCOPE(DebugDisplayLOD)
 		UWorld* World = EntityManager.GetWorld();
-		DebugEntityQuery.ForEachEntityChunk(EntityManager, Context, [World](FMassExecutionContext& Context)
+		DebugEntityQuery.ForEachEntityChunk(Context, [World](FMassExecutionContext& Context)
 		{
 			FMassDistanceLODSharedFragment& LODSharedFragment = Context.GetMutableSharedFragment<FMassDistanceLODSharedFragment>();
 			TConstArrayView<FMassRepresentationLODFragment> RepresentationLODList = Context.GetFragmentView<FMassRepresentationLODFragment>();
@@ -100,7 +100,7 @@ void UMassDistanceLODProcessor::Execute(FMassEntityManager& EntityManager, FMass
 	if (UE::Mass::Representation::Debug::DebugRepresentationLOD >= 2)
 	{
 		TRACE_CPUPROFILER_EVENT_SCOPE(VisLogLOD)
-		DebugEntityQuery.ForEachEntityChunk(EntityManager, Context, [this](FMassExecutionContext& Context)
+		DebugEntityQuery.ForEachEntityChunk(Context, [this](FMassExecutionContext& Context)
 		{
 			FMassDistanceLODSharedFragment& LODSharedFragment = Context.GetMutableSharedFragment<FMassDistanceLODSharedFragment>();
 			TConstArrayView<FMassRepresentationLODFragment> RepresentationLODList = Context.GetFragmentView<FMassRepresentationLODFragment>();
diff --git a/Engine/Plugins/Runtime/MassGameplay/Source/MassRepresentation/Private/MassRepresentationProcessor.cpp b/Engine/Plugins/Runtime/MassGameplay/Source/MassRepresentation/Private/MassRepresentationProcessor.cpp
--- a/Engine/Plugins/Runtime/MassGameplay/Source/MassRepresentation/Private/MassRepresentationProcessor.cpp
+++ b/Engine/Plugins/Runtime/MassGameplay/Source/MassRepresentation/Private/MassRepresentationProcessor.cpp
@@ -289,7 +289,7 @@ void UMassRepresentationProcessor::UpdateRepresentation(FMassExecutionContext& C
 void UMassRepresentationProcessor::Execute(FMassEntityManager& InEntityManager, FMassExecutionContext& Context)
 {
 	// Update entities representation
-	EntityQuery.ForEachEntityChunk(InEntityManager, Context, [this](FMassExecutionContext& Context)
+	EntityQuery.ForEachEntityChunk(Context, [this](FMassExecutionContext& Context)
 	{
 		UpdateRepresentation(Context, UpdateParams);
 	});
@@ -430,7 +430,7 @@ void UMassVisualizationProcessor::Execute(FMassEntityManager& InEntityManager, F
 
 	int32 TotalEntitiesProcessed = 0;
 	// Update entities visualization
-	EntityQuery.ForEachEntityChunk(InEntityManager, Context, [this, &TotalEntitiesProcessed](FMassExecutionContext& Context)
+	EntityQuery.ForEachEntityChunk(Context, [this, &TotalEntitiesProcessed](FMassExecutionContext& Context)
 	{
 		TotalEntitiesProcessed += Context.GetNumEntities();
 		UpdateVisualization(Context);
@@ -462,7 +462,7 @@ void UMassRepresentationFragmentDestructor::ConfigureQueries(const TSharedRef<FM
 
 void UMassRepresentationFragmentDestructor::Execute(FMassEntityManager& EntityManager, FMassExecutionContext& Context)
 {
-	EntityQuery.ForEachEntityChunk(EntityManager, Context, [this](FMassExecutionContext& Context)
+	EntityQuery.ForEachEntityChunk(Context, [this](FMassExecutionContext& Context)
 	{
 		UMassRepresentationSubsystem* RepresentationSubsystem = Context.GetMutableSharedFragment<FMassRepresentationSubsystemSharedFragment>().RepresentationSubsystem;
 		check(RepresentationSubsystem);
diff --git a/Engine/Plugins/Runtime/MassGameplay/Source/MassRepresentation/Private/MassStationaryISMRepresentationFragmentDestructor.cpp b/Engine/Plugins/Runtime/MassGameplay/Source/MassRepresentation/Private/MassStationaryISMRepresentationFragmentDestructor.cpp
--- a/Engine/Plugins/Runtime/MassGameplay/Source/MassRepresentation/Private/MassStationaryISMRepresentationFragmentDestructor.cpp
+++ b/Engine/Plugins/Runtime/MassGameplay/Source/MassRepresentation/Private/MassStationaryISMRepresentationFragmentDestructor.cpp
@@ -27,7 +27,7 @@ void UMassStationaryISMRepresentationFragmentDestructor::ConfigureQueries(const
 
 void UMassStationaryISMRepresentationFragmentDestructor::Execute(FMassEntityManager& EntityManager, FMassExecutionContext& Context)
 {
-	EntityQuery.ForEachEntityChunk(EntityManager, Context, [this](FMassExecutionContext& Context)
+	EntityQuery.ForEachEntityChunk(Context, [this](FMassExecutionContext& Context)
 	{
 		UMassRepresentationSubsystem* RepresentationSubsystem = Context.GetMutableSharedFragment<FMassRepresentationSubsystemSharedFragment>().RepresentationSubsystem;
 		check(RepresentationSubsystem);
diff --git a/Engine/Plugins/Runtime/MassGameplay/Source/MassRepresentation/Private/MassStationaryISMSwitcherProcessor.cpp b/Engine/Plugins/Runtime/MassGameplay/Source/MassRepresentation/Private/MassStationaryISMSwitcherProcessor.cpp
--- a/Engine/Plugins/Runtime/MassGameplay/Source/MassRepresentation/Private/MassStationaryISMSwitcherProcessor.cpp
+++ b/Engine/Plugins/Runtime/MassGameplay/Source/MassRepresentation/Private/MassStationaryISMSwitcherProcessor.cpp
@@ -31,7 +31,7 @@ void UMassStationaryISMSwitcherProcessor::ConfigureQueries(const TSharedRef<FMas
 
 void UMassStationaryISMSwitcherProcessor::Execute(FMassEntityManager& EntityManager, FMassExecutionContext& Context)
 {
-	EntityQuery.ForEachEntityChunk(EntityManager, Context, &UMassStationaryISMSwitcherProcessor::ProcessContext);
+	EntityQuery.ForEachEntityChunk(Context, &UMassStationaryISMSwitcherProcessor::ProcessContext);
 }
 	
 void UMassStationaryISMSwitcherProcessor::ProcessContext(FMassExecutionContext& Context)
diff --git a/Engine/Plugins/Runtime/MassGameplay/Source/MassRepresentation/Private/MassUpdateISMProcessor.cpp b/Engine/Plugins/Runtime/MassGameplay/Source/MassRepresentation/Private/MassUpdateISMProcessor.cpp
--- a/Engine/Plugins/Runtime/MassGameplay/Source/MassRepresentation/Private/MassUpdateISMProcessor.cpp
+++ b/Engine/Plugins/Runtime/MassGameplay/Source/MassRepresentation/Private/MassUpdateISMProcessor.cpp
@@ -34,7 +34,7 @@ void UMassUpdateISMProcessor::ConfigureQueries(const TSharedRef<FMassEntityManag
 
 void UMassUpdateISMProcessor::Execute(FMassEntityManager& EntityManager, FMassExecutionContext& Context)
 {
-	EntityQuery.ForEachEntityChunk(EntityManager, Context, [](FMassExecutionContext& Context)
+	EntityQuery.ForEachEntityChunk(Context, [](FMassExecutionContext& Context)
 	{
 		UMassRepresentationSubsystem* RepresentationSubsystem = Context.GetSharedFragment<FMassRepresentationSubsystemSharedFragment>().RepresentationSubsystem;
 		check(RepresentationSubsystem);
diff --git a/Engine/Plugins/Runtime/MassGameplay/Source/MassRepresentation/Private/MassVisualizationLODProcessor.cpp b/Engine/Plugins/Runtime/MassGameplay/Source/MassRepresentation/Private/MassVisualizationLODProcessor.cpp
--- a/Engine/Plugins/Runtime/MassGameplay/Source/MassRepresentation/Private/MassVisualizationLODProcessor.cpp
+++ b/Engine/Plugins/Runtime/MassGameplay/Source/MassRepresentation/Private/MassVisualizationLODProcessor.cpp
@@ -51,7 +51,7 @@ void UMassVisualizationLODProcessor::Execute(FMassEntityManager& EntityManager,
 {
 	if (bForceOFFLOD)
 	{
-		CloseEntityQuery.ForEachEntityChunk(EntityManager, Context, [this](FMassExecutionContext& Context)
+		CloseEntityQuery.ForEachEntityChunk(Context, [this](FMassExecutionContext& Context)
 		{
 			FMassVisualizationLODSharedFragment& LODSharedFragment = Context.GetMutableSharedFragment<FMassVisualizationLODSharedFragment>();
 			TArrayView<FMassRepresentationLODFragment> RepresentationLODList = Context.GetMutableFragmentView<FMassRepresentationLODFragment>();
@@ -83,8 +83,8 @@ void UMassVisualizationLODProcessor::Execute(FMassEntityManager& EntityManager,
 			TConstArrayView<FMassViewerInfoFragment> ViewerInfoList = Context.GetFragmentView<FMassViewerInfoFragment>();
 			LODSharedFragment.LODCalculator.CalculateLOD(Context, ViewerInfoList, RepresentationLODList);
 		};
-		CloseEntityQuery.ForEachEntityChunk(EntityManager, Context, CalculateLOD);
-		FarEntityQuery.ForEachEntityChunk(EntityManager, Context, CalculateLOD);
+		CloseEntityQuery.ForEachEntityChunk(Context, CalculateLOD);
+		FarEntityQuery.ForEachEntityChunk(Context, CalculateLOD);
 	}
 
 	if (bDoAdjustmentFromCount)
@@ -98,7 +98,7 @@ void UMassVisualizationLODProcessor::Execute(FMassEntityManager& EntityManager,
 			}
 		});
 
-		CloseEntityAdjustDistanceQuery.ForEachEntityChunk(EntityManager, Context, [this](FMassExecutionContext& Context)
+		CloseEntityAdjustDistanceQuery.ForEachEntityChunk(Context, [this](FMassExecutionContext& Context)
 		{
 			FMassVisualizationLODSharedFragment& LODSharedFragment = Context.GetMutableSharedFragment<FMassVisualizationLODSharedFragment>();
 			TConstArrayView<FMassViewerInfoFragment> ViewerInfoList = Context.GetFragmentView<FMassViewerInfoFragment>();
@@ -114,7 +114,7 @@ void UMassVisualizationLODProcessor::Execute(FMassEntityManager& EntityManager,
 	{
 		TRACE_CPUPROFILER_EVENT_SCOPE(DebugDisplayLOD)
 		UWorld* World = EntityManager.GetWorld();
-		DebugEntityQuery.ForEachEntityChunk(EntityManager, Context, [World](FMassExecutionContext& Context)
+		DebugEntityQuery.ForEachEntityChunk(Context, [World](FMassExecutionContext& Context)
 		{
 			FMassVisualizationLODSharedFragment& LODSharedFragment = Context.GetMutableSharedFragment<FMassVisualizationLODSharedFragment>();
 			TConstArrayView<FMassRepresentationLODFragment> RepresentationLODList = Context.GetFragmentView<FMassRepresentationLODFragment>();
@@ -126,7 +126,7 @@ void UMassVisualizationLODProcessor::Execute(FMassEntityManager& EntityManager,
 	if (UE::Mass::Representation::Debug::DebugRepresentationLOD >= 2)
 	{
 		TRACE_CPUPROFILER_EVENT_SCOPE(VisLogLOD)
-		DebugEntityQuery.ForEachEntityChunk(EntityManager, Context, [this](FMassExecutionContext& Context)
+		DebugEntityQuery.ForEachEntityChunk(Context, [this](FMassExecutionContext& Context)
 		{
 			FMassVisualizationLODSharedFragment& LODSharedFragment = Context.GetMutableSharedFragment<FMassVisualizationLODSharedFragment>();
 			TConstArrayView<FMassRepresentationLODFragment> RepresentationLODList = Context.GetFragmentView<FMassRepresentationLODFragment>();
diff --git a/Engine/Plugins/Runtime/MassGameplay/Source/MassSmartObjects/Private/MassSmartObjectProcessor.cpp b/Engine/Plugins/Runtime/MassGameplay/Source/MassSmartObjects/Private/MassSmartObjectProcessor.cpp
--- a/Engine/Plugins/Runtime/MassGameplay/Source/MassSmartObjects/Private/MassSmartObjectProcessor.cpp
+++ b/Engine/Plugins/Runtime/MassGameplay/Source/MassSmartObjects/Private/MassSmartObjectProcessor.cpp
@@ -87,7 +87,7 @@ void UMassSmartObjectCandidatesFinderProcessor::Execute(FMassEntityManager& Enti
 	};
 
 	// Process world location based requests
-	WorldRequestQuery.ForEachEntityChunk(EntityManager, Context, [this, &Filter, &EntitiesToSignal, &BeginRequestProcessing, &EndRequestProcessing](FMassExecutionContext& Context)
+	WorldRequestQuery.ForEachEntityChunk(Context, [this, &Filter, &EntitiesToSignal, &BeginRequestProcessing, &EndRequestProcessing](FMassExecutionContext& Context)
 	{
 		const USmartObjectSubsystem& SmartObjectSubsystem = Context.GetSubsystemChecked<USmartObjectSubsystem>();
 
@@ -156,7 +156,7 @@ void UMassSmartObjectCandidatesFinderProcessor::Execute(FMassEntityManager& Enti
 	const FZoneGraphTag SmartObjectTag = GetDefault<UMassSmartObjectSettings>()->SmartObjectTag;
 	USmartObjectZoneAnnotations* Annotations = Cast<USmartObjectZoneAnnotations>(AnnotationSubsystem.GetFirstAnnotationForTag(SmartObjectTag));
 
-	LaneRequestQuery.ForEachEntityChunk(EntityManager, Context,
+	LaneRequestQuery.ForEachEntityChunk(Context,
 		[&AnnotationSubsystem, Annotations, &Filter, SmartObjectTag, &EntitiesToSignal, &BeginRequestProcessing, &EndRequestProcessing](FMassExecutionContext& Context)
 		{
 #if WITH_MASSGAMEPLAY_DEBUG
@@ -325,7 +325,7 @@ void UMassSmartObjectTimedBehaviorProcessor::Execute(FMassEntityManager& EntityM
 
 	QUICK_SCOPE_CYCLE_COUNTER(UMassProcessor_SmartObjectTestBehavior_Run);
 
-	EntityQuery.ForEachEntityChunk(EntityManager, Context, [this, &ToRelease](FMassExecutionContext& Context)
+	EntityQuery.ForEachEntityChunk(Context, [this, &ToRelease](FMassExecutionContext& Context)
 	{
 		const USmartObjectSubsystem& SmartObjectSubsystem = Context.GetSubsystemChecked<USmartObjectSubsystem>();
 
@@ -400,7 +400,7 @@ void UMassSmartObjectUserFragmentDeinitializer::ConfigureQueries(const TSharedRe
 
 void UMassSmartObjectUserFragmentDeinitializer::Execute(FMassEntityManager& EntityManager, FMassExecutionContext& Context)
 {
-	EntityQuery.ForEachEntityChunk(EntityManager, Context, [this](FMassExecutionContext& Context)
+	EntityQuery.ForEachEntityChunk(Context, [this](FMassExecutionContext& Context)
 		{
 			USmartObjectSubsystem& SmartObjectSubsystem = Context.GetMutableSubsystemChecked<USmartObjectSubsystem>();
 			const int32 NumEntities = Context.GetNumEntities();
diff --git a/Engine/Plugins/Runtime/MassGameplay/Source/MassSmartObjects/Private/MassSmartObjectRegistration.cpp b/Engine/Plugins/Runtime/MassGameplay/Source/MassSmartObjects/Private/MassSmartObjectRegistration.cpp
--- a/Engine/Plugins/Runtime/MassGameplay/Source/MassSmartObjects/Private/MassSmartObjectRegistration.cpp
+++ b/Engine/Plugins/Runtime/MassGameplay/Source/MassSmartObjects/Private/MassSmartObjectRegistration.cpp
@@ -37,7 +37,7 @@ void UMassSmartObjectInitializerBase::Execute(FMassEntityManager& EntityManager,
 	UMassSignalSubsystem& SignalSubsystem = Context.GetMutableSubsystemChecked<UMassSignalSubsystem>();
 	TArray<FMassEntityHandle> EntitiesToSignal;
 
-	EntityQuery.ForEachEntityChunk(EntityManager, Context,[&EntitiesToSignal](FMassExecutionContext& Context)
+	EntityQuery.ForEachEntityChunk(Context,[&EntitiesToSignal](FMassExecutionContext& Context)
 	{
 		const TConstArrayView<FMassActorInstanceFragment> InstancedActorFragments = Context.GetFragmentView<FMassActorInstanceFragment>();
 		const TArrayView<FSmartObjectRegistrationFragment> RegistrationFragments = Context.GetMutableFragmentView<FSmartObjectRegistrationFragment>();
@@ -92,7 +92,7 @@ void UMassSmartObjectDeinitializerBase::Execute(FMassEntityManager& EntityManage
 	UMassSignalSubsystem& SignalSubsystem = Context.GetMutableSubsystemChecked<UMassSignalSubsystem>();
 	TArray<FMassEntityHandle> EntitiesToSignal;
 
-	EntityQuery.ForEachEntityChunk(EntityManager, Context, [&EntitiesToSignal](FMassExecutionContext& Context)
+	EntityQuery.ForEachEntityChunk(Context, [&EntitiesToSignal](FMassExecutionContext& Context)
 	{
 		const int32 NumEntities = Context.GetNumEntities();
 		EntitiesToSignal.Reserve(EntitiesToSignal.Num() + NumEntities);
@@ -182,7 +182,7 @@ void UMassActiveSmartObjectSignalProcessor::InitializeInternal(UObject& Owner, c
 void UMassActiveSmartObjectSignalProcessor::SignalEntities(FMassEntityManager& EntityManager, FMassExecutionContext& Context, FMassSignalNameLookup&)
 {
 	// Process entities in active range
-	InsideSmartObjectActiveRangeQuery.ForEachEntityChunk(EntityManager, Context, [this](FMassExecutionContext& Context)
+	InsideSmartObjectActiveRangeQuery.ForEachEntityChunk(Context, [this](FMassExecutionContext& Context)
 	{
 		USmartObjectSubsystem& Subsystem = Context.GetMutableSubsystemChecked<USmartObjectSubsystem>();
 		const TConstArrayView<FTransformFragment> TransformFragments = Context.GetFragmentView<FTransformFragment>();
@@ -218,7 +218,7 @@ void UMassActiveSmartObjectSignalProcessor::SignalEntities(FMassEntityManager& E
 	});
 
 	// Process out of active range entities
-	OutsideSmartObjectActiveRangeQuery.ForEachEntityChunk(EntityManager, Context, [this](FMassExecutionContext& Context)
+	OutsideSmartObjectActiveRangeQuery.ForEachEntityChunk(Context, [this](FMassExecutionContext& Context)
 	{
 		USmartObjectSubsystem& Subsystem = Context.GetMutableSubsystemChecked<USmartObjectSubsystem>();
 		const TArrayView<FSmartObjectRegistrationFragment> RegistrationFragments = Context.GetMutableFragmentView<FSmartObjectRegistrationFragment>();
diff --git a/Engine/Plugins/Runtime/MassGameplay/Source/MassSpawner/Private/MassSpawnLocationProcessor.cpp b/Engine/Plugins/Runtime/MassGameplay/Source/MassSpawner/Private/MassSpawnLocationProcessor.cpp
--- a/Engine/Plugins/Runtime/MassGameplay/Source/MassSpawner/Private/MassSpawnLocationProcessor.cpp
+++ b/Engine/Plugins/Runtime/MassGameplay/Source/MassSpawner/Private/MassSpawnLocationProcessor.cpp
@@ -48,7 +48,7 @@ void UMassSpawnLocationProcessor::Execute(FMassEntityManager& EntityManager, FMa
 	}
 
 	int32 NumRequiredSpawnTransforms = 0;
-	EntityQuery.ForEachEntityChunk(EntityManager, Context, [&NumRequiredSpawnTransforms](const FMassExecutionContext& Context)
+	EntityQuery.ForEachEntityChunk(Context, [&NumRequiredSpawnTransforms](const FMassExecutionContext& Context)
 		{
 			NumRequiredSpawnTransforms += Context.GetNumEntities();
 		});
@@ -69,7 +69,7 @@ void UMassSpawnLocationProcessor::Execute(FMassEntityManager& EntityManager, FMa
 
 	if (AuxData.bRandomize && !UE::Mass::Utils::IsDeterministic())
 	{
-		EntityQuery.ForEachEntityChunk(EntityManager, Context, [&Transforms, this](FMassExecutionContext& Context)
+		EntityQuery.ForEachEntityChunk(Context, [&Transforms, this](FMassExecutionContext& Context)
 			{
 				const TArrayView<FTransformFragment> LocationList = Context.GetMutableFragmentView<FTransformFragment>();
 				const int32 NumEntities = Context.GetNumEntities();
@@ -84,7 +84,7 @@ void UMassSpawnLocationProcessor::Execute(FMassEntityManager& EntityManager, FMa
 	else
 	{
 		int32 NextTransformIndex = 0;
-		EntityQuery.ForEachEntityChunk(EntityManager, Context, [&Transforms, &NextTransformIndex, this](FMassExecutionContext& Context)
+		EntityQuery.ForEachEntityChunk(Context, [&Transforms, &NextTransformIndex, this](FMassExecutionContext& Context)
 			{
 				const int32 NumEntities = Context.GetNumEntities();
 				TArrayView<FTransformFragment> LocationList = Context.GetMutableFragmentView<FTransformFragment>();
diff --git a/Engine/Source/Developer/MassEntityTestSuite/Private/MassArchetypeGroupTest.cpp b/Engine/Source/Developer/MassEntityTestSuite/Private/MassArchetypeGroupTest.cpp
--- a/Engine/Source/Developer/MassEntityTestSuite/Private/MassArchetypeGroupTest.cpp
+++ b/Engine/Source/Developer/MassEntityTestSuite/Private/MassArchetypeGroupTest.cpp
@@ -38,7 +38,7 @@ struct FArchetypeGroup_SingleLevelQuery : FEntityTestBase
 		EntityQuery.AddRequirement<FTestFragment_Int>(EMassFragmentAccess::ReadOnly);
 
 		FMassExecutionContext ExecutionContext = EntityManager->CreateExecutionContext(0);
-		EntityQuery.ForEachEntityChunk(*EntityManager, ExecutionContext, [&VerifyEntities](const FMassExecutionContext& Context)
+		EntityQuery.ForEachEntityChunk(ExecutionContext, [&VerifyEntities](const FMassExecutionContext& Context)
 		{
 			VerifyEntities.Append(Context.GetEntities());
 		});
@@ -55,7 +55,7 @@ struct FArchetypeGroup_SingleLevelQuery : FEntityTestBase
 		EntityQuery.GroupBy(GroupType);
 
 		VerifyEntities.Reset();
-		EntityQuery.ForEachEntityChunk(*EntityManager, ExecutionContext, [&VerifyEntities](const FMassExecutionContext& Context)
+		EntityQuery.ForEachEntityChunk(ExecutionContext, [&VerifyEntities](const FMassExecutionContext& Context)
 		{
 			VerifyEntities.Append(Context.GetEntities());
 		});
@@ -66,7 +66,7 @@ struct FArchetypeGroup_SingleLevelQuery : FEntityTestBase
 		EntityQuery.GroupBy(GroupType, DescendingGroupSorter);
 
 		VerifyEntities.Reset();
-		EntityQuery.ForEachEntityChunk(*EntityManager, ExecutionContext, [&VerifyEntities](const FMassExecutionContext& Context)
+		EntityQuery.ForEachEntityChunk(ExecutionContext, [&VerifyEntities](const FMassExecutionContext& Context)
 		{
 			// note that we're inserting at index 0 to end up with the ascending order of entities in the VerifyEntities array for ease of comparison
 			VerifyEntities.Insert(Context.GetEntities()[0], 0);
@@ -146,7 +146,7 @@ struct FArchetypeGroup_MultiLevelQuery : FEntityTestBase
 		}
 
 		VerifyEntities.Reset();
-		EntityQuery.ForEachEntityChunk(*EntityManager, ExecutionContext, [&VerifyEntities](const FMassExecutionContext& Context)
+		EntityQuery.ForEachEntityChunk(ExecutionContext, [&VerifyEntities](const FMassExecutionContext& Context)
 		{
 			VerifyEntities.Append(Context.GetEntities());
 		});
@@ -164,7 +164,7 @@ struct FArchetypeGroup_MultiLevelQuery : FEntityTestBase
 			const int32 ExpectedEntityIndices[NumEntities] = { 4, 5, 6, 7, 0, 1, 2, 3, 12, 13, 14, 15, 8, 9, 10, 11 };
 
 			VerifyEntities.Reset();
-			EntityQuery.ForEachEntityChunk(*EntityManager, ExecutionContext, [&VerifyEntities](const FMassExecutionContext& Context)
+			EntityQuery.ForEachEntityChunk(ExecutionContext, [&VerifyEntities](const FMassExecutionContext& Context)
 			{
 				VerifyEntities.Append(Context.GetEntities());
 			});
@@ -185,7 +185,7 @@ struct FArchetypeGroup_MultiLevelQuery : FEntityTestBase
 			const int32 ExpectedEntityIndices[NumEntities] = { 2, 3, 0, 1, 6, 7, 4, 5, 10, 11, 8, 9, 14, 15, 12, 13 };
 
 			VerifyEntities.Reset();
-			EntityQuery.ForEachEntityChunk(*EntityManager, ExecutionContext, [&VerifyEntities](const FMassExecutionContext& Context)
+			EntityQuery.ForEachEntityChunk(ExecutionContext, [&VerifyEntities](const FMassExecutionContext& Context)
 			{
 				VerifyEntities.Append(Context.GetEntities());
 			});
@@ -206,7 +206,7 @@ struct FArchetypeGroup_MultiLevelQuery : FEntityTestBase
 			const int32 ExpectedEntityIndices[NumEntities] = { 12, 13, 14, 15, 4, 5, 6, 7, 8, 9, 10, 11, 0, 1, 2, 3 };
 
 			VerifyEntities.Reset();
-			EntityQuery.ForEachEntityChunk(*EntityManager, ExecutionContext, [&VerifyEntities](const FMassExecutionContext& Context)
+			EntityQuery.ForEachEntityChunk(ExecutionContext, [&VerifyEntities](const FMassExecutionContext& Context)
 			{
 				VerifyEntities.Append(Context.GetEntities());
 			});
diff --git a/Engine/Source/Developer/MassEntityTestSuite/Private/MassChunkCollectionTest.cpp b/Engine/Source/Developer/MassEntityTestSuite/Private/MassChunkCollectionTest.cpp
--- a/Engine/Source/Developer/MassEntityTestSuite/Private/MassChunkCollectionTest.cpp
+++ b/Engine/Source/Developer/MassEntityTestSuite/Private/MassChunkCollectionTest.cpp
@@ -403,7 +403,7 @@ struct FEntityCollection_ZeroLength_Destroy : FEntityTestBase
 		AITEST_EQUAL("Entity creation generated expected number of entities", EntityCount, FMassArchetypeHelper::ArchetypeDataFromHandleChecked(SourceArchetype).GetNumEntities());
 
 		FMassEntityQuery Query(EntityManager);
-		Query.CacheArchetypes(*EntityManager);
+		Query.CacheArchetypes();
 		FMassArchetypeEntityCollection Collection(SourceArchetype);
 		AITEST_EQUAL("Created collection has expected number of ranges (matching number of chunks created"
 			, Collection.GetRanges().Num(), NumChunks);
diff --git a/Engine/Source/Developer/MassEntityTestSuite/Private/MassComplexTest.cpp b/Engine/Source/Developer/MassEntityTestSuite/Private/MassComplexTest.cpp
--- a/Engine/Source/Developer/MassEntityTestSuite/Private/MassComplexTest.cpp
+++ b/Engine/Source/Developer/MassEntityTestSuite/Private/MassComplexTest.cpp
@@ -43,7 +43,7 @@ struct FComplex_AddingFragmentInstancesToDiscontinouousEntitiesCollection : FEnt
 
 		FMassExecutionContext ExecContext(*EntityManager.Get());
 		FMassEntityQuery Query(EntityManager.ToSharedRef(), { FTestFragment_Int::StaticStruct() });
-		Query.ForEachEntityChunk(*EntityManager, ExecContext, [](FMassExecutionContext& Context)
+		Query.ForEachEntityChunk(ExecContext, [](FMassExecutionContext& Context)
 			{
 				TArrayView<FTestFragment_Int> Ints = Context.GetMutableFragmentView<FTestFragment_Int>();
 				for (int32 i = 0; i < Context.GetNumEntities(); ++i)
diff --git a/Engine/Source/Developer/MassEntityTestSuite/Private/MassEntityTestFarmPlot.cpp b/Engine/Source/Developer/MassEntityTestSuite/Private/MassEntityTestFarmPlot.cpp
--- a/Engine/Source/Developer/MassEntityTestSuite/Private/MassEntityTestFarmPlot.cpp
+++ b/Engine/Source/Developer/MassEntityTestSuite/Private/MassEntityTestFarmPlot.cpp
@@ -33,7 +33,7 @@ void UFarmWaterProcessor::ConfigureQueries(const TSharedRef<FMassEntityManager>&
 void UFarmWaterProcessor::Execute(FMassEntityManager& EntityManager, FMassExecutionContext& Context)
 {
 	QUICK_SCOPE_CYCLE_COUNTER(UFarmWaterProcessor_Run);
-	EntityQuery.ForEachEntityChunk(EntityManager, Context, [this](FMassExecutionContext& Context) {
+	EntityQuery.ForEachEntityChunk(Context, [this](FMassExecutionContext& Context) {
 
 		const float DeltaTimeSeconds = Context.GetDeltaTimeSeconds();
 		TArrayView<FFarmWaterFragment> WaterList = Context.GetMutableFragmentView<FFarmWaterFragment>();
@@ -59,7 +59,7 @@ void UFarmHarvestTimerSystem_Flowers::Execute(FMassEntityManager& EntityManager,
 {
 	QUICK_SCOPE_CYCLE_COUNTER(UFarmHarvestTimerSystem_Flowers_Run);
 
-	EntityQuery.ForEachEntityChunk(EntityManager, Context, [this](FMassExecutionContext& Context) {
+	EntityQuery.ForEachEntityChunk(Context, [this](FMassExecutionContext& Context) {
 		const int32 NumEntities = Context.GetNumEntities();
 		const float WellWateredThreshold = 0.25f;
 		TArrayView<FHarvestTimerFragment> TimerList = Context.GetMutableFragmentView<FHarvestTimerFragment>();
@@ -95,7 +95,7 @@ void UFarmHarvestTimerSystem_Crops::Execute(FMassEntityManager & EntityManager,
 {
 	QUICK_SCOPE_CYCLE_COUNTER(UFarmHarvestTimerSystem_Crops_Run);
 
-	EntityQuery.ForEachEntityChunk(EntityManager, Context, [this](FMassExecutionContext& Context) {
+	EntityQuery.ForEachEntityChunk(Context, [this](FMassExecutionContext& Context) {
 
 		const int32 NumEntities = Context.GetNumEntities();
 		const float WellWateredThreshold = 0.25f;
@@ -124,7 +124,7 @@ void UFarmHarvestTimerExpired::Execute(FMassEntityManager& EntityManager, FMassE
 {
 	QUICK_SCOPE_CYCLE_COUNTER(UFarmHarvestTimerExpired_Run);
 
-	EntityQuery.ForEachEntityChunk(EntityManager, Context, [this](FMassExecutionContext& Context) {
+	EntityQuery.ForEachEntityChunk(Context, [this](FMassExecutionContext& Context) {
 		const int32 NumEntities = Context.GetNumEntities();
 		TConstArrayView<FHarvestTimerFragment> TimerList = Context.GetFragmentView<FHarvestTimerFragment>();
 
@@ -157,7 +157,7 @@ void UFarmHarvestTimerSetIcon::Execute(FMassEntityManager& EntityManager, FMassE
 
 	QUICK_SCOPE_CYCLE_COUNTER(SET_ICON_SET_ICON_SET_ICON);
 
-	EntityQuery.ForEachEntityChunk(EntityManager, Context, [this](FMassExecutionContext& Context) {
+	EntityQuery.ForEachEntityChunk(Context, [this](FMassExecutionContext& Context) {
 
 		const int32 NumEntities = Context.GetNumEntities();
 		TConstArrayView<FFarmGridCellData> GridCoordList = Context.GetFragmentView<FFarmGridCellData>();
diff --git a/Engine/Source/Developer/MassEntityTestSuite/Private/MassParallelQueryTest.cpp b/Engine/Source/Developer/MassEntityTestSuite/Private/MassParallelQueryTest.cpp
--- a/Engine/Source/Developer/MassEntityTestSuite/Private/MassParallelQueryTest.cpp
+++ b/Engine/Source/Developer/MassEntityTestSuite/Private/MassParallelQueryTest.cpp
@@ -35,7 +35,7 @@ struct FQueryTest_ParallelForBasic : FEntityTestBase
 		FCriticalSection ParallelDetectionCS;
 
 		FMassExecutionContext Context = EntityManager->CreateExecutionContext(/*DeltaSeconds=*/0.f);
-		Query.ParallelForEachEntityChunk(*EntityManager, Context, [&ParallelDetectionCS, &EntitiesProcessed, &ConcurrentAccess](FMassExecutionContext& Context)
+		Query.ParallelForEachEntityChunk(Context, [&ParallelDetectionCS, &EntitiesProcessed, &ConcurrentAccess](FMassExecutionContext& Context)
 			{
 				EntitiesProcessed += Context.GetNumEntities();
 
@@ -91,7 +91,7 @@ struct FQueryTest_ParallelForCollection : FEntityTestBase
 		FMassExecutionContext Context = EntityManager->CreateExecutionContext(/*DeltaSeconds=*/0.f);
 		Context.SetEntityCollection(EntityCollection);
 
-		Query.ParallelForEachEntityChunk(*EntityManager, Context, [&JobsExecuted, &ParallelDetectionCS, &EntitiesProcessed, &ConcurrentAccess](FMassExecutionContext& Context)
+		Query.ParallelForEachEntityChunk(Context, [&JobsExecuted, &ParallelDetectionCS, &EntitiesProcessed, &ConcurrentAccess](FMassExecutionContext& Context)
 			{
 				++JobsExecuted;
 				EntitiesProcessed += Context.GetNumEntities();
@@ -136,7 +136,7 @@ struct FQueryTest_ParallelCommands : FEntityTestBase
 		AITEST_EQUAL("Target archetype should only contain freshly created entities", OriginalEntityCount, NumEntitiesToCreate);
 
 		FMassExecutionContext Context = EntityManager->CreateExecutionContext(/*DeltaSeconds=*/0.f);
-		Query.ParallelForEachEntityChunk(*EntityManager, Context, [](FMassExecutionContext& Context)
+		Query.ParallelForEachEntityChunk(Context, [](FMassExecutionContext& Context)
 			{
 				Context.Defer().PushCommand<FMassCommandAddTags<FTestFragment_Tag>>(Context.GetEntities());
 			}, FMassEntityQuery::EParallelExecutionFlags::Force);
diff --git a/Engine/Source/Developer/MassEntityTestSuite/Private/MassQueryTest.cpp b/Engine/Source/Developer/MassEntityTestSuite/Private/MassQueryTest.cpp
--- a/Engine/Source/Developer/MassEntityTestSuite/Private/MassQueryTest.cpp
+++ b/Engine/Source/Developer/MassEntityTestSuite/Private/MassQueryTest.cpp
@@ -98,7 +98,7 @@ struct FQueryTest_ExecuteSingleArchetype : FEntityTestBase
 
 		FMassExecutionContext ExecContext(*EntityManager.Get());
 		FMassEntityQuery Query(EntityManager.ToSharedRef(), { FTestFragment_Float::StaticStruct() });
-		Query.ForEachEntityChunk(*EntityManager, ExecContext, [&TotalProcessed](FMassExecutionContext& Context)
+		Query.ForEachEntityChunk(ExecContext, [&TotalProcessed](FMassExecutionContext& Context)
 			{
 				TotalProcessed += Context.GetNumEntities();
 				TArrayView<FTestFragment_Float> Floats = Context.GetMutableFragmentView<FTestFragment_Float>();
@@ -142,7 +142,7 @@ struct FQueryTest_ExecuteMultipleArchetypes : FEntityTestBase
 		int TotalProcessed = 0;
 		FMassExecutionContext ExecContext(*EntityManager.Get());
 		FMassEntityQuery Query(EntityManager.ToSharedRef(), { FTestFragment_Float::StaticStruct() });
-		Query.ForEachEntityChunk(*EntityManager, ExecContext, [&TotalProcessed](FMassExecutionContext& Context)
+		Query.ForEachEntityChunk(ExecContext, [&TotalProcessed](FMassExecutionContext& Context)
 			{
 				TotalProcessed += Context.GetNumEntities();
 				TArrayView<FTestFragment_Float> Floats = Context.GetMutableFragmentView<FTestFragment_Float>();
@@ -200,7 +200,7 @@ struct FQueryTest_ExecuteSparse : FEntityTestBase
 		FMassEntityQuery TestQuery(EntityManager);
 		TestQuery.AddRequirement<FTestFragment_Float>(EMassFragmentAccess::ReadWrite);
 		TestQuery.ForEachEntityChunk(FMassArchetypeEntityCollection(FloatsArchetype, EntitiesToProcess, FMassArchetypeEntityCollection::NoDuplicates)
-								, *EntityManager, ExecContext, [&TotalProcessed](FMassExecutionContext& Context)
+								, ExecContext, [&TotalProcessed](FMassExecutionContext& Context)
 			{
 				TotalProcessed += Context.GetNumEntities();
 				TArrayView<FTestFragment_Float> Floats = Context.GetMutableFragmentView<FTestFragment_Float>();
@@ -243,7 +243,7 @@ struct FQueryTest_TagPresent : FEntityTestBase
 		FMassEntityQuery Query(EntityManager);
 		Query.AddRequirement<FTestFragment_Float>(EMassFragmentAccess::ReadWrite);
 		Query.AddTagRequirement<FTestFragment_Tag>(EMassFragmentPresence::All);
-		Query.CacheArchetypes(*EntityManager);
+		Query.CacheArchetypes();
 
 		AITEST_EQUAL("There's a single archetype matching the requirements", Query.GetArchetypes().Num(), 1);
 		AITEST_TRUE("The only valid archetype is FloatsTagArchetype", FloatsTagArchetype == Query.GetArchetypes()[0]);
@@ -266,7 +266,7 @@ struct FQueryTest_TagAbsent : FEntityTestBase
 		FMassEntityQuery Query(EntityManager);
 		Query.AddRequirement<FTestFragment_Float>(EMassFragmentAccess::ReadWrite);
 		Query.AddTagRequirement<FTestFragment_Tag>(EMassFragmentPresence::None);
-		Query.CacheArchetypes(*EntityManager);
+		Query.CacheArchetypes();
 
 		AITEST_EQUAL("There are exactly two archetypes matching the requirements", Query.GetArchetypes().Num(), 2);
 		AITEST_TRUE("FloatsTagArchetype is not amongst matching archetypes"
@@ -288,7 +288,7 @@ struct FQueryTest_FragmentPresent : FEntityTestBase
 		FMassEntityQuery Query(EntityManager);
 		// using EMassFragmentAccess::None to indicate we're interested only in the archetype having the fragment, no binding is required
 		Query.AddRequirement<FTestFragment_Int>(EMassFragmentAccess::None, EMassFragmentPresence::Any);
-		Query.CacheArchetypes(*EntityManager);
+		Query.CacheArchetypes();
 
 		AITEST_EQUAL("There are exactly two archetypes matching the requirements", Query.GetArchetypes().Num(), 2);
 		AITEST_TRUE("FloatsArchetype is not amongst matching archetypes"
@@ -303,7 +303,7 @@ struct FQueryTest_FragmentPresent : FEntityTestBase
 
 		int TotalProcessed = 0;
 		FMassExecutionContext ExecContext(*EntityManager.Get());
-		Query.ForEachEntityChunk(*EntityManager, ExecContext, [&TotalProcessed](FMassExecutionContext& Context) {
+		Query.ForEachEntityChunk(ExecContext, [&TotalProcessed](FMassExecutionContext& Context) {
 			TotalProcessed += Context.GetNumEntities();
 		});
 		AITEST_EQUAL("We expect the number of entities processed to match number added to matching archetypes", MatchingEntities.Num(), TotalProcessed);
@@ -330,12 +330,12 @@ struct FQueryTest_OnlyAbsentFragments : FEntityTestBase
 		Query.AddRequirement<FTestFragment_Bool>(EMassFragmentAccess::ReadOnly, EMassFragmentPresence::None);
 		AITEST_TRUE("Multiple negative requirement is valid", Query.CheckValidity());
 		
-		Query.CacheArchetypes(*EntityManager);
+		Query.CacheArchetypes();
 		AITEST_EQUAL("There's only one default test archetype matching the query", Query.GetArchetypes().Num(), 1);
 		AITEST_TRUE("Only the Empty archetype matches the query", Query.GetArchetypes()[0] == EmptyArchetype);
 
 		const FMassArchetypeHandle NewMatchingArchetypeHandle = EntityManager->CreateArchetype({ FTestFragment_Large::StaticStruct() });
-		Query.CacheArchetypes(*EntityManager);
+		Query.CacheArchetypes();
 		AITEST_EQUAL("The number of matching queries matches expectations", Query.GetArchetypes().Num(), 2);
 		AITEST_TRUE("The new archetype matches the query", Query.GetArchetypes()[1] == NewMatchingArchetypeHandle);
 
@@ -356,7 +356,7 @@ struct FQueryTest_AbsentAndPresentFragments : FEntityTestBase
 		Query.AddRequirement<FTestFragment_Float>(EMassFragmentAccess::ReadOnly, EMassFragmentPresence::All);
 
 		AITEST_TRUE("The query is valid", Query.CheckValidity());
-		Query.CacheArchetypes(*EntityManager);
+		Query.CacheArchetypes();
 		AITEST_EQUAL("There is only one archetype matching the query", Query.GetArchetypes().Num(), 1);
 		AITEST_TRUE("FloatsArchetype is the only one matching the query", FloatsArchetype == Query.GetArchetypes()[0]);
 
@@ -374,7 +374,7 @@ struct FQueryTest_SingleOptionalFragment : FEntityTestBase
 
 		FMassEntityQuery Query(EntityManager);
 		Query.AddRequirement<FTestFragment_Int>(EMassFragmentAccess::ReadWrite, EMassFragmentPresence::Optional);
-		Query.CacheArchetypes(*EntityManager);
+		Query.CacheArchetypes();
 
 		AITEST_EQUAL("There are exactly two archetypes matching the requirements", Query.GetArchetypes().Num(), 2);
 		AITEST_TRUE("FloatsArchetype is not amongst matching archetypes"
@@ -395,7 +395,7 @@ struct FQueryTest_MultipleOptionalFragment : FEntityTestBase
 		FMassEntityQuery Query(EntityManager);
 		Query.AddRequirement<FTestFragment_Int>(EMassFragmentAccess::ReadWrite, EMassFragmentPresence::Optional);
 		Query.AddRequirement<FTestFragment_Float>(EMassFragmentAccess::ReadWrite, EMassFragmentPresence::Optional);
-		Query.CacheArchetypes(*EntityManager);
+		Query.CacheArchetypes();
 
 		AITEST_EQUAL("All three archetype meet requirements", Query.GetArchetypes().Num(), 3);
 		return true;
@@ -424,7 +424,7 @@ struct FQueryTest_UsingOptionalFragment : FEntityTestBase
 		Query.AddRequirement<FTestFragment_Int>(EMassFragmentAccess::ReadWrite, EMassFragmentPresence::Optional);
 		Query.AddRequirement<FTestFragment_Float>(EMassFragmentAccess::ReadWrite, EMassFragmentPresence::All);
 		FMassExecutionContext ExecContext(*EntityManager.Get());
-		Query.ForEachEntityChunk(*EntityManager, ExecContext, [&TotalProcessed, &EmptyIntsViewCount, IntValueSet](FMassExecutionContext& Context) {
+		Query.ForEachEntityChunk(ExecContext, [&TotalProcessed, &EmptyIntsViewCount, IntValueSet](FMassExecutionContext& Context) {
 			++TotalProcessed;
 			TArrayView<FTestFragment_Int> Ints = Context.GetMutableFragmentView<FTestFragment_Int>();
 			if (Ints.Num() == 0)
@@ -470,7 +470,7 @@ struct FQueryTest_AnyFragment : FEntityTestBase
 		Query.AddRequirement<FTestFragment_Bool>(EMassFragmentAccess::ReadWrite, EMassFragmentPresence::Any);
 		// this query should match: 
 		// IntsArchetype, FloatsIntsArchetype, BoolArchetype, BoolFloatArchetype
-		Query.CacheArchetypes(*EntityManager);
+		Query.CacheArchetypes();
 
 		AITEST_EQUAL("Archetypes containing Int or Bool should meet requirements", Query.GetArchetypes().Num(), 4);
 
@@ -481,7 +481,7 @@ struct FQueryTest_AnyFragment : FEntityTestBase
 		}
 
 		FMassExecutionContext TestContext(*EntityManager.Get());
-		Query.ForEachEntityChunk(*EntityManager, TestContext, [this](FMassExecutionContext& Context)
+		Query.ForEachEntityChunk(TestContext, [this](FMassExecutionContext& Context)
 			{
 				TArrayView<FTestFragment_Bool> BoolView = Context.GetMutableFragmentView<FTestFragment_Bool>();
 				TArrayView<FTestFragment_Int> IntView = Context.GetMutableFragmentView<FTestFragment_Int>();
@@ -514,7 +514,7 @@ struct FQueryTest_AnyTag : FEntityTestBase
 		Query.AddTagRequirement<FTestTag_C>(EMassFragmentPresence::Any);
 		// this query should match: 
 		// ABArchetype, ACArchetype and ABCrchetype but not BDArchetype nor FEntityTestBase.IntsArchetype
-		Query.CacheArchetypes(*EntityManager);
+		Query.CacheArchetypes();
 
 		AITEST_EQUAL("Only Archetypes tagged with A or C should matched the query", Query.GetArchetypes().Num(), 3);
 		AITEST_TRUE("ABArchetype should be amongst the matched archetypes", Query.GetArchetypes().Find(ABArchetype) != INDEX_NONE);
@@ -544,7 +544,7 @@ struct FQueryTest_AutoRecache : FEntityTestBase
 		};
 		
 		FMassExecutionContext ExecutionContext(*EntityManager, /*DeltaSeconds=*/0.f);
-		Query.ForEachEntityChunk(*EntityManager, ExecutionContext, QueryExecFunction);
+		Query.ForEachEntityChunk(ExecutionContext, QueryExecFunction);
 		
 		AITEST_EQUAL("No entities have been created so we expect counting to yield 0", EntitiesFound, 0);
 
@@ -553,7 +553,7 @@ struct FQueryTest_AutoRecache : FEntityTestBase
 		EntityManager->BatchCreateEntities(IntsArchetype, NumberOfEntitiesMatching, MatchingEntities);
 
 		EntitiesFound = 0;
-		Query.ForEachEntityChunk(*EntityManager, ExecutionContext, QueryExecFunction);
+		Query.ForEachEntityChunk(ExecutionContext, QueryExecFunction);
 		AITEST_EQUAL("The number of entities found should match the number of entities created in the matching archetype", EntitiesFound, MatchingEntities.Num());
 
 		// create more entities, but in an archetype not matching the query
@@ -561,13 +561,13 @@ struct FQueryTest_AutoRecache : FEntityTestBase
 		TArray<FMassEntityHandle> NotMatchingEntities;
 		EntityManager->BatchCreateEntities(FloatsArchetype, NumberOfEntitiesNotMatching, NotMatchingEntities);
 		EntitiesFound = 0;
-		Query.ForEachEntityChunk(*EntityManager, ExecutionContext, QueryExecFunction);
+		Query.ForEachEntityChunk(ExecutionContext, QueryExecFunction);
 		AITEST_EQUAL("The number of entities found should not change with addition of entities not matching the query", EntitiesFound, MatchingEntities.Num());
 
 		// create some more in another matching archetype
 		EntityManager->BatchCreateEntities(FloatsIntsArchetype, NumberOfEntitiesMatching, MatchingEntities);
 		EntitiesFound = 0;
-		Query.ForEachEntityChunk(*EntityManager, ExecutionContext, QueryExecFunction);
+		Query.ForEachEntityChunk(ExecutionContext, QueryExecFunction);
 		AITEST_EQUAL("The total number of entities found should include entities from both matching archetypes", EntitiesFound, MatchingEntities.Num());
 
 		return true;
@@ -589,7 +589,7 @@ struct FQueryTest_AllOptional : FEntityTestBase
 		Query.AddSharedRequirement<FTestSharedFragment_Int>(EMassFragmentAccess::None, EMassFragmentPresence::Optional);
 		Query.AddConstSharedRequirement<FTestConstSharedFragment_Int>(EMassFragmentPresence::Optional);
 
-		Query.CacheArchetypes(*EntityManager);
+		Query.CacheArchetypes();
 
 		int32 ExpectedNumOfArchetypes = 2;
 		// only the FloatsArchetype and FloatsIntsArchetype should match
@@ -603,7 +603,7 @@ struct FQueryTest_AllOptional : FEntityTestBase
 		EntityManager->AddTagToEntity(Entities[CurrentEntityIndex++], FTestTag_A::StaticStruct());
 		++ExpectedNumOfArchetypes;
 		EntityManager->AddTagToEntity(Entities[CurrentEntityIndex++], FTestTag_B::StaticStruct());
-		Query.CacheArchetypes(*EntityManager);
+		Query.CacheArchetypes();
 		AITEST_EQUAL("A: number of matching archetypes matches expectations.", Query.GetArchetypes().Num(), ExpectedNumOfArchetypes);
 
 		{
@@ -617,7 +617,7 @@ struct FQueryTest_AllOptional : FEntityTestBase
 			Descriptor.ChunkFragments.Add<FTestChunkFragment_Float>();
 			EntityManager->CreateArchetype(Descriptor);
 		}
-		Query.CacheArchetypes(*EntityManager);
+		Query.CacheArchetypes();
 		AITEST_EQUAL("B: number of matching archetypes matches expectations.", Query.GetArchetypes().Num(), ExpectedNumOfArchetypes);
 
 		{
@@ -639,7 +639,7 @@ struct FQueryTest_AllOptional : FEntityTestBase
 			FMassArchetypeEntityCollection Collection(IntsArchetype, MakeArrayView(&Entities[CurrentEntityIndex++], 1), FMassArchetypeEntityCollection::EDuplicatesHandling::NoDuplicates);
 			EntityManager->BatchAddSharedFragmentsForEntities(MakeArrayView(&Collection, 1), SharedFragmentValues);
 		}
-		Query.CacheArchetypes(*EntityManager);
+		Query.CacheArchetypes();
 		AITEST_EQUAL("C: number of matching archetypes matches expectations.", Query.GetArchetypes().Num(), ExpectedNumOfArchetypes);
 
 		{
@@ -653,7 +653,7 @@ struct FQueryTest_AllOptional : FEntityTestBase
 			FConstSharedStruct SharedFragmentInstance = FSharedStruct::Make(FragmentInstance);
 			EntityManager->AddConstSharedFragmentToEntity(Entities[CurrentEntityIndex++], SharedFragmentInstance);
 		}
-		Query.CacheArchetypes(*EntityManager);
+		Query.CacheArchetypes();
 		AITEST_EQUAL("D: number of matching archetypes matches expectations.", Query.GetArchetypes().Num(), ExpectedNumOfArchetypes);
 
 		return true;
@@ -669,7 +669,7 @@ struct FQueryTest_JustATag : FEntityTestBase
 
 		FMassEntityQuery Query(EntityManager);
 		Query.AddTagRequirement<FTestTag_A>(EMassFragmentPresence::All);
-		Query.CacheArchetypes(*EntityManager);
+		Query.CacheArchetypes();
 
 		int32 ExpectedNumOfArchetypes = 0;
 		// only the FloatsArchetype and FloatsIntsArchetype should match
@@ -686,7 +686,7 @@ struct FQueryTest_JustATag : FEntityTestBase
 			Descriptor.Tags.Add<FTestTag_B>();
 			EntityManager->CreateArchetype(Descriptor);
 		}
-		Query.CacheArchetypes(*EntityManager);
+		Query.CacheArchetypes();
 		AITEST_EQUAL("A: number of matching archetypes matches expectations.", Query.GetArchetypes().Num(), ExpectedNumOfArchetypes);
 
 		{
@@ -704,7 +704,7 @@ struct FQueryTest_JustATag : FEntityTestBase
 			Descriptor.Tags.Add<FTestTag_D>();
 			EntityManager->CreateArchetype(Descriptor);
 		}
-		Query.CacheArchetypes(*EntityManager);
+		Query.CacheArchetypes();
 		AITEST_EQUAL("B: number of matching archetypes matches expectations.", Query.GetArchetypes().Num(), ExpectedNumOfArchetypes);
 		
 		return true;
@@ -722,7 +722,7 @@ struct FQueryTest_JustAChunkFragment : FEntityTestBase
 
 		FMassEntityQuery Query(EntityManager);
 		Query.AddChunkRequirement<FTestChunkFragment_Int>(EMassFragmentAccess::ReadOnly, EMassFragmentPresence::All);
-		Query.CacheArchetypes(*EntityManager);
+		Query.CacheArchetypes();
 
 		int32 ExpectedNumOfArchetypes = 0;
 		// no matching archetypes at this time
@@ -739,7 +739,7 @@ struct FQueryTest_JustAChunkFragment : FEntityTestBase
 			Descriptor.ChunkFragments.Add<FTestChunkFragment_Float>();
 			EntityManager->CreateArchetype(Descriptor);
 		}
-		Query.CacheArchetypes(*EntityManager);
+		Query.CacheArchetypes();
 		AITEST_EQUAL("Number of matching archetypes matches expectations.", Query.GetArchetypes().Num(), ExpectedNumOfArchetypes);
 		
 		// try to access the chunk fragment
@@ -748,7 +748,7 @@ struct FQueryTest_JustAChunkFragment : FEntityTestBase
 
 			FMassExecutionContext ExecContext(*EntityManager.Get());
 			bool bExecuted = false;
-			Query.ForEachEntityChunk(*EntityManager, ExecContext, [&bExecuted](FMassExecutionContext& Context)
+			Query.ForEachEntityChunk(ExecContext, [&bExecuted](FMassExecutionContext& Context)
 				{
 					const FTestChunkFragment_Int& ChunkFragment = Context.GetChunkFragment<FTestChunkFragment_Int>();
 					bExecuted = true;
@@ -772,7 +772,7 @@ struct FQueryTest_JustASharedFragment : FEntityTestBase
 
 		FMassEntityQuery Query(EntityManager);
 		Query.AddSharedRequirement<FTestSharedFragment_Int>(EMassFragmentAccess::ReadOnly, EMassFragmentPresence::All);
-		Query.CacheArchetypes(*EntityManager);
+		Query.CacheArchetypes();
 
 		int32 ExpectedNumOfArchetypes = 0;
 		// no matching archetypes at this time
@@ -801,14 +801,14 @@ struct FQueryTest_JustASharedFragment : FEntityTestBase
 			FMassArchetypeEntityCollection Collection(IntsArchetype, MakeArrayView(&Entities[CurrentEntityIndex++], 1), FMassArchetypeEntityCollection::EDuplicatesHandling::NoDuplicates);
 			EntityManager->BatchAddSharedFragmentsForEntities(MakeArrayView(&Collection, 1), SharedFragmentValues);
 		}
-		Query.CacheArchetypes(*EntityManager);
+		Query.CacheArchetypes();
 		AITEST_EQUAL("Number of matching archetypes matches expectations.", Query.GetArchetypes().Num(), ExpectedNumOfArchetypes);
 
 		// try to access the shared fragment
 		{
 			bool bExecuted = false;
 			FMassExecutionContext ExecContext(*EntityManager.Get());
-			Query.ForEachEntityChunk(*EntityManager, ExecContext, [&bExecuted](FMassExecutionContext& Context)
+			Query.ForEachEntityChunk(ExecContext, [&bExecuted](FMassExecutionContext& Context)
 				{
 					const FTestSharedFragment_Int& SharedFragment = Context.GetSharedFragment<FTestSharedFragment_Int>();
 					bExecuted = true;
@@ -832,7 +832,7 @@ struct FQueryTest_JustAConstSharedFragment : FEntityTestBase
 
 		FMassEntityQuery Query(EntityManager);
 		Query.AddConstSharedRequirement<FTestConstSharedFragment_Int>(EMassFragmentPresence::All);
-		Query.CacheArchetypes(*EntityManager);
+		Query.CacheArchetypes();
 
 		int32 ExpectedNumOfArchetypes = 0;
 		// no matching archetypes at this time
@@ -853,14 +853,14 @@ struct FQueryTest_JustAConstSharedFragment : FEntityTestBase
 			FConstSharedStruct SharedFragmentInstance = FSharedStruct::Make(FragmentInstance);
 			EntityManager->AddConstSharedFragmentToEntity(Entities[CurrentEntityIndex++], SharedFragmentInstance);
 		}
-		Query.CacheArchetypes(*EntityManager);
+		Query.CacheArchetypes();
 		AITEST_EQUAL("Number of matching archetypes matches expectations.", Query.GetArchetypes().Num(), ExpectedNumOfArchetypes);
 
 		// try to access the shared fragment
 		{
 			bool bExecuted = false;
 			FMassExecutionContext ExecContext(*EntityManager.Get());
-			Query.ForEachEntityChunk(*EntityManager, ExecContext, [&bExecuted](FMassExecutionContext& Context)
+			Query.ForEachEntityChunk(ExecContext, [&bExecuted](FMassExecutionContext& Context)
 				{
 					const FTestConstSharedFragment_Int& SharedFragment = Context.GetConstSharedFragment<FTestConstSharedFragment_Int>();
 					bExecuted = true;
diff --git a/Engine/Source/Developer/MassEntityTestSuite/Private/MassSystemRequirementTest.cpp b/Engine/Source/Developer/MassEntityTestSuite/Private/MassSystemRequirementTest.cpp
--- a/Engine/Source/Developer/MassEntityTestSuite/Private/MassSystemRequirementTest.cpp
+++ b/Engine/Source/Developer/MassEntityTestSuite/Private/MassSystemRequirementTest.cpp
@@ -70,7 +70,7 @@ struct FGenericSubsystemAccessAPI : FSystemRequirementTestBase
 		const UMassTestEngineSubsystem* ConstAccessEngineSubsystem = nullptr;
 
 		FMassExecutionContext ExecutionContext(*EntityManager);
-		EntityQuery.ForEachEntityChunk(*EntityManager, ExecutionContext, [&](FMassExecutionContext& Context)
+		EntityQuery.ForEachEntityChunk(ExecutionContext, [&](FMassExecutionContext& Context)
 			{
 				AccessWorldSubsystem = Context.GetMutableSubsystem<UMassTestWorldSubsystem>();
 				AccessEngineSubsystem = Context.GetMutableSubsystem<UMassTestEngineSubsystem>();
@@ -138,7 +138,7 @@ struct FGameSubsystemAccessAPI : FSystemRequirementTestBase
 		const UMassTestGameInstanceSubsystem* ConstAccessGameInstanceSubsystem = nullptr;
 
 		FMassExecutionContext ExecutionContext(*EntityManager);
-		EntityQuery.ForEachEntityChunk(*EntityManager, ExecutionContext, [&](FMassExecutionContext& Context)
+		EntityQuery.ForEachEntityChunk(ExecutionContext, [&](FMassExecutionContext& Context)
 			{
 				AccessLocalPlayerSubsystem = Context.GetMutableSubsystem<UMassTestLocalPlayerSubsystem>();
 				AccessGameInstanceSubsystem = Context.GetMutableSubsystem<UMassTestGameInstanceSubsystem>();
diff --git a/Engine/Source/Developer/MassEntityTestSuite/Private/MassTestTypes.cpp b/Engine/Source/Developer/MassEntityTestSuite/Private/MassTestTypes.cpp
--- a/Engine/Source/Developer/MassEntityTestSuite/Private/MassTestTypes.cpp
+++ b/Engine/Source/Developer/MassEntityTestSuite/Private/MassTestTypes.cpp
@@ -54,9 +54,9 @@ UMassTestProcessorBase::UMassTestProcessorBase()
 	ExecutionFlags = int32(EProcessorExecutionFlags::All);
 
 	ForEachEntityChunkExecutionFunction = [](FMassExecutionContext& Context) {};
-	ExecutionFunction = [this](FMassEntityManager& InEntitySubsystem, FMassExecutionContext& Context) 
+	ExecutionFunction = [this](FMassEntityManager& InEntityManager, FMassExecutionContext& Context) 
 	{
-		EntityQuery.ForEachEntityChunk(InEntitySubsystem, Context, ForEachEntityChunkExecutionFunction);
+		EntityQuery.ForEachEntityChunk(Context, ForEachEntityChunkExecutionFunction);
 	};
 }
 
diff --git a/Engine/Source/Runtime/MassEntity/Private/MassDebugger.cpp b/Engine/Source/Runtime/MassEntity/Private/MassDebugger.cpp
--- a/Engine/Source/Runtime/MassEntity/Private/MassDebugger.cpp
+++ b/Engine/Source/Runtime/MassEntity/Private/MassDebugger.cpp
@@ -393,7 +393,7 @@ TConstArrayView<FMassEntityQuery*> FMassDebugger::GetUpToDateProcessorQueries(co
 	{
 		if (Query)
 		{
-			Query->CacheArchetypes(EntityManager);
+			Query->CacheArchetypes();
 		}
 	}
 
diff --git a/Engine/Source/Runtime/MassEntity/Private/MassEntityQuery.cpp b/Engine/Source/Runtime/MassEntity/Private/MassEntityQuery.cpp
--- a/Engine/Source/Runtime/MassEntity/Private/MassEntityQuery.cpp
+++ b/Engine/Source/Runtime/MassEntity/Private/MassEntityQuery.cpp
@@ -10,7 +10,6 @@
 #include "Async/ParallelFor.h"
 #include "Containers/UnrealString.h"
 #include "MassProcessor.h"
-#include "MassProcessorDependencySolver.h"
 
 #include UE_INLINE_GENERATED_CPP_BY_NAME(MassEntityQuery)
 
@@ -35,18 +34,37 @@ namespace UE::Mass::Tweakables
 }
 
 //-----------------------------------------------------------------------------
-// FScopedSubsystemRequirementsRestore
+// FScopedEntityQueryContext
 //-----------------------------------------------------------------------------
-FMassEntityQuery::FScopedSubsystemRequirementsRestore::FScopedSubsystemRequirementsRestore(FMassExecutionContext& ExecutionContext)
-	: CachedExecutionContext(ExecutionContext)
+struct FScopedEntityQueryContext
 {
-	CachedExecutionContext.GetSubsystemRequirementBits(ConstSubsystemsBitSet, MutableSubsystemsBitSet);
-}
+	FScopedEntityQueryContext(FMassEntityQuery& InQuery, FMassExecutionContext& ExecutionContext)
+		: Query(InQuery), CachedExecutionContext(ExecutionContext)
+	{
+		CachedExecutionContext.PushQuery(InQuery);
 
-FMassEntityQuery::FScopedSubsystemRequirementsRestore::~FScopedSubsystemRequirementsRestore()
-{
-	CachedExecutionContext.SetSubsystemRequirementBits(ConstSubsystemsBitSet, MutableSubsystemsBitSet);
-}
+		bSubsystemRequirementsCached = CachedExecutionContext.CacheSubsystemRequirements(InQuery);
+	}
+
+	~FScopedEntityQueryContext()
+	{
+		if (IsSuccessfullySetUp())
+		{
+			CachedExecutionContext.ClearExecutionData();
+			CachedExecutionContext.FlushDeferred();
+		}
+		CachedExecutionContext.PopQuery(Query);
+	}
+
+	bool IsSuccessfullySetUp() const
+	{
+		return bSubsystemRequirementsCached;
+	}
+
+	FMassEntityQuery& Query;
+	FMassExecutionContext& CachedExecutionContext;
+	bool bSubsystemRequirementsCached;
+};
 
 //-----------------------------------------------------------------------------
 // FMassEntityQuery
@@ -88,12 +106,14 @@ void FMassEntityQuery::RegisterWithProcessor(UMassProcessor& Owner)
 #endif // WITH_MASSENTITY_DEBUG
 }
 
-void FMassEntityQuery::CacheArchetypes(const FMassEntityManager& InEntityManager)
+void FMassEntityQuery::CacheArchetypes()
 {
-	const uint32 InEntityManagerHash = PointerHash(&InEntityManager);
+	check(CachedEntityManager);
+
+	const uint32 InEntityManagerHash = PointerHash(CachedEntityManager.Get());
 
 	// Do an incremental update if the last updated archetype data version is different 
-    bool bUpdateArchetypes = InEntityManager.GetArchetypeDataVersion() != LastUpdatedArchetypeDataVersion;
+    bool bUpdateArchetypes = CachedEntityManager->GetArchetypeDataVersion() != LastUpdatedArchetypeDataVersion;
 
 	// Force a full update if the entity system changed or if the requirements changed
 	if (EntitySubsystemHash != InEntityManagerHash || HasIncrementalChanges())
@@ -116,7 +136,7 @@ void FMassEntityQuery::CacheArchetypes(const FMassEntityManager& InEntityManager
 			else
 			{
 				bUpdateArchetypes = false;
-				UE_VLOG_UELOG(InEntityManager.GetOwner(), LogMass, Error, TEXT("FMassEntityQuery::CacheArchetypes: requirements not valid: %s"), *FMassDebugger::GetRequirementsDescription(*this));
+				UE_VLOG_UELOG(CachedEntityManager->GetOwner(), LogMass, Error, TEXT("FMassEntityQuery::CacheArchetypes: requirements not valid: %s"), *FMassDebugger::GetRequirementsDescription(*this));
 			}
 		}
 	}
@@ -128,7 +148,7 @@ void FMassEntityQuery::CacheArchetypes(const FMassEntityManager& InEntityManager
 #if WITH_EDITOR
 		if (bHasArchetypeMatchOverride)
 		{
-			InEntityManager.ForEachArchetype(static_cast<int32>(LastUpdatedArchetypeDataVersion), TNumericLimits<int32>::Max() /*last*/, [this, &NewValidArchetypes](const FMassEntityManager& EntityManager, const FMassArchetypeHandle& Handle)
+			CachedEntityManager->ForEachArchetype(static_cast<int32>(LastUpdatedArchetypeDataVersion), TNumericLimits<int32>::Max() /*last*/, [this, &NewValidArchetypes](const FMassEntityManager& EntityManager, const FMassArchetypeHandle& Handle)
 			{
 				const FMassArchetypeCompositionDescriptor& Composition = EntityManager.GetArchetypeComposition(Handle);
 				if (ArchetypeMatchOverride.Match(&ArchetypeMatchOverride.Data, Composition))
@@ -140,9 +160,9 @@ void FMassEntityQuery::CacheArchetypes(const FMassEntityManager& InEntityManager
 		else
 #endif
 		{
-			InEntityManager.GetMatchingArchetypes(*this, NewValidArchetypes, LastUpdatedArchetypeDataVersion);
+			CachedEntityManager->GetMatchingArchetypes(*this, NewValidArchetypes, LastUpdatedArchetypeDataVersion);
 		}
-		LastUpdatedArchetypeDataVersion = InEntityManager.GetArchetypeDataVersion();
+		LastUpdatedArchetypeDataVersion = CachedEntityManager->GetArchetypeDataVersion();
 		if (NewValidArchetypes.Num())
 		{
 			const int32 FirstNewArchetype = ValidArchetypes.Num();
@@ -181,34 +201,39 @@ void FMassEntityQuery::CacheArchetypes(const FMassEntityManager& InEntityManager
 	}
 }
 
-void FMassEntityQuery::ForEachEntityChunkInCollections(TConstArrayView<FMassArchetypeEntityCollection> EntityCollections, FMassEntityManager& EntityManager, FMassExecutionContext& ExecutionContext, const FMassExecuteFunction& ExecuteFunction)
+void FMassEntityQuery::ForEachEntityChunkInCollections(TConstArrayView<FMassArchetypeEntityCollection> EntityCollections
+	, FMassExecutionContext& ExecutionContext, const FMassExecuteFunction& ExecuteFunction)
 {
 	for (const FMassArchetypeEntityCollection& EntityCollection : EntityCollections)
 	{
-		ForEachEntityChunk(EntityCollection, EntityManager, ExecutionContext, ExecuteFunction);
+		ForEachEntityChunk(EntityCollection, ExecutionContext, ExecuteFunction);
 	}
 }
 
-void FMassEntityQuery::ForEachEntityChunk(const FMassArchetypeEntityCollection& EntityCollection, FMassEntityManager& EntityManager, FMassExecutionContext& ExecutionContext, const FMassExecuteFunction& ExecuteFunction)
+void FMassEntityQuery::ForEachEntityChunk(const FMassArchetypeEntityCollection& EntityCollection, FMassExecutionContext& ExecutionContext
+	, const FMassExecuteFunction& ExecuteFunction)
 {
 	// mz@todo I don't like that we're copying data here.
 	ExecutionContext.SetEntityCollection(EntityCollection);
-	ForEachEntityChunk(EntityManager, ExecutionContext, ExecuteFunction);
+	ForEachEntityChunk(ExecutionContext, ExecuteFunction);
 	ExecutionContext.ClearEntityCollection();
 }
 
-void FMassEntityQuery::ForEachEntityChunk(FMassEntityManager& EntityManager, FMassExecutionContext& ExecutionContext, const FMassExecuteFunction& ExecuteFunction)
+void FMassEntityQuery::ForEachEntityChunk(FMassExecutionContext& ExecutionContext, const FMassExecuteFunction& ExecuteFunction)
 {
+	checkf(CachedEntityManager == ExecutionContext.GetSharedEntityManager() && CachedEntityManager.IsValid()
+		, TEXT("FMassEntityQuery needs to be initialized with a valid EntityManager and the execution context has to come from the same manager"));
+
 #if WITH_MASSENTITY_DEBUG
-	checkf(ExecutionContext.ExecutionType == ExpectedContextType && (ExpectedContextType == EMassExecutionContextType::Local || bRegistered)
+	checkf(ExecutionContext.GetExecutionType() == ExpectedContextType && (ExpectedContextType == EMassExecutionContextType::Local || bRegistered)
 		, TEXT("ExecutionContextType mismatch, make sure all the queries run as part of processor execution are registered with some processor with a FMassEntityQuery::RegisterWithProcessor call"));
 
-	EntityManager.GetRequirementAccessDetector().RequireAccess(*this);
+	CachedEntityManager->GetRequirementAccessDetector().RequireAccess(*this);
 #endif
 
-	FScopedSubsystemRequirementsRestore SubsystemRestore(ExecutionContext);
+	FScopedEntityQueryContext ScopedQueryContext(*this, ExecutionContext);
 
-	if (ExecutionContext.CacheSubsystemRequirements(*this) == false)
+	if (ScopedQueryContext.IsSuccessfullySetUp() == false)
 	{
 		// required subsystems are not available, bail out.
 		return;
@@ -233,7 +258,7 @@ void FMassEntityQuery::ForEachEntityChunk(FMassEntityManager& EntityManager, FMa
 		// note that the following function will usually only resort to verifying that the data is up to date by
 			// checking the version number. In rare cases when it would result in non trivial cost we actually
 			// do need those calculations.
-		CacheArchetypes(EntityManager);
+		CacheArchetypes();
 
 		// if there's a chunk collection set by the external code - use that
 		if (ExecutionContext.GetEntityCollection().IsSet())
@@ -244,16 +269,16 @@ void FMassEntityQuery::ForEachEntityChunk(FMassEntityManager& EntityManager, FMa
 			// if given ArchetypeHandle cannot be found in ValidArchetypes then it doesn't match the query's requirements
 			if (ArchetypeIndex == INDEX_NONE)
 			{
-				UE_VLOG_UELOG(EntityManager.GetOwner(), LogMass, Log, TEXT("Attempted to execute FMassEntityQuery with an incompatible Archetype: %s")
+				UE_VLOG_UELOG(CachedEntityManager->GetOwner(), LogMass, Log, TEXT("Attempted to execute FMassEntityQuery with an incompatible Archetype: %s")
 					, *FMassDebugger::GetArchetypeRequirementCompatibilityDescription(*this, ArchetypeHandle));
 
 #if WITH_MASSENTITY_DEBUG
-				EntityManager.GetRequirementAccessDetector().ReleaseAccess(*this);
+				CachedEntityManager->GetRequirementAccessDetector().ReleaseAccess(*this);
 #endif // WITH_MASSENTITY_DEBUG
 				return;
 			}
 
-			ExecutionContext.SetFragmentRequirements(*this);
+			ExecutionContext.ApplyFragmentRequirements(*this);
 
 			FMassArchetypeData& ArchetypeData = FMassArchetypeHelper::ArchetypeDataFromHandleChecked(ArchetypeHandle);
 			ArchetypeData.ExecuteFunction(ExecutionContext, ExecuteFunction
@@ -264,7 +289,7 @@ void FMassEntityQuery::ForEachEntityChunk(FMassEntityManager& EntityManager, FMa
 		else
 		{
 			// it's important to set requirements after caching archetypes due to that call potentially sorting the requirements and the order is relevant here.
-			ExecutionContext.SetFragmentRequirements(*this);
+			ExecutionContext.ApplyFragmentRequirements(*this);
 
 			// note that this checkSlow is here on purpose, for debugging purposes, not data verification purposes.
 			checkSlow(OrderedArchetypeIndices.Num() == ValidArchetypes.Num());
@@ -273,71 +298,56 @@ void FMassEntityQuery::ForEachEntityChunk(FMassEntityManager& EntityManager, FMa
 				const FMassArchetypeHandle& ArchetypeHandle = ValidArchetypes[ArchetypeIndex];
 				FMassArchetypeData& ArchetypeData = FMassArchetypeHelper::ArchetypeDataFromHandleChecked(ArchetypeHandle);
 				ArchetypeData.ExecuteFunction(ExecutionContext, ExecuteFunction, ArchetypeFragmentMapping[ArchetypeIndex], ChunkCondition);
-				ExecutionContext.ClearFragmentViews();
+				ExecutionContext.ClearFragmentViews(*this);
 			}
 		}
 	}
 
 #if WITH_MASSENTITY_DEBUG
-	EntityManager.GetRequirementAccessDetector().ReleaseAccess(*this);
+	CachedEntityManager->GetRequirementAccessDetector().ReleaseAccess(*this);
 #endif
-
-	ExecutionContext.ClearExecutionData();
-	ExecutionContext.FlushDeferred();
 }
 
 void FMassEntityQuery::ParallelForEachEntityChunkInCollection(TConstArrayView<FMassArchetypeEntityCollection> EntityCollections
-	, FMassEntityManager& EntityManager, FMassExecutionContext& ExecutionContext, const FMassExecuteFunction& ExecuteFunction
-	, const EParallelForMode ParallelMode)
-{
-	EParallelExecutionFlags Flags = EnumHasAnyFlags(ParallelMode, ForceParallelExecution) ? EParallelExecutionFlags::Force : EParallelExecutionFlags::Default;
-	ParallelForEachEntityChunkInCollection(EntityCollections, EntityManager, ExecutionContext, ExecuteFunction, Flags);
-}
-
-void FMassEntityQuery::ParallelForEachEntityChunkInCollection(TConstArrayView<FMassArchetypeEntityCollection> EntityCollections
-	, FMassEntityManager& EntityManager, FMassExecutionContext& ExecutionContext, const FMassExecuteFunction& ExecuteFunction
+	, FMassExecutionContext& ExecutionContext, const FMassExecuteFunction& ExecuteFunction
 	, const EParallelExecutionFlags Flags)
 {
 	if (UE::Mass::Tweakables::bAllowParallelExecution == false && !EnumHasAnyFlags(Flags, EParallelExecutionFlags::Force))
 	{
-		ForEachEntityChunkInCollections(EntityCollections, EntityManager, ExecutionContext, ExecuteFunction);
+		ForEachEntityChunkInCollections(EntityCollections, ExecutionContext, ExecuteFunction);
 		return;
 	}
 
-	ParallelFor(EntityCollections.Num(), [this, &EntityManager, &ExecutionContext, &ExecuteFunction, &EntityCollections, Flags](const int32 JobIndex)
+	ParallelFor(EntityCollections.Num(), [this, &ExecutionContext, &ExecuteFunction, &EntityCollections, Flags](const int32 JobIndex)
 	{
 		FMassExecutionContext LocalExecutionContext = ExecutionContext; 
 		LocalExecutionContext.SetEntityCollection(EntityCollections[JobIndex]);
-		ParallelForEachEntityChunk(EntityManager, LocalExecutionContext, ExecuteFunction, Flags);
+		ParallelForEachEntityChunk(LocalExecutionContext, ExecuteFunction, Flags);
 	});
 }
 
-void FMassEntityQuery::ParallelForEachEntityChunk(FMassEntityManager& EntityManager, FMassExecutionContext& ExecutionContext
-	, const FMassExecuteFunction& ExecuteFunction, const EParallelForMode ParallelMode)
-{
-	EParallelExecutionFlags Flags = EnumHasAnyFlags(ParallelMode, ForceParallelExecution) ? EParallelExecutionFlags::Force : EParallelExecutionFlags::Default;
-	ParallelForEachEntityChunk(EntityManager, ExecutionContext, ExecuteFunction, Flags);
-}
-
-void FMassEntityQuery::ParallelForEachEntityChunk(FMassEntityManager& EntityManager, FMassExecutionContext& ExecutionContext
+void FMassEntityQuery::ParallelForEachEntityChunk(FMassExecutionContext& ExecutionContext
 	, const FMassExecuteFunction& ExecuteFunction, const EParallelExecutionFlags Flags)
 {
 	if (UE::Mass::Tweakables::bAllowParallelExecution == false && !EnumHasAnyFlags(Flags, EParallelExecutionFlags::Force))
 	{
-		ForEachEntityChunk(EntityManager, ExecutionContext, ExecuteFunction);
+		ForEachEntityChunk(ExecutionContext, ExecuteFunction);
 		return;
 	}
 
+	checkf(CachedEntityManager == ExecutionContext.GetSharedEntityManager() && CachedEntityManager.IsValid()
+		, TEXT("FMassEntityQuery needs to be initialized with a valid EntityManager and the execution context has to come from the same manager"));
+
 #if WITH_MASSENTITY_DEBUG
-	checkf(ExecutionContext.ExecutionType == ExpectedContextType && (ExpectedContextType == EMassExecutionContextType::Local || bRegistered)
+	checkf(ExecutionContext.GetExecutionType() == ExpectedContextType && (ExpectedContextType == EMassExecutionContextType::Local || bRegistered)
 		, TEXT("ExecutionContextType mismatch, make sure all the queries run as part of processor execution are registered with some processor with a FMassEntityQuery::RegisterWithProcessor call"));
 
-	EntityManager.GetRequirementAccessDetector().RequireAccess(*this);
+	CachedEntityManager->GetRequirementAccessDetector().RequireAccess(*this);
 #endif
 
-	FScopedSubsystemRequirementsRestore SubsystemRestore(ExecutionContext);
+	FScopedEntityQueryContext ScopedQueryContext(*this, ExecutionContext);
 
-	if (ExecutionContext.CacheSubsystemRequirements(*this) == false)
+	if (ScopedQueryContext.IsSuccessfullySetUp() == false)
 	{
 		// required subsystems are not available, bail out.
 		return;
@@ -369,7 +379,7 @@ void FMassEntityQuery::ParallelForEachEntityChunk(FMassEntityManager& EntityMana
 		// note that the following function will usualy only resort to verifying that the data is up to date by
 		// checking the version number. In rare cases when it would result in non trivial cost we actually
 		// do need those calculations.
-		CacheArchetypes(EntityManager);
+		CacheArchetypes();
 
 		// if there's a chunk collection set by the external code - use that
 		if (ExecutionContext.GetEntityCollection().IsSet())
@@ -380,16 +390,16 @@ void FMassEntityQuery::ParallelForEachEntityChunk(FMassEntityManager& EntityMana
 			// if given ArchetypeHandle cannot be found in ValidArchetypes then it doesn't match the query's requirements
 			if (ArchetypeIndex == INDEX_NONE)
 			{
-				UE_VLOG_UELOG(EntityManager.GetOwner(), LogMass, Log, TEXT("Attempted to execute FMassEntityQuery with an incompatible Archetype: %s")
+				UE_VLOG_UELOG(CachedEntityManager->GetOwner(), LogMass, Log, TEXT("Attempted to execute FMassEntityQuery with an incompatible Archetype: %s")
 					, *FMassDebugger::GetArchetypeRequirementCompatibilityDescription(*this, ExecutionContext.GetEntityCollection().GetArchetype()));
 
 #if WITH_MASSENTITY_DEBUG
-				EntityManager.GetRequirementAccessDetector().ReleaseAccess(*this);
+				CachedEntityManager->GetRequirementAccessDetector().ReleaseAccess(*this);
 #endif // WITH_MASSENTITY_DEBUG
 				return;
 			}
 
-			ExecutionContext.SetFragmentRequirements(*this);
+			ExecutionContext.ApplyFragmentRequirements(*this);
 
 			FMassArchetypeData& ArchetypeRef = FMassArchetypeHelper::ArchetypeDataFromHandleChecked(ArchetypeHandle);
 			for (const FMassArchetypeEntityCollection::FArchetypeEntityRange& EntityRange : ExecutionContext.GetEntityCollection().GetRanges())
@@ -399,7 +409,7 @@ void FMassEntityQuery::ParallelForEachEntityChunk(FMassEntityManager& EntityMana
 		}
 		else
 		{
-			ExecutionContext.SetFragmentRequirements(*this);
+			ExecutionContext.ApplyFragmentRequirements(*this);
 			for (int ArchetypeIndex = 0; ArchetypeIndex < ValidArchetypes.Num(); ++ArchetypeIndex)
 			{
 				FMassArchetypeHandle& ArchetypeHandle = ValidArchetypes[ArchetypeIndex];
@@ -440,14 +450,12 @@ void FMassEntityQuery::ParallelForEachEntityChunk(FMassEntityManager& EntityMana
 
 			ParallelForWithTaskContext(TaskContext, Jobs.Num(), [this, &ExecutionContext, &ExecuteFunction, &Jobs](FTaskContext& TaskContext, const int32 JobIndex)
 				{
-					FMassExecutionContext LocalExecutionContext = ExecutionContext;
-
-					LocalExecutionContext.SetDeferredCommandBuffer(TaskContext.GetCommandBuffer());
-
+					FMassExecutionContext LocalExecutionContext(ExecutionContext, *this, TaskContext.GetCommandBuffer());
 					Jobs[JobIndex].Archetype.ExecutionFunctionForChunk(LocalExecutionContext, ExecuteFunction
 						, Jobs[JobIndex].ArchetypeIndex != INDEX_NONE ? ArchetypeFragmentMapping[Jobs[JobIndex].ArchetypeIndex] : FMassQueryRequirementIndicesMapping()
 						, Jobs[JobIndex].EntityRange
 						, ChunkCondition);
+					LocalExecutionContext.PopQuery(*this);
 				}, ParallelForFlags);
 
 			// merge all command buffers
@@ -460,26 +468,24 @@ void FMassEntityQuery::ParallelForEachEntityChunk(FMassEntityManager& EntityMana
 		{
 			ParallelFor(Jobs.Num(), [this, &ExecutionContext, &ExecuteFunction, &Jobs](const int32 JobIndex)
 				{
-					FMassExecutionContext LocalExecutionContext = ExecutionContext;
+					FMassExecutionContext LocalExecutionContext(ExecutionContext, *this);
 					Jobs[JobIndex].Archetype.ExecutionFunctionForChunk(LocalExecutionContext, ExecuteFunction
 						, Jobs[JobIndex].ArchetypeIndex != INDEX_NONE ? ArchetypeFragmentMapping[Jobs[JobIndex].ArchetypeIndex] : FMassQueryRequirementIndicesMapping()
 						, Jobs[JobIndex].EntityRange
 						, ChunkCondition);
+					LocalExecutionContext.PopQuery(*this);
 				}, ParallelForFlags);
 		}
 	}
 
 #if WITH_MASSENTITY_DEBUG
-	EntityManager.GetRequirementAccessDetector().ReleaseAccess(*this);
+	CachedEntityManager->GetRequirementAccessDetector().ReleaseAccess(*this);
 #endif
-
-	ExecutionContext.ClearExecutionData();
-	ExecutionContext.FlushDeferred();
 }
 
-int32 FMassEntityQuery::GetNumMatchingEntities(FMassEntityManager& InEntityManager)
+int32 FMassEntityQuery::GetNumMatchingEntities()
 {
-	CacheArchetypes(InEntityManager);
+	CacheArchetypes();
 	int32 TotalEntities = 0;
 	for (FMassArchetypeHandle& ArchetypeHandle : ValidArchetypes)
 	{
@@ -507,9 +513,9 @@ int32 FMassEntityQuery::GetNumMatchingEntities(TConstArrayView<FMassArchetypeEnt
 	return TotalEntities;
 }
 
-bool FMassEntityQuery::HasMatchingEntities(FMassEntityManager& InEntityManager)
+bool FMassEntityQuery::HasMatchingEntities()
 {
-	CacheArchetypes(InEntityManager);
+	CacheArchetypes();
 
 	for (FMassArchetypeHandle& ArchetypeHandle : ValidArchetypes)
 	{
@@ -649,12 +655,73 @@ void FMassEntityQuery::ExportRequirements(FMassExecutionRequirements& OutRequire
 //-----------------------------------------------------------------------------
 // DEPRECATED
 //-----------------------------------------------------------------------------
+// deprecated
 FMassEntityQuery::FMassEntityQuery(std::initializer_list<UScriptStruct*>)
 	: FMassEntityQuery()
 {
 }
 
+// deprecated
 FMassEntityQuery::FMassEntityQuery(TConstArrayView<const UScriptStruct*>)
 	: FMassEntityQuery()
 {
 }
+
+// deprecated
+void FMassEntityQuery::ForEachEntityChunk(FMassEntityManager&, FMassExecutionContext& ExecutionContext, const FMassExecuteFunction& ExecuteFunction)
+{
+	ForEachEntityChunk(ExecutionContext, ExecuteFunction);
+}
+
+// deprecated
+void FMassEntityQuery::ForEachEntityChunk(const FMassArchetypeEntityCollection& EntityCollection, FMassEntityManager&
+		, FMassExecutionContext& ExecutionContext, const FMassExecuteFunction& ExecuteFunction)
+{
+	ForEachEntityChunk(EntityCollection, ExecutionContext, ExecuteFunction);
+}
+
+// deprecated
+void FMassEntityQuery::ParallelForEachEntityChunk(FMassEntityManager&, FMassExecutionContext& ExecutionContext
+		, const FMassExecuteFunction& ExecuteFunction, const EParallelForMode ParallelMode)
+{
+	const EParallelExecutionFlags Flags = EnumHasAnyFlags(ParallelMode, ForceParallelExecution)
+		? EParallelExecutionFlags::Force
+		: EParallelExecutionFlags::Default;
+	ParallelForEachEntityChunk(ExecutionContext, ExecuteFunction, Flags);
+}
+
+// deprecated
+void FMassEntityQuery::ForEachEntityChunkInCollections(TConstArrayView<FMassArchetypeEntityCollection> EntityCollections, FMassEntityManager&
+		, FMassExecutionContext& ExecutionContext, const FMassExecuteFunction& ExecuteFunction)
+{
+	ForEachEntityChunkInCollections(EntityCollections, ExecutionContext, ExecuteFunction);
+}
+
+// deprecated
+void FMassEntityQuery::ParallelForEachEntityChunkInCollection(TConstArrayView<FMassArchetypeEntityCollection> EntityCollections
+		, FMassEntityManager&, FMassExecutionContext& ExecutionContext, const FMassExecuteFunction& ExecuteFunction
+		, const EParallelForMode ParallelMode)
+{
+	const EParallelExecutionFlags Flags = EnumHasAnyFlags(ParallelMode, ForceParallelExecution)
+		? EParallelExecutionFlags::Force
+		: EParallelExecutionFlags::Default;
+	ParallelForEachEntityChunkInCollection(EntityCollections, ExecutionContext, ExecuteFunction, Flags);
+}
+
+// deprecated
+void FMassEntityQuery::CacheArchetypes(const FMassEntityManager&)
+{
+	CacheArchetypes();
+}
+
+// deprecated
+int32 FMassEntityQuery::GetNumMatchingEntities(FMassEntityManager&)
+{
+	return GetNumMatchingEntities();
+}
+
+// deprecated
+bool FMassEntityQuery::HasMatchingEntities(FMassEntityManager&)
+{
+	return HasMatchingEntities();
+}
diff --git a/Engine/Source/Runtime/MassEntity/Private/MassExecutionContext.cpp b/Engine/Source/Runtime/MassEntity/Private/MassExecutionContext.cpp
--- a/Engine/Source/Runtime/MassEntity/Private/MassExecutionContext.cpp
+++ b/Engine/Source/Runtime/MassEntity/Private/MassExecutionContext.cpp
@@ -16,6 +16,27 @@ FMassExecutionContext::FMassExecutionContext(FMassEntityManager& InEntityManager
 {
 }
 
+FMassExecutionContext::FMassExecutionContext(const FMassExecutionContext& Other)
+	: EntityManager(Other.EntityManager)
+{
+	// we're using assignment operator here as a setup helper, as per operator=='s comment.
+	*this = Other;
+
+	QueriesStack.Reset();
+}
+
+FMassExecutionContext::FMassExecutionContext(const FMassExecutionContext& Other, FMassEntityQuery& Query, const TSharedPtr<FMassCommandBuffer>& InCommandBuffer)
+	: FMassExecutionContext(Other)
+{
+	QueriesStack.Add({&Query});
+	SetDeferredCommandBuffer(InCommandBuffer);
+}
+
+FMassExecutionContext::~FMassExecutionContext()
+{
+	ensureMsgf(QueriesStack.Num() == 0, TEXT("Destroying a FMassExecutionContext instance while not all queries have been popped is unexpected."));
+}
+
 void FMassExecutionContext::FlushDeferred()
 {
 	if (bFlushDeferredCommands && DeferredCommandBuffer)
@@ -95,3 +116,19 @@ UWorld* FMassExecutionContext::GetWorld()
 { 
 	return EntityManager->GetWorld(); 
 }
+
+void FMassExecutionContext::PushQuery(FMassEntityQuery& InQuery)
+{
+	FQueryTransientRuntime& RuntimeData = QueriesStack.Add_GetRef({&InQuery});
+	GetSubsystemRequirementBits(RuntimeData.ConstSubsystemsBitSet, RuntimeData.MutableSubsystemsBitSet);
+}
+
+void FMassExecutionContext::PopQuery(const FMassEntityQuery& InQuery)
+{
+	const FQueryTransientRuntime& RuntimeData = QueriesStack.Last();
+	checkf(&InQuery == RuntimeData.Query, TEXT("Queries are stored in a stack and as such it requires elements to be added in LIFO order"));
+
+	SetSubsystemRequirementBits(RuntimeData.ConstSubsystemsBitSet, RuntimeData.MutableSubsystemsBitSet);
+
+	QueriesStack.RemoveAt(QueriesStack.Num() - 1, EAllowShrinking::No);
+}
diff --git a/Engine/Source/Runtime/MassEntity/Private/MassProcessor.cpp b/Engine/Source/Runtime/MassEntity/Private/MassProcessor.cpp
--- a/Engine/Source/Runtime/MassEntity/Private/MassProcessor.cpp
+++ b/Engine/Source/Runtime/MassEntity/Private/MassProcessor.cpp
@@ -187,7 +187,7 @@ void UMassProcessor::GetArchetypesMatchingOwnedQueries(const FMassEntityManager&
 	for (FMassEntityQuery* QueryPtr : OwnedQueries)
 	{
 		CA_ASSUME(QueryPtr);
-		QueryPtr->CacheArchetypes(EntityManager);
+		QueryPtr->CacheArchetypes();
 
 		for (const FMassArchetypeHandle& ArchetypeHandle : QueryPtr->GetArchetypes())
 		{
@@ -201,7 +201,7 @@ bool UMassProcessor::DoesAnyArchetypeMatchOwnedQueries(const FMassEntityManager&
 	for (FMassEntityQuery* QueryPtr : OwnedQueries)
 	{
 		CA_ASSUME(QueryPtr);
-		QueryPtr->CacheArchetypes(EntityManager);
+		QueryPtr->CacheArchetypes();
 
 		if (QueryPtr->GetArchetypes().Num() > 0)
 		{
diff --git a/Engine/Source/Runtime/MassEntity/Public/MassEntityQuery.h b/Engine/Source/Runtime/MassEntity/Public/MassEntityQuery.h
--- a/Engine/Source/Runtime/MassEntity/Public/MassEntityQuery.h
+++ b/Engine/Source/Runtime/MassEntity/Public/MassEntityQuery.h
@@ -88,29 +88,29 @@ public:
 	void RegisterWithProcessor(UMassProcessor& Owner);
 
 	/** Runs ExecuteFunction on all entities matching Requirements */
-	void ForEachEntityChunk(FMassEntityManager& EntityManager, FMassExecutionContext& ExecutionContext, const FMassExecuteFunction& ExecuteFunction);
+	void ForEachEntityChunk(FMassExecutionContext& ExecutionContext, const FMassExecuteFunction& ExecuteFunction);
 	
 	/** Will first verify that the archetype given with Collection matches the query's requirements, and if so will run the other, more generic ForEachEntityChunk implementation */
-	void ForEachEntityChunk(const FMassArchetypeEntityCollection& EntityCollection, FMassEntityManager& EntitySubsystem
+	void ForEachEntityChunk(const FMassArchetypeEntityCollection& EntityCollection
 		, FMassExecutionContext& ExecutionContext, const FMassExecuteFunction& ExecuteFunction);
 
 	/**
 	 * Attempts to process every chunk of every affected archetype in parallel.
 	 */
-	void ParallelForEachEntityChunk(FMassEntityManager& EntityManager, FMassExecutionContext& ExecutionContext
+	void ParallelForEachEntityChunk(FMassExecutionContext& ExecutionContext
 		, const FMassExecuteFunction& ExecuteFunction, const EParallelExecutionFlags Flags = EParallelExecutionFlags::Default);
 
-	void ForEachEntityChunkInCollections(TConstArrayView<FMassArchetypeEntityCollection> EntityCollections, FMassEntityManager& EntityManager
+	void ForEachEntityChunkInCollections(TConstArrayView<FMassArchetypeEntityCollection> EntityCollections
 		, FMassExecutionContext& ExecutionContext, const FMassExecuteFunction& ExecuteFunction);
 
 	void ParallelForEachEntityChunkInCollection(TConstArrayView<FMassArchetypeEntityCollection> EntityCollections
-		, FMassEntityManager& EntityManager, FMassExecutionContext& ExecutionContext, const FMassExecuteFunction& ExecuteFunction
-		, const EParallelExecutionFlags Flags);
+		, FMassExecutionContext& ExecutionContext, const FMassExecuteFunction& ExecuteFunction
+		, const EParallelExecutionFlags Flags = EParallelExecutionFlags::Default);
 
 	/** Will gather all archetypes from InEntityManager matching this->Requirements.
 	 *  Note that no work will be done if the cached data is up to date (as tracked by EntitySubsystemHash and 
 	 *	ArchetypeDataVersion properties). */
-	void CacheArchetypes(const FMassEntityManager& InEntityManager);
+	void CacheArchetypes();
 
 	void Clear()
 	{
@@ -147,7 +147,7 @@ public:
 	 * Note that the function is not const because calling it can result in re-caching of ValidArchetypes 
 	 * @return the number of entities this given query would process if called "now"
 	 */
-	int32 GetNumMatchingEntities(FMassEntityManager& InEntityManager);
+	int32 GetNumMatchingEntities();
 
 	/** 
 	 * Sums the entity range lengths for each collection in EntityCollections, where the collection's 
@@ -161,14 +161,14 @@ public:
 	 * Note that the function is not const because calling it can result in re-caching of ValidArchetypes
 	 * @return "true" if any of the ValidArchetypes has any entities, "false" otherwise
 	 */
-	bool HasMatchingEntities(FMassEntityManager& InEntityManager);
+	bool HasMatchingEntities();
 
 	/** 
 	 * Sets a chunk filter condition that will be applied to each chunk of all valid archetypes. Note 
 	 * that this condition won't be applied when a specific entity collection is used (via FMassArchetypeEntityCollection )
 	 * The value returned by InFunction controls whether to allow execution (true) or block it (false).
 	 */
-	void SetChunkFilter(const FMassChunkConditionFunction& InFunction) { checkf(!HasChunkFilter(), TEXT("Chunk filter needs to be cleared before setting a new one.")); ChunkCondition = InFunction; }
+	void SetChunkFilter(const FMassChunkConditionFunction& InFunction);
 
 	void ClearChunkFilter() { ChunkCondition.Reset(); }
 
@@ -198,16 +198,6 @@ public:
 	void SetParallelCommandBufferEnabled(const bool bInAllowParallelCommands) { bAllowParallelCommands = bInAllowParallelCommands; }
 
 private:
-	struct FScopedSubsystemRequirementsRestore
-	{
-		FScopedSubsystemRequirementsRestore(FMassExecutionContext& ExecutionContext);
-		~FScopedSubsystemRequirementsRestore();
-
-		FMassExecutionContext& CachedExecutionContext;
-		FMassExternalSubsystemBitSet ConstSubsystemsBitSet;
-		FMassExternalSubsystemBitSet MutableSubsystemsBitSet;
-	};
-
 	/**
 	 * Incrementally sorts all ValidArchetypes to fill OrderedArchetypeIndices with the expected order of archetype processing.
 	 * This function will only ever get called when there are actual sorting steps registered (@see GroupBy)
@@ -286,18 +276,41 @@ public:
 		ForceParallelExecution = static_cast<int32>(EParallelExecutionFlags::Force),
 	};
 
-	UE_DEPRECATED(5.6, "The version of ParallelForEachEntityChunk using EParallelForMode has been replaced by the version using EParallelExecutionFlags.")
-	void ParallelForEachEntityChunk(FMassEntityManager& EntityManager, FMassExecutionContext& ExecutionContext
-		, const FMassExecuteFunction& ExecuteFunction, const EParallelForMode ParallelMode);
+	UE_DEPRECATED(5.6, "ForEachEntityChunk is deprecated. New version doesn't require the FMassEntityManager parameter")
+	void ForEachEntityChunk(FMassEntityManager& EntityManager, FMassExecutionContext& ExecutionContext, const FMassExecuteFunction& ExecuteFunction);
 	
-	UE_DEPRECATED(5.6, "The version of ParallelForEachEntityChunk using EParallelForMode has been replaced by the version using EParallelExecutionFlags.")
+	UE_DEPRECATED(5.6, "ForEachEntityChunk is deprecated. New version doesn't require the FMassEntityManager parameter")
+	void ForEachEntityChunk(const FMassArchetypeEntityCollection& EntityCollection, FMassEntityManager& EntityManager
+		, FMassExecutionContext& ExecutionContext, const FMassExecuteFunction& ExecuteFunction);
+
+	UE_DEPRECATED(5.6, "ParallelForEachEntityChunk is deprecated. New version doesn't require the FMassEntityManager parameter. Also ParallelMode parameter changed type, usee EParallelExecutionFlags instead.")
+	void ParallelForEachEntityChunk(FMassEntityManager& EntityManager, FMassExecutionContext& ExecutionContext
+		, const FMassExecuteFunction& ExecuteFunction, const EParallelForMode ParallelMode = Default);
+
+	UE_DEPRECATED(5.6, "ForEachEntityChunkInCollections is deprecated. New version doesn't require the FMassEntityManager parameter")
+	void ForEachEntityChunkInCollections(TConstArrayView<FMassArchetypeEntityCollection> EntityCollections, FMassEntityManager& EntityManager
+		, FMassExecutionContext& ExecutionContext, const FMassExecuteFunction& ExecuteFunction);
+
+	UE_DEPRECATED(5.6, "ParallelForEachEntityChunkInCollection is deprecated. New version doesn't require the FMassEntityManager parameter. Also ParallelMode parameter changed type, usee EParallelExecutionFlags instead.")
 	void ParallelForEachEntityChunkInCollection(TConstArrayView<FMassArchetypeEntityCollection> EntityCollections
 		, FMassEntityManager& EntityManager, FMassExecutionContext& ExecutionContext, const FMassExecuteFunction& ExecuteFunction
 		, const EParallelForMode ParallelMode);
+
+	UE_DEPRECATED(5.6, "This flavor of CacheArchetypes is deprecated. EntityQueries are not tied to a specific EntityManager instance and there's no need to pass it in as a parameter")
+	void CacheArchetypes(const FMassEntityManager& InEntityManager);
+
+	UE_DEPRECATED(5.6, "This flavor of GetNumMatchingEntities is deprecated. EntityQueries are not tied to a specific EntityManager instance and there's no need to pass it in as a parameter")
+	int32 GetNumMatchingEntities(FMassEntityManager& InEntityManager);
+
+	UE_DEPRECATED(5.6, "This flavor of HasMatchingEntities is deprecated. EntityQueries are not tied to a specific EntityManager instance and there's no need to pass it in as a parameter")
+	bool HasMatchingEntities(FMassEntityManager& InEntityManager);
 };
 
 ENUM_CLASS_FLAGS(FMassEntityQuery::EParallelExecutionFlags);
 
+//-----------------------------------------------------------------------------
+// INLINES
+//-----------------------------------------------------------------------------
 #if WITH_ARCHETYPE_MATCH_OVERRIDE
 template <typename T> requires TArchetypeMatchOverrideConcept<T>
 void FMassEntityQuery::SetArchetypeMatchOverride(const T& Context)
@@ -319,5 +332,11 @@ void FMassEntityQuery::SetArchetypeMatchOverride(const T& Context)
 }
 #endif
 
+inline void FMassEntityQuery::SetChunkFilter(const FMassChunkConditionFunction& InFunction)
+{
+	checkf(!HasChunkFilter(), TEXT("Chunk filter needs to be cleared before setting a new one."));
+	ChunkCondition = InFunction;
+}
+
 #undef MASS_ACHETYPE_OVERRIDE_MAX_SIZE
 #undef MASS_ACHETYPE_OVERRIDE_MAX_ALIGNMENT
\ No newline at end of file
diff --git a/Engine/Source/Runtime/MassEntity/Public/MassExecutionContext.h b/Engine/Source/Runtime/MassEntity/Public/MassExecutionContext.h
--- a/Engine/Source/Runtime/MassEntity/Public/MassExecutionContext.h
+++ b/Engine/Source/Runtime/MassEntity/Public/MassExecutionContext.h
@@ -10,6 +10,7 @@
 #include "MassEntityQuery.h"
 #endif // UE_ENABLE_INCLUDE_ORDER_DEPRECATED_IN_5_6
 #include "MassArchetypeTypes.h"
+#include "MassEntityQuery.h"
 #include "MassSubsystemAccess.h"
 
 
@@ -69,6 +70,15 @@ private:
 
 	TSharedRef<FMassEntityManager> EntityManager;
 
+	struct FQueryTransientRuntime
+	{
+		gsl::not_null<FMassEntityQuery*> Query;
+		FMassExternalSubsystemBitSet ConstSubsystemsBitSet;
+		FMassExternalSubsystemBitSet MutableSubsystemsBitSet;
+	};
+	/** We usually expect the queries to go only a single layer deep, so 2 elements here should suffice most of the time */
+	TArray<FQueryTransientRuntime, TInlineAllocator<2>> QueriesStack;
+
 #if WITH_MASSENTITY_DEBUG
 	FString DebugExecutionDescription;
 #endif
@@ -95,32 +105,48 @@ private:
 	EMassExecutionContextType ExecutionType = EMassExecutionContextType::Local;
 
 	friend FMassArchetypeData;
-	friend FMassEntityQuery;
+
+	/**
+	 * Note that this operator is private on purpose, used to simplify the implementation of constructors.
+	 * The context itself does not support assignment.
+	 */
+	FMassExecutionContext& operator=(const FMassExecutionContext& Other) = default;
 
 public:
 	explicit FMassExecutionContext(FMassEntityManager& InEntityManager, const float InDeltaTimeSeconds = 0.f, const bool bInFlushDeferredCommands = true);
+	FMassExecutionContext(FMassExecutionContext&& Other) = default;
+	FMassExecutionContext(const FMassExecutionContext& Other);
+	FMassExecutionContext(const FMassExecutionContext& Other, FMassEntityQuery& Query, const TSharedPtr<FMassCommandBuffer>& InCommandBuffer = {});
+	~FMassExecutionContext();
 
-	FMassEntityManager& GetEntityManagerChecked() { return EntityManager.Get(); }
+	FMassEntityManager& GetEntityManagerChecked();
+	const TSharedRef<FMassEntityManager>& GetSharedEntityManager();
 
 #if WITH_MASSENTITY_DEBUG
 	const FString& DebugGetExecutionDesc() const { return DebugExecutionDescription; }
 	void DebugSetExecutionDesc(const FString& Description) { DebugExecutionDescription = Description; }
 #endif
 
+	void PushQuery(FMassEntityQuery& InQuery);
+	void PopQuery(const FMassEntityQuery& InQuery);
+	const FMassEntityQuery& GetCurrentQuery() const;
+	bool IsCurrentQuery(const FMassEntityQuery& Query) const;
+	void ApplyFragmentRequirements(const FMassEntityQuery& RequestingQuery);
+	void ClearFragmentViews(const FMassEntityQuery& RequestingQuery);
+
 	/** Sets bFlushDeferredCommands. Note that setting to True while the system is being executed doesn't result in
 	 *  immediate commands flushing */
-	void SetFlushDeferredCommands(const bool bNewFlushDeferredCommands) { bFlushDeferredCommands = bNewFlushDeferredCommands; } 
-	void SetDeferredCommandBuffer(const TSharedPtr<FMassCommandBuffer>& InDeferredCommandBuffer) { DeferredCommandBuffer = InDeferredCommandBuffer; }
+	void SetFlushDeferredCommands(const bool bNewFlushDeferredCommands);
+	void SetDeferredCommandBuffer(const TSharedPtr<FMassCommandBuffer>& InDeferredCommandBuffer);
 	void SetEntityCollection(const FMassArchetypeEntityCollection& InEntityCollection);
 	void SetEntityCollection(FMassArchetypeEntityCollection&& InEntityCollection);
-	void ClearEntityCollection() { EntityCollection.Reset(); }
-	void SetAuxData(const FInstancedStruct& InAuxData) { AuxData = InAuxData; }
-	void SetExecutionType(EMassExecutionContextType InExecutionType) { check(InExecutionType != EMassExecutionContextType::MAX); ExecutionType = InExecutionType; }
+	void ClearEntityCollection();
+	void SetAuxData(const FInstancedStruct& InAuxData);
 
-	float GetDeltaTimeSeconds() const
-	{
-		return DeltaTimeSeconds;
-	}
+	void SetExecutionType(EMassExecutionContextType InExecutionType);
+	EMassExecutionContextType GetExecutionType() const;
+
+	float GetDeltaTimeSeconds() const;
 
 	UWorld* GetWorld();
 
@@ -410,5 +436,77 @@ public:
 	void SetCurrentArchetypesTagBitSet(const FMassTagBitSet&) {}
 };
 
+//-----------------------------------------------------------------------------
+// Inlines
+//-----------------------------------------------------------------------------
+inline FMassEntityManager& FMassExecutionContext::GetEntityManagerChecked()
+{
+	return EntityManager.Get();
+}
+
+inline const TSharedRef<FMassEntityManager>& FMassExecutionContext::GetSharedEntityManager()
+{
+	return EntityManager;
+}
+
+inline void FMassExecutionContext::SetFlushDeferredCommands(const bool bNewFlushDeferredCommands)
+{
+	bFlushDeferredCommands = bNewFlushDeferredCommands;
+}
+
+inline void FMassExecutionContext::SetDeferredCommandBuffer(const TSharedPtr<FMassCommandBuffer>& InDeferredCommandBuffer)
+{
+	DeferredCommandBuffer = InDeferredCommandBuffer;
+}
+
+inline void FMassExecutionContext::ClearEntityCollection()
+{
+	EntityCollection.Reset();
+}
+
+inline void FMassExecutionContext::SetAuxData(const FInstancedStruct& InAuxData)
+{
+	AuxData = InAuxData;
+}
+
+inline float FMassExecutionContext::GetDeltaTimeSeconds() const
+{
+	return DeltaTimeSeconds;
+}
+
+inline void FMassExecutionContext::SetExecutionType(EMassExecutionContextType InExecutionType)
+{
+	check(InExecutionType != EMassExecutionContextType::MAX);
+	ExecutionType = InExecutionType;
+}
+
+inline EMassExecutionContextType FMassExecutionContext::GetExecutionType() const
+{
+	return ExecutionType;
+}
+
+inline const FMassEntityQuery& FMassExecutionContext::GetCurrentQuery() const
+{
+	check(QueriesStack.Num());
+	return *QueriesStack.Last().Query;
+}
+
+inline bool FMassExecutionContext::IsCurrentQuery(const FMassEntityQuery& Query) const
+{
+	return QueriesStack.Num() && QueriesStack.Last().Query == &Query;
+}
+
+inline void FMassExecutionContext::ApplyFragmentRequirements(const FMassEntityQuery& RequestingQuery)
+{
+	check(IsCurrentQuery(RequestingQuery));
+	SetFragmentRequirements(RequestingQuery);
+}
+
+inline void FMassExecutionContext::ClearFragmentViews(const FMassEntityQuery& RequestingQuery)
+{
+	check(IsCurrentQuery(RequestingQuery));
+	ClearFragmentViews();
+}
+
 #undef CHECK_IF_VALID
 #undef CHECK_IF_READWRITE
