Mass: Allow replacing shared fragments when moving entity to another archetype

diff --git a/Engine/Source/Runtime/MassEntity/Private/MassEntityManager.cpp b/Engine/Source/Runtime/MassEntity/Private/MassEntityManager.cpp
--- a/Engine/Source/Runtime/MassEntity/Private/MassEntityManager.cpp
+++ b/Engine/Source/Runtime/MassEntity/Private/MassEntityManager.cpp
@@ -1816,7 +1816,8 @@ void FMassEntityManager::BatchAddSharedFragmentsForEntities(TConstArrayView<FMas
 	}
 }
 
-void FMassEntityManager::MoveEntityToAnotherArchetype(FMassEntityHandle Entity, FMassArchetypeHandle NewArchetypeHandle)
+void FMassEntityManager::MoveEntityToAnotherArchetype(FMassEntityHandle Entity, FMassArchetypeHandle NewArchetypeHandle,
+	const FMassArchetypeSharedFragmentValues* SharedFragmentValuesOverride)
 {
 	CHECK_SYNC_API();
 
@@ -1827,7 +1828,7 @@ void FMassEntityManager::MoveEntityToAnotherArchetype(FMassEntityHandle Entity,
 	// Move the entity over
 	FMassArchetypeData* CurrentArchetype = GetEntityStorageInterface().GetArchetype(Entity.Index);
 	check(CurrentArchetype);
-	CurrentArchetype->MoveEntityToAnotherArchetype(Entity, NewArchetype);
+	CurrentArchetype->MoveEntityToAnotherArchetype(Entity, NewArchetype, SharedFragmentValuesOverride);
 	GetEntityStorageInterface().SetArchetypeFromShared(Entity.Index, NewArchetypeHandle.DataPtr);
 }
 
diff --git a/Engine/Source/Runtime/MassEntity/Public/MassEntityManager.h b/Engine/Source/Runtime/MassEntity/Public/MassEntityManager.h
--- a/Engine/Source/Runtime/MassEntity/Public/MassEntityManager.h
+++ b/Engine/Source/Runtime/MassEntity/Public/MassEntityManager.h
@@ -474,8 +474,9 @@ public:
 	 * Moves an entity over to a new archetype by copying over fragments common to both archetypes
 	 * @param Entity is the entity to move 
 	 * @param NewArchetypeHandle the handle to the new archetype
+	 * @param SharedFragmentValuesOverride if provided will override all given Entity's shared fragment values
 	 */
-	void MoveEntityToAnotherArchetype(FMassEntityHandle Entity, FMassArchetypeHandle NewArchetypeHandle);
+	void MoveEntityToAnotherArchetype(FMassEntityHandle Entity, FMassArchetypeHandle NewArchetypeHandle, const FMassArchetypeSharedFragmentValues* SharedFragmentValuesOverride = nullptr);
 
 	/** Copies values from FragmentInstanceList over to Entity's fragment. Caller is responsible for ensuring that 
 	 *  the given entity does have given fragments. Failing this assumption will cause a check-fail.*/
