[Mass] Fixed inconsistencies in MassEntity's shared and const shared fragments-related API. The change also significantly reduces code duplication in the related code.
The change fixes a flaw in existing API allowing calling ForEachSharedFragment on FMassConstSharedFragments which in turns let them modify the instances, which is conceptually wrong.

In additiona added missing tests for ForEach[Const]SharedFragment and ForEach[Const]SharedFragmentConditional

diff --git a/Engine/Source/Developer/MassEntityTestSuite/Private/MassSharedFragmentsTest.cpp b/Engine/Source/Developer/MassEntityTestSuite/Private/MassSharedFragmentsTest.cpp
--- a/Engine/Source/Developer/MassEntityTestSuite/Private/MassSharedFragmentsTest.cpp
+++ b/Engine/Source/Developer/MassEntityTestSuite/Private/MassSharedFragmentsTest.cpp
@@ -264,6 +264,112 @@ struct FSharedFragmentValues_Hash : FExecutionTestBase
 };
 IMPLEMENT_AI_INSTANT_TEST(FSharedFragmentValues_Hash, "System.Mass.SharedFragments.Hash");
 
+struct FSharedFragment_ForEach : FExecutionTestBase
+{
+	virtual bool InstantTest() override
+	{
+		constexpr int32 NumSharedFragments = 4;
+		TStaticArray<int32, NumSharedFragments> TestInitValues = { 9, 1, 12, 13};
+		
+		for (int32 InitValue : TestInitValues)
+		{
+			EntityManager->GetOrCreateSharedFragment<FTestSharedFragment_Int>(InitValue);
+		}
+
+		TArray<int32> Results;
+		TArray<int32> ModifiedValues;
+		EntityManager->ForEachSharedFragment<FTestSharedFragment_Int>([&Results, &ModifiedValues](FTestSharedFragment_Int& SharedFragment)
+		{
+			Results.Add(SharedFragment.Value);
+			ModifiedValues.Add(SharedFragment.Value += 100);
+		});
+
+		AITEST_EQUAL("Number of processed shared fragments", Results.Num(), NumSharedFragments);
+		for (int32 InitValue : TestInitValues)
+		{
+			AITEST_TRUE("Read values matches init values", Results.Find(InitValue) != INDEX_NONE);
+		}
+
+		TArray<int32> MutatedResults;
+		EntityManager->ForEachSharedFragment<FTestSharedFragment_Int>([&MutatedResults](const FTestSharedFragment_Int& SharedFragment)
+		{
+			MutatedResults.Add(SharedFragment.Value);
+		});
+
+		AITEST_EQUAL("Number of shared fragments processed in second round", Results.Num(), NumSharedFragments);
+		for (int32 ModifiedValue : ModifiedValues)
+		{
+			AITEST_TRUE("Read values matches values set in the first round", MutatedResults.Find(ModifiedValue) != INDEX_NONE);
+		}
+
+		constexpr int32 ConditionalLimit = 10;
+		TArray<int32> ConditionalResults;
+		EntityManager->ForEachSharedFragmentConditional<FTestSharedFragment_Int>(
+			[](FTestSharedFragment_Int& SharedFragment)
+			{
+				return SharedFragment.Value > ConditionalLimit;
+			}
+			, [&ConditionalResults](FTestSharedFragment_Int& SharedFragment)
+			{
+				ConditionalResults.Add(SharedFragment.Value);
+			}
+		);
+		for (int32 Value : ConditionalResults)
+		{
+			AITEST_TRUE("Only the values matching the condition get processed", Value > ConditionalLimit);
+		}
+		
+		return true;
+	}
+};
+IMPLEMENT_AI_INSTANT_TEST(FSharedFragment_ForEach, "System.Mass.SharedFragments.ForEach");
+
+struct FConstSharedFragment_ForEach : FExecutionTestBase
+{
+	virtual bool InstantTest() override
+	{
+		constexpr int32 NumSharedFragments = 4;
+		TStaticArray<int32, NumSharedFragments> TestInitValues = { 9, 1, 12, 13 };
+		
+		for (int32 InitValue : TestInitValues)
+		{
+			EntityManager->GetOrCreateConstSharedFragment<FTestConstSharedFragment_Int>(InitValue);
+		}
+
+		TArray<int32> Results;
+		EntityManager->ForEachConstSharedFragment<FTestConstSharedFragment_Int>([&Results](const FTestConstSharedFragment_Int& SharedFragment)
+		{
+			Results.Add(SharedFragment.Value);
+		});
+
+		AITEST_EQUAL("Number of processed shared fragments", Results.Num(), NumSharedFragments);
+		for (int32 InitValue : TestInitValues)
+		{
+			AITEST_TRUE("Read values matches init values", Results.Find(InitValue) != INDEX_NONE);
+		}
+
+		constexpr int32 ConditionalLimit = 10;
+		TArray<int32> ConditionalResults;
+		EntityManager->ForEachConstSharedFragmentConditional<FTestConstSharedFragment_Int>(
+			[](const FTestConstSharedFragment_Int& SharedFragment)
+			{
+				return SharedFragment.Value > ConditionalLimit;
+			}
+			, [&ConditionalResults](const FTestConstSharedFragment_Int& SharedFragment)
+			{
+				ConditionalResults.Add(SharedFragment.Value);
+			}
+		);
+		for (int32 Value : ConditionalResults)
+		{
+			AITEST_TRUE("Only the values matching the condition get processed", Value > ConditionalLimit);
+		}
+
+		return true;
+	}
+};
+IMPLEMENT_AI_INSTANT_TEST(FConstSharedFragment_ForEach, "System.Mass.ConstSharedFragments.ForEach");
+
 //-----------------------------------------------------------------------------
 // Entity-related Tests
 //-----------------------------------------------------------------------------
diff --git a/Engine/Source/Runtime/MassEntity/Private/MassEntityManager.cpp b/Engine/Source/Runtime/MassEntity/Private/MassEntityManager.cpp
--- a/Engine/Source/Runtime/MassEntity/Private/MassEntityManager.cpp
+++ b/Engine/Source/Runtime/MassEntity/Private/MassEntityManager.cpp
@@ -192,12 +192,12 @@ void FMassEntityManager::AddReferencedObjects(FReferenceCollector& Collector)
 		return;
 	}
 
-	for (FConstSharedStruct& Struct : ConstSharedFragments)
+	for (FConstSharedStruct& Struct : ConstSharedFragmentsContainer.GetAllInstances())
 	{
 		Struct.AddStructReferencedObjects(Collector);
 	}
 
-	for (FSharedStruct& Struct : SharedFragments)
+	for (FSharedStruct& Struct : SharedFragmentsContainer.GetAllInstances())
 	{
 		Struct.AddStructReferencedObjects(Collector);
 	}
diff --git a/Engine/Source/Runtime/MassEntity/Public/MassEntityManager.h b/Engine/Source/Runtime/MassEntity/Public/MassEntityManager.h
--- a/Engine/Source/Runtime/MassEntity/Public/MassEntityManager.h
+++ b/Engine/Source/Runtime/MassEntity/Public/MassEntityManager.h
@@ -562,7 +562,7 @@ public:
 		static_assert(TIsDerivedFrom<SharedFragmentType, FMassSharedFragment>::IsDerived
 			, "Given struct doesn't represent a valid shared fragment type. Make sure to inherit from FMassSharedFragment or one of its child-types.");
 
-		TArray<FSharedStruct>* InstancesOfType = SharedFragmentsTypeMap.Find(SharedFragmentType::StaticStruct());
+		TArray<FSharedStruct>* InstancesOfType = SharedFragmentsContainer.Find(SharedFragmentType::StaticStruct());
 		return InstancesOfType ? *InstancesOfType : TConstArrayView<FSharedStruct>();
 	}
 
@@ -631,12 +631,7 @@ public:
 	const FConstSharedStruct& GetOrCreateConstSharedFragmentByHash(const uint32 Hash, const T& Fragment)
 	{
 		static_assert(TIsDerivedFrom<T, FMassConstSharedFragment>::IsDerived, "Given struct doesn't represent a valid const shared fragment type. Make sure to inherit from FMassConstSharedFragment or one of its child-types.");
-		int32& Index = ConstSharedFragmentsMap.FindOrAddByHash(Hash, Hash, INDEX_NONE);
-		if (Index == INDEX_NONE)
-		{
-			Index = ConstSharedFragments.Add(FConstSharedStruct::Make(Fragment));
-		}
-		return ConstSharedFragments[Index];
+		return ConstSharedFragmentsContainer.FindOrAdd(Hash, T::StaticStruct(), reinterpret_cast<const uint8*>(&Fragment));
 	}
 
 private:
@@ -644,54 +639,24 @@ private:
 	const FSharedStruct& GetOrCreateSharedFragmentByHash(const uint32 Hash, const T& Fragment)
 	{
 		static_assert(TIsDerivedFrom<T, FMassSharedFragment>::IsDerived, "Given struct doesn't represent a valid shared fragment type. Make sure to inherit from FMassSharedFragment or one of its child-types.");
-		int32& Index = SharedFragmentsMap.FindOrAddByHash(Hash, Hash, INDEX_NONE);
-		if (Index == INDEX_NONE)
-		{
-			Index = SharedFragments.Add(FSharedStruct::Make(Fragment));
-			// note that even though we're copying the freshly created FSharedStruct instance it's perfectly fine since 
-			// FSharedStruct do guarantee there's not going to be data duplication (via a member shared pointer to hosted data)
-			TArray<FSharedStruct>& InstancesOfType = SharedFragmentsTypeMap.FindOrAdd(T::StaticStruct(), {});
-			InstancesOfType.Add(SharedFragments[Index]);
-		}
-
-		return SharedFragments[Index];
+		return SharedFragmentsContainer.FindOrAdd(Hash, T::StaticStruct(), reinterpret_cast<const uint8*>(&Fragment));
 	}
 
 	const FConstSharedStruct& GetOrCreateConstSharedFragmentByHash(const uint32 Hash, const UScriptStruct* InScriptStruct, const uint8* InStructMemory)
 	{
-		int32& Index = ConstSharedFragmentsMap.FindOrAddByHash(Hash, Hash, INDEX_NONE);
-		if (Index == INDEX_NONE)
-		{
-			Index = ConstSharedFragments.Add(FConstSharedStruct::Make(InScriptStruct, InStructMemory));
-		}
-		return ConstSharedFragments[Index];
+		return ConstSharedFragmentsContainer.FindOrAdd(Hash, InScriptStruct, InStructMemory);
 	}
 
 	const FSharedStruct& GetOrCreateSharedFragmentByHash(const uint32 Hash, const UScriptStruct* InScriptStruct, const uint8* InStructMemory)
 	{
-		int32& Index = SharedFragmentsMap.FindOrAddByHash(Hash, Hash, INDEX_NONE);
-		if (Index == INDEX_NONE)
-		{
-			Index = SharedFragments.Add(FSharedStruct::Make(InScriptStruct, InStructMemory));
-			// note that even though we're copying the freshly created FSharedStruct instance it's perfectly fine since 
-			// FSharedStruct do guarantee there's not going to be data duplication (via a member shared pointer to hosted data)
-			TArray<FSharedStruct>& InstancesOfType = SharedFragmentsTypeMap.FindOrAdd(InScriptStruct, {});
-			InstancesOfType.Add(SharedFragments[Index]);
-		}
-		return SharedFragments[Index];
+		return SharedFragmentsContainer.FindOrAdd(Hash, InScriptStruct, InStructMemory);
 	}
 
 	template<typename T, typename... TArgs>
 	const FConstSharedStruct& GetOrCreateConstSharedFragmentByHash(const uint32 Hash, TArgs&&... InArgs)
 	{
 		static_assert(TIsDerivedFrom<T, FMassConstSharedFragment>::IsDerived, "Given struct doesn't represent a valid const shared fragment type. Make sure to inherit from FMassConstSharedFragment or one of its child-types.");
-		int32& Index = ConstSharedFragmentsMap.FindOrAddByHash(Hash, Hash, INDEX_NONE);
-		if (Index == INDEX_NONE)
-		{
-			Index = ConstSharedFragments.Add(FConstSharedStruct::Make<T>(Forward<TArgs>(InArgs)...));
-		}
-
-		return ConstSharedFragments[Index];
+		return ConstSharedFragmentsContainer.FindOrAdd<T>(Hash, Forward<TArgs>(InArgs)...);
 	}
 
 public:
@@ -717,17 +682,7 @@ public:
 	const FSharedStruct& GetOrCreateSharedFragmentByHash(const uint32 Hash, TArgs&&... InArgs)
 	{
 		static_assert(TIsDerivedFrom<T, FMassSharedFragment>::IsDerived, "Given struct doesn't represent a valid shared fragment type. Make sure to inherit from FMassSharedFragment or one of its child-types.");
-		int32& Index = SharedFragmentsMap.FindOrAddByHash(Hash, Hash, INDEX_NONE);
-		if (Index == INDEX_NONE)
-		{
-			Index = SharedFragments.Add(FSharedStruct::Make<T>(Forward<TArgs>(InArgs)...));
-			// note that even though we're copying the freshly created FSharedStruct instance it's perfectly fine since 
-			// FSharedStruct do guarantee there's not going to be data duplication (via a member shared pointer to hosted data)
-			TArray<FSharedStruct>& InstancesOfType = SharedFragmentsTypeMap.FindOrAdd(T::StaticStruct(), {});
-			InstancesOfType.Add(SharedFragments[Index]);
-		}
-
-		return SharedFragments[Index];
+		return SharedFragmentsContainer.FindOrAdd<T>(Hash, Forward<TArgs>(InArgs)...);
 	}
 
 	/**
@@ -992,33 +947,17 @@ PRAGMA_ENABLE_DEPRECATION_WARNINGS
 		return GetOrCreateSharedFragmentByHash(Hash, MoveTemp(Fragment));
 	}
 
-	template<typename T>
-	void ForEachSharedFragment(TFunctionRef< void(T& /*SharedFragment*/) > ExecuteFunction)
-	{
-		if (TArray<FSharedStruct>* InstancesOfType = SharedFragmentsTypeMap.Find(T::StaticStruct()))
-		{
-			for (const FSharedStruct& SharedStruct : *InstancesOfType)
-			{
-				ExecuteFunction(SharedStruct.Get<T>());
-			}
-		}
-	}
+	template<UE::Mass::CSharedFragment T>
+	void ForEachSharedFragment(TFunctionRef< void(T& /*SharedFragment*/) > ExecuteFunction);
 
-	template<typename T>
-	void ForEachSharedFragmentConditional(TFunctionRef< bool(T& /*SharedFragment*/) > ConditionFunction, TFunctionRef< void(T& /*SharedFragment*/) > ExecuteFunction)
-	{
-		if (TArray<FSharedStruct>* InstancesOfType = SharedFragmentsTypeMap.Find(T::StaticStruct()))
-		{
-			for (const FSharedStruct& SharedStruct : *InstancesOfType)
-			{
-				T& StructInstanceRef = SharedStruct.Get<T>();
-				if (ConditionFunction(StructInstanceRef))
-				{
-					ExecuteFunction(StructInstanceRef);
-				}
-			}
-		}
-	}
+	template<UE::Mass::CSharedFragment T>
+	void ForEachSharedFragmentConditional(TFunctionRef< bool(T& /*SharedFragment*/) > ConditionFunction, TFunctionRef< void(T& /*SharedFragment*/) > ExecuteFunction);
+
+	template<UE::Mass::CConstSharedFragment T>
+	void ForEachConstSharedFragment(TFunctionRef< void(const T& /*ConstSharedFragment*/) > ExecuteFunction);
+
+	template<UE::Mass::CConstSharedFragment T>
+	void ForEachConstSharedFragmentConditional(TFunctionRef< bool(const T& /*ConstSharedFragment*/) > ConditionFunction, TFunctionRef< void(const T& /*ConstSharedFragment*/) > ExecuteFunction);
 
 	[[nodiscard]] UE::Mass::FEntityBuilder MakeEntityBuilder();
 
@@ -1213,16 +1152,69 @@ private:
 	// throughout its lifetime, and the index is never reused for another archetype. 
 	TArray<TSharedPtr<FMassArchetypeData>> AllArchetypes;
 
-	// Shared fragments
-	TArray<FConstSharedStruct> ConstSharedFragments;
-	// Hash/Index in array pair
-	TMap<uint32, int32> ConstSharedFragmentsMap;
+	/**
+	 * This is a struct wrapping shared fragment management to ensure consistency between how
+	 * shared and const shared fragment are added and fetched, across all the functions that do that
+	 */
+	template<typename TSharedStructType>
+	struct TSharedFragmentsContainer
+	{
+		TArray<TSharedStructType>* Find(const UScriptStruct* Type)
+		{
+			return TypeToInstanceMap.Find(Type);
+		}
+
+		TSharedStructType& FindOrAdd(const uint32 Hash, const UScriptStruct* Type, const uint8* Data)
+		{
+			int32& Index = HashToInstanceIndexMap.FindOrAddByHash(Hash, Hash, INDEX_NONE);
+			if (Index == INDEX_NONE)
+			{
+				Index = Add(TSharedStructType::Make(Type, Data));
+			}
+			return Instances[Index];
+		}
 
-	TArray<FSharedStruct> SharedFragments;
-	// Hash/Index in array pair, indices point at SharedFragments
-	TMap<uint32, int32> SharedFragmentsMap;
-	// Maps specific struct type to a collection of FSharedStruct instances of that type
-	TMap<const UScriptStruct*, TArray<FSharedStruct>> SharedFragmentsTypeMap;
+		template<typename T, typename... TArgs>
+		TSharedStructType& FindOrAdd(const uint32 Hash, TArgs&&... InArgs)
+		{
+			int32& Index = HashToInstanceIndexMap.FindOrAddByHash(Hash, Hash, INDEX_NONE);
+			if (Index == INDEX_NONE)
+			{
+				Index = Add(TSharedStructType::template Make<T>(Forward<TArgs>(InArgs)...));
+			}
+			return Instances[Index];
+		}
+
+		TSharedStructType& operator[](const int32 Index)
+		{
+			return Instances[Index];
+		}
+
+		TArrayView<TSharedStructType> GetAllInstances()
+		{
+			return Instances;
+		}
+
+	private:
+		int32 Add(TSharedStructType&& SharedStruct)
+		{
+			TArray<TSharedStructType>& InstancesOfType = TypeToInstanceMap.FindOrAdd(SharedStruct.GetScriptStruct(), {});
+			const int32 Index = Instances.Add(Forward<TSharedStructType>(SharedStruct));
+			// note that even though we're copying the input F[Const]SharedStruct instance it's perfectly fine since 
+			// F[Const]SharedStruct does guarantee there's not going to be data duplication (via a member shared pointer to hosted data)
+			InstancesOfType.Add(Instances[Index]);
+			return Index;
+		}
+
+		TArray<TSharedStructType> Instances;
+		// Hash/Index in array pair
+		TMap<uint32, int32> HashToInstanceIndexMap;
+		// Maps specific struct type to a collection of FSharedStruct instances of that type
+		TMap<const UScriptStruct*, TArray<TSharedStructType>> TypeToInstanceMap;
+	};
+
+	TSharedFragmentsContainer<FConstSharedStruct> ConstSharedFragmentsContainer;
+	TSharedFragmentsContainer<FSharedStruct> SharedFragmentsContainer;
 
 	FMassObserverManager ObserverManager;
 
@@ -1247,9 +1239,10 @@ private:
 	//-----------------------------------------------------------------------------
 	// DEPRECATED
 	//-----------------------------------------------------------------------------
+protected:
 	UE_DEPRECATED(5.6, "This flavor of InternalCreateSimilarArchetype is deprecated due to the introduction of archetype grouping. Use InternalCreateSimilarArchetype with a FArchetypeGroups parameter instead")
 	FMassArchetypeHandle InternalCreateSimilarArchetype(const FMassArchetypeData& SourceArchetypeRef, FMassArchetypeCompositionDescriptor&& NewComposition);
-
+public:
 	UE_DEPRECATED(5.6, "SetEntityFragmentsValues is deprecated. Use SetEntityFragmentValues instead (note the slight change in name).")
 	void SetEntityFragmentsValues(FMassEntityHandle Entity, TArrayView<const FInstancedStruct> FragmentInstanceList);
 
@@ -1258,4 +1251,75 @@ private:
 
 	UE_DEPRECATED(5.6, "Static BatchSetEntityFragmentsValues is deprecated. Use EntityManager's member function BatchSetEntityFragmentValues (note the slight change in name).")
 	static void BatchSetEntityFragmentsValues(TConstArrayView<FMassArchetypeEntityCollection> EntityCollections, TArrayView<const FInstancedStruct> FragmentInstanceList);
+
+	template<UE::Mass::CConstSharedFragment T>
+	UE_DEPRECATED(5.6, "Using ForEachSharedFragment for Const Shared Fragments has been deprecated. Use ForEachConstSharedFragment instead.")
+	void ForEachSharedFragment(TFunctionRef< void(T& /*SharedFragment*/) > ExecuteFunction)
+	{
+	}
+
+	template<UE::Mass::CConstSharedFragment T>
+	UE_DEPRECATED(5.6, "Using ForEachSharedFragmentConditional for Const Shared Fragments has been deprecated. Use ForEachConstSharedFragmentConditional instead.")
+	void ForEachSharedFragmentConditional(TFunctionRef< bool(T& /*SharedFragment*/) > ConditionFunction, TFunctionRef< void(T& /*SharedFragment*/) > ExecuteFunction)
+	{
+	}
 };
+
+//-----------------------------------------------------------------------------
+// INLINE
+//-----------------------------------------------------------------------------
+template<UE::Mass::CSharedFragment T>
+void FMassEntityManager::ForEachSharedFragment(TFunctionRef< void(T& /*SharedFragment*/) > ExecuteFunction)
+{
+	if (TArray<FSharedStruct>* InstancesOfType = SharedFragmentsContainer.Find(T::StaticStruct()))
+	{
+		for (const FSharedStruct& SharedStruct : *InstancesOfType)
+		{
+			ExecuteFunction(SharedStruct.Get<T>());
+		}
+	}
+}
+
+template<UE::Mass::CSharedFragment T>
+void FMassEntityManager::ForEachSharedFragmentConditional(TFunctionRef< bool(T& /*SharedFragment*/) > ConditionFunction, TFunctionRef< void(T& /*SharedFragment*/) > ExecuteFunction)
+{
+	if (TArray<FSharedStruct>* InstancesOfType = SharedFragmentsContainer.Find(T::StaticStruct()))
+	{
+		for (const FSharedStruct& SharedStruct : *InstancesOfType)
+		{
+			T& StructInstanceRef = SharedStruct.Get<T>();
+			if (ConditionFunction(StructInstanceRef))
+			{
+				ExecuteFunction(StructInstanceRef);
+			}
+		}
+	}
+}
+
+template<UE::Mass::CConstSharedFragment T>
+void FMassEntityManager::ForEachConstSharedFragment(TFunctionRef< void(const T& /*ConstSharedFragment*/) > ExecuteFunction)
+{
+	if (TArray<FConstSharedStruct>* InstancesOfType = ConstSharedFragmentsContainer.Find(T::StaticStruct()))
+	{
+		for (const FConstSharedStruct& SharedStruct : *InstancesOfType)
+		{
+			ExecuteFunction(SharedStruct.Get<const T>());
+		}
+	}
+}
+
+template<UE::Mass::CConstSharedFragment T>
+void FMassEntityManager::ForEachConstSharedFragmentConditional(TFunctionRef< bool(const T& /*ConstSharedFragment*/) > ConditionFunction, TFunctionRef< void(const T& /*ConstSharedFragment*/) > ExecuteFunction)
+{
+	if (TArray<FConstSharedStruct>* InstancesOfType = ConstSharedFragmentsContainer.Find(T::StaticStruct()))
+	{
+		for (const FConstSharedStruct& SharedStruct : *InstancesOfType)
+		{
+			const T& StructInstanceRef = SharedStruct.Get<const T>();
+			if (ConditionFunction(StructInstanceRef))
+			{
+				ExecuteFunction(StructInstanceRef);
+			}
+		}
+	}
+}
