Fix compile error by adding missing include

diff --git a/Engine/Source/Runtime/MassEntity/Private/MassProcessor.cpp b/Engine/Source/Runtime/MassEntity/Private/MassProcessor.cpp
--- a/Engine/Source/Runtime/MassEntity/Private/MassProcessor.cpp
+++ b/Engine/Source/Runtime/MassEntity/Private/MassProcessor.cpp
@@ -7,6 +7,7 @@
 #include "Engine/World.h"
 #include "MassCommandBuffer.h"
 #include "MassExecutionContext.h"
+#include "ProfilingDebugging/CsvProfiler.h"
 
 #include UE_INLINE_GENERATED_CPP_BY_NAME(MassProcessor)
 
