[Mass] minor MassProcessor touchups

Fixed deprecation warning on UMassProcessor::Initialize
Made UMassProcessor::ConfigureQueries no longer pure-virtual, but made it complain when the default implementation gets called while there are registered entity queries.

diff --git a/Engine/Source/Runtime/MassEntity/Private/MassProcessor.cpp b/Engine/Source/Runtime/MassEntity/Private/MassProcessor.cpp
--- a/Engine/Source/Runtime/MassEntity/Private/MassProcessor.cpp
+++ b/Engine/Source/Runtime/MassEntity/Private/MassProcessor.cpp
@@ -22,8 +22,6 @@ namespace UE::Mass::Debug
 		, TEXT("When enabled will log task graph tasks created while dispatching processors to other threads, along with their dependencies"), ECVF_Cheat);
 }
 
-#define PROCESSOR_LOG(Verbosity, Fmt, ...) UE_VLOG_UELOG(this, LogMass, Verbosity, Fmt, ##__VA_ARGS__)
-
 // change to 1 to enable more detailed processing tasks logging
 #if 0
 #define PROCESSOR_TASK_LOG(Fmt, ...) UE_VLOG_UELOG(this, LogMass, Verbose, Fmt, ##__VA_ARGS__)
@@ -32,7 +30,6 @@ namespace UE::Mass::Debug
 #endif // 0
 
 #else 
-#define PROCESSOR_LOG(...)
 #define PROCESSOR_TASK_LOG(...) 
 #endif // WITH_MASSENTITY_DEBUG
 
@@ -165,6 +162,13 @@ void UMassProcessor::InitializeInternal(UObject&, const TSharedRef<FMassEntityMa
 	// empty in base class
 }
 
+void UMassProcessor::ConfigureQueries(const TSharedRef<FMassEntityManager>&)
+{
+	UE_CVLOG_UELOG(OwnedQueries.Num(), this, LogMass, Warning
+		, TEXT("%s has entity queries registered. Make sure to override ConfigureQueries to configure the queries, and do not call the Super implementation")
+		, *GetProcessorName());
+}
+
 void UMassProcessor::SetShouldAutoRegisterWithGlobalList(const bool bAutoRegister)
 {	
 	if (ensureMsgf(HasAnyFlags(RF_ClassDefaultObject), TEXT("Setting bAutoRegisterWithProcessingPhases for non-CDOs has no effect")))
@@ -243,7 +247,7 @@ void UMassProcessor::CallExecute(FMassEntityManager& EntityManager, FMassExecuti
 		}
 		else
 		{
-			PROCESSOR_LOG(VeryVerbose, TEXT("%s Skipping Execute due to subsystem requirements not being met"), *GetProcessorName());
+			UE_VLOG_UELOG(this, LogMass, VeryVerbose, TEXT("%s Skipping Execute due to subsystem requirements not being met"), *GetProcessorName());
 		}
 
 		if (ActivationState == EActivationState::OneShot)
diff --git a/Engine/Source/Runtime/MassEntity/Public/MassProcessor.h b/Engine/Source/Runtime/MassEntity/Public/MassProcessor.h
--- a/Engine/Source/Runtime/MassEntity/Public/MassProcessor.h
+++ b/Engine/Source/Runtime/MassEntity/Public/MassProcessor.h
@@ -149,10 +149,16 @@ public:
 #endif
 	
 protected:
-	/** Called to initialize the processor. Override to perform custom steps. */
+	/** Called to initialize the processor's internal state. Override to perform custom steps. */
 	virtual void InitializeInternal(UObject& Owner, const TSharedRef<FMassEntityManager>& EntityManager);
 
-	virtual void ConfigureQueries(const TSharedRef<FMassEntityManager>& EntityManager) PURE_VIRTUAL(UMassProcessor::ConfigureQueries);
+	/**
+	 * Called internally during processor's initialization so that child classes configure their owned queries
+	 * with requirements. This function is called before the processor gets considered by Mass dependency
+	 * solver, and the requirement information stored in queries is crucial for that process. 
+	 */
+	virtual void ConfigureQueries(const TSharedRef<FMassEntityManager>& EntityManager);
+
 	virtual void PostInitProperties() override;
 	virtual void Execute(FMassEntityManager& EntityManager, FMassExecutionContext& Context) PURE_VIRTUAL(UMassProcessor::Execute);
 
@@ -234,7 +240,7 @@ private:
 	// DEPRECATED
 	//-----------------------------------------------------------------------------
 public:
-	UE_DEPRECATED(5.6, "Initialize is deprecated. Override Initialize(UObject&, const TSharedRef<FMassEntityManager>&) instead. If you want to call the function, use CallInitialize.")
+	UE_DEPRECATED(5.6, "Initialize is deprecated. Override InitializeInternal(UObject&, const TSharedRef<FMassEntityManager>&) instead. If you want to call the function, use CallInitialize.")
 	virtual void Initialize(UObject& Owner) final;
 	UE_DEPRECATED(5.6, "This flavor of ConfigureQueries is deprecated. Override ConfigureQueries(const TSharedRef<FMassEntityManager>&) instead.")
 	virtual void ConfigureQueries() final {};
