[Mass] EntityBuilder can now be used when IsProcessing == true - the apropriate command will be issued instead of synchronized creation

diff --git a/Engine/Source/Developer/MassEntityTestSuite/Private/MassEntityBuilderTest.cpp b/Engine/Source/Developer/MassEntityTestSuite/Private/MassEntityBuilderTest.cpp
--- a/Engine/Source/Developer/MassEntityTestSuite/Private/MassEntityBuilderTest.cpp
+++ b/Engine/Source/Developer/MassEntityTestSuite/Private/MassEntityBuilderTest.cpp
@@ -4,6 +4,8 @@
 #include "MassEntityTestTypes.h"
 #include "MassEntityBuilder.h"
 #include "MassEntityManager.h"
+#include "MassExecutionContext.h"
+#include "MassExecutor.h"
 
 
 #define LOCTEXT_NAMESPACE "MassTest"
@@ -348,6 +350,145 @@ namespace UE::Mass::Test::EntityBuilder
 		}
 	};
 	IMPLEMENT_AI_INSTANT_TEST(FPassOver, "System.Mass.EntityBuilder.PassOver");
+
+	struct FBuilderReuse : FEntityTestBase
+	{
+		virtual bool InstantTest() override
+		{
+			FEntityBuilder Builder = EntityManager->MakeEntityBuilder();
+			Builder.Add<FTestFragment_Int>();
+
+			auto TestBuilder = [this, &Builder](const FMassArchetypeHandle& ExpectedArchetype)-> bool
+			{
+				TArray<const FMassEntityHandle> Entities = {
+					Builder.CommitAndReprepare()
+					, Builder.CommitAndReprepare()
+				};
+
+				const FMassArchetypeHandle LocalArchetype = Builder.GetArchetypeHandle();
+				AITEST_NOT_EQUAL("Two entities created sequentially", Entities[0], Entities[1]);
+				AITEST_EQUAL("Entities' archetype", EntityManager->GetArchetypeForEntity(Entities[0]), EntityManager->GetArchetypeForEntity(Entities[1]));
+				AITEST_EQUAL("Builders archetype and entities' archetype", Builder.GetArchetypeHandle(), EntityManager->GetArchetypeForEntity(Entities[0]))
+				AITEST_TRUE("Archetype matches expectations", EntityManager->GetArchetypeForEntity(Entities[0]) == ExpectedArchetype);
+				return true;
+			};
+
+			if (TestBuilder(IntsArchetype))
+			{
+				Builder.Add<FTestFragment_Float>();
+				return TestBuilder(FloatsIntsArchetype);
+			}
+
+			return false;
+		}
+	};
+	IMPLEMENT_AI_INSTANT_TEST(FBuilderReuse, "System.Mass.EntityBuilder.Reuse");
+
+	struct FDuringProcessing : FEntityTestBase
+	{
+		virtual bool InstantTest() override
+		{
+			constexpr int32 NumIterations = 5;
+			// creating a single entity to enforce the execution function of the processor we're going to use to execute
+			// exactly once
+			EntityManager->CreateEntity(IntsArchetype);
+
+			TArray<FMassEntityHandle> EntityHandles;
+
+			constexpr int32 InitialValueToSet = 100;
+			int32 ValueToSet = InitialValueToSet;
+			UMassTestProcessorBase* Processor = NewTestProcessor<UMassTestProcessorBase>(EntityManager);
+			Processor->ForEachEntityChunkExecutionFunction = [&ValueToSet, &EntityHandles](FMassExecutionContext& Context)
+			{
+				FEntityBuilder AsyncBuilder = Context.GetEntityManagerChecked().MakeEntityBuilder();
+				AsyncBuilder.Add<FTestFragment_Int>(ValueToSet++);
+				EntityHandles.Add(AsyncBuilder.Commit());
+			};
+			Processor->EntityQuery.AddRequirement<FTestFragment_Int>(EMassFragmentAccess::ReadOnly);
+
+			FMassProcessingContext ProcessingContext(EntityManager, /*DeltaSeconds=*/0.f);
+			ProcessingContext.bFlushCommandBuffer = false;
+
+			for (int32 Iteration = 0; Iteration < NumIterations; ++Iteration)
+			{
+				Executor::RunProcessorsView(MakeArrayView(reinterpret_cast<UMassProcessor**>(&Processor), 1), ProcessingContext);
+				AITEST_EQUAL(FString::Printf(TEXT("Number of entities after iteration %d"), Iteration), EntityHandles.Num(), Iteration + 1);
+			}
+
+			for (int32 Iteration = 0; Iteration < NumIterations; ++Iteration)
+			{
+				AITEST_FALSE(FString::Printf(TEXT("(NOT) Entity %d is `created`"), Iteration), EntityManager->IsEntityBuilt(EntityHandles[Iteration]));
+			}
+
+			EntityManager->FlushCommands();
+
+			for (int32 Iteration = 0; Iteration < NumIterations; ++Iteration)
+			{
+				AITEST_TRUE(FString::Printf(TEXT("Entity %d is `created`"), Iteration), EntityManager->IsEntityBuilt(EntityHandles[Iteration]));
+				AITEST_TRUE(FString::Printf(TEXT("Entity %d has the right archetype"), Iteration), EntityManager->GetArchetypeForEntity(EntityHandles[Iteration]) == IntsArchetype);
+				if (Iteration + 1 < NumIterations)
+				{
+					AITEST_FALSE(FString::Printf(TEXT("(NOT) Entity handles are the same %d"), Iteration), EntityHandles[Iteration] == EntityHandles[Iteration + 1]);
+				}
+			}
+
+			return true;
+		}
+	};
+	IMPLEMENT_AI_INSTANT_TEST(FDuringProcessing, "System.Mass.EntityBuilder.DuringProcessing");
+
+	struct FSyncBuildingAsyncSubmission : FEntityTestBase
+	{
+		virtual bool InstantTest() override
+		{
+			// creating a single entity to enforce the execution function of the processor we're going to use to execute
+			// exactly once
+			EntityManager->CreateEntity(IntsArchetype);
+
+			FEntityBuilder SyncBuilder = EntityManager->MakeEntityBuilder();
+			SyncBuilder.Add<FTestFragment_Int>();
+
+			const FMassEntityHandle ReservedHandle = SyncBuilder.GetEntityHandle();
+
+			int32 ProcessedEntitiesCount = 0;
+			UMassTestProcessorBase* Processor = NewTestProcessor<UMassTestProcessorBase>(EntityManager);
+			Processor->ForEachEntityChunkExecutionFunction = [&SyncBuilder, &ProcessedEntitiesCount](FMassExecutionContext& Context)
+			{
+				SyncBuilder.Commit();
+				ProcessedEntitiesCount += Context.GetNumEntities();
+			};
+			Processor->EntityQuery.AddRequirement<FTestFragment_Int>(EMassFragmentAccess::ReadOnly);
+
+			FMassProcessingContext ProcessingContext(EntityManager, /*DeltaSeconds=*/0.f);
+			ProcessingContext.bFlushCommandBuffer = false;
+			Executor::RunProcessorsView(MakeArrayView(reinterpret_cast<UMassProcessor**>(&Processor), 1), ProcessingContext);
+
+			AITEST_EQUAL("Number of fully-formed entities expected", ProcessedEntitiesCount, 1);
+			AITEST_EQUAL("The entity handle before and after async commit", ReservedHandle, SyncBuilder.GetEntityHandle());
+			AITEST_TRUE("The Builder is in `Committed` state", SyncBuilder.IsCommitted());
+			// since the commands are not flushed yet, due to ProcessingContext's values, we expect the entity to not be created yet
+			AITEST_FALSE("(NOT) the entity has been created", EntityManager->IsEntityBuilt(ReservedHandle));
+
+			{
+				AITEST_SCOPED_CHECK("Trying to commit an already committed", 1);
+				AITEST_INFO("Second execution of the processor shouldn't change a thing.");
+				Executor::RunProcessorsView(MakeArrayView(reinterpret_cast<UMassProcessor**>(&Processor), 1), ProcessingContext);
+			}
+
+			AITEST_EQUAL("Run 2: The entity handle before and after async commit", ReservedHandle, SyncBuilder.GetEntityHandle());
+			AITEST_TRUE("Run 2: The Builder is in `Committed` state", SyncBuilder.IsCommitted());
+			// since the commands are not flushed yet, due to ProcessingContext's values, we expect the entity to not be created yet
+			AITEST_FALSE("(NOT) Run 2: the entity has been created", EntityManager->IsEntityBuilt(ReservedHandle));
+
+			EntityManager->FlushCommands();
+			int32 ArchetypeEntitiesNum = EntityManager->DebugGetArchetypeEntitiesCount(IntsArchetype);
+			AITEST_EQUAL("Number of entities in the target archetype, after flushing", ArchetypeEntitiesNum, 2);
+			AITEST_TRUE("the entity has been created", EntityManager->IsEntityBuilt(ReservedHandle));
+
+			return true;
+		}
+	};
+	IMPLEMENT_AI_INSTANT_TEST(FSyncBuildingAsyncSubmission, "System.Mass.EntityBuilder.SyncBuildingAsyncSubmission");
 } // UE::Mass::Test::EntityBuilder
 
 UE_ENABLE_OPTIMIZATION_SHIP
diff --git a/Engine/Source/Runtime/MassEntity/Private/MassEntityBuilder.cpp b/Engine/Source/Runtime/MassEntity/Private/MassEntityBuilder.cpp
--- a/Engine/Source/Runtime/MassEntity/Private/MassEntityBuilder.cpp
+++ b/Engine/Source/Runtime/MassEntity/Private/MassEntityBuilder.cpp
@@ -1,6 +1,9 @@
 // Copyright Epic Games, Inc. All Rights Reserved.
 
 #include "MassEntityBuilder.h"
+
+#include "MassCommandBuffer.h"
+#include "MassCommands.h"
 #include "MassEntityManager.h"
 #include "VisualLogger/VisualLogger.h"
 
@@ -106,6 +109,8 @@ FEntityBuilder::~FEntityBuilder()
 
 FMassEntityHandle FEntityBuilder::Commit()
 {
+	// @todo consider locking every builder instance to a single thread to prevent concurent add/flush
+
 	if (!testableEnsureMsgf(State != EState::Committed, TEXT("Trying to commit an already committed EntityBuilder. This request will be ignored.")))
 	{
 		return EntityHandle;
@@ -118,32 +123,25 @@ FMassEntityHandle FEntityBuilder::Commit()
 	}
 
 	CacheEntityHandle();
-	
-	if (!ensureMsgf(EntityManager->IsProcessing() == false, TEXT("Committing EntityBuilder while Mass processing is in progress is not supported yet.")))
+	CacheSharedFragmentValue();
+	CacheArchetypeHandle();
+
+	if (EntityManager->IsProcessing())
 	{
-		// @todo we need to issue commands in this case
+		// we need to issue commands in this case
+		EntityManager->Defer().PushCommand<FMassDeferredCreateCommand>(
+					[ReservedEntityHandle = EntityHandle, SharedFragmentValues = CachedSharedFragmentValues, ArchetypeHandle = CachedArchetypeHandle, FragmentsCopy = Fragments](FMassEntityManager& Manager)
+					{
+						Manager.BuildEntity(ReservedEntityHandle, ArchetypeHandle, SharedFragmentValues);
+						Manager.SetEntityFragmentValues(ReservedEntityHandle, FragmentsCopy);
+					});
 	}
 	else
 	{
 		// directly create the entity
 		TSharedRef<FMassEntityManager::FEntityCreationContext> CreationContext = EntityManager->GetOrMakeCreationContext();
 
-		FMassArchetypeSharedFragmentValues SharedFragmentValues;
-		for (FInstancedStruct& SharedFragmentInstance : SharedFragments)
-		{
-			check(SharedFragmentInstance.IsValid());
-			const FSharedStruct& SharedStruct = EntityManager->GetOrCreateSharedFragment(*SharedFragmentInstance.GetScriptStruct(), SharedFragmentInstance.GetMemory());
-			SharedFragmentValues.AddSharedFragment(SharedStruct);
-		}
-		for (FInstancedStruct& ConstSharedFragmentInstance : ConstSharedFragments)
-		{
-			check(ConstSharedFragmentInstance.IsValid());
-			const FSharedStruct& ConstSharedStruct = EntityManager->GetOrCreateSharedFragment(*ConstSharedFragmentInstance.GetScriptStruct(), ConstSharedFragmentInstance.GetMemory());
-			SharedFragmentValues.AddConstSharedFragment(ConstSharedStruct);
-		}
-
-		FMassArchetypeHandle ArchetypeHandle = EntityManager->CreateArchetype(Composition);
-		EntityManager->BuildEntity(EntityHandle, ArchetypeHandle, SharedFragmentValues);
+		EntityManager->BuildEntity(EntityHandle, CachedArchetypeHandle, CachedSharedFragmentValues);
 		EntityManager->SetEntityFragmentValues(EntityHandle, Fragments);
 	}
 
@@ -152,6 +150,22 @@ FMassEntityHandle FEntityBuilder::Commit()
 	return EntityHandle;
 }
 
+FMassEntityHandle FEntityBuilder::CommitAndReprepare()
+{
+	FMassEntityHandle CreatedEntity = Commit();
+	Reprepare();
+	return CreatedEntity;
+}
+
+void FEntityBuilder::Reprepare()
+{
+	if (ensureMsgf(State == EState::Committed, TEXT("Expected to be called only on Committed builders")))
+	{
+		EntityHandle.Reset();
+		State = EState::ReadyToCommit;
+	}
+}
+
 void FEntityBuilder::Reset()
 {
 	ConditionallyReleaseEntityHandle();
@@ -189,4 +203,39 @@ void FEntityBuilder::CacheEntityHandle() const
 	}
 }
 
+void FEntityBuilder::CacheArchetypeHandle()
+{
+	if (CachedArchetypeHandle.IsValid() == false)
+	{
+		CachedArchetypeHandle = EntityManager->CreateArchetype(Composition);
+	}
+}
+
+FMassArchetypeHandle FEntityBuilder::GetArchetypeHandle()
+{
+	CacheArchetypeHandle();
+	return CachedArchetypeHandle;
+}
+
+void FEntityBuilder::CacheSharedFragmentValue()
+{
+	if (CachedSharedFragmentValues.IsEmpty())
+	{
+		for (FInstancedStruct& SharedFragmentInstance : SharedFragments)
+		{
+			check(SharedFragmentInstance.IsValid());
+			const FSharedStruct& SharedStruct = EntityManager->GetOrCreateSharedFragment(*SharedFragmentInstance.GetScriptStruct(), SharedFragmentInstance.GetMemory());
+			CachedSharedFragmentValues.AddSharedFragment(SharedStruct);
+		}
+		for (FInstancedStruct& ConstSharedFragmentInstance : ConstSharedFragments)
+		{
+			check(ConstSharedFragmentInstance.IsValid());
+			const FSharedStruct& ConstSharedStruct = EntityManager->GetOrCreateSharedFragment(*ConstSharedFragmentInstance.GetScriptStruct(), ConstSharedFragmentInstance.GetMemory());
+			CachedSharedFragmentValues.AddConstSharedFragment(ConstSharedStruct);
+		}
+
+		CachedSharedFragmentValues.Sort();
+	}
+}
+
 } // namespace UE::Mass
\ No newline at end of file
diff --git a/Engine/Source/Runtime/MassEntity/Public/MassEntityBuilder.h b/Engine/Source/Runtime/MassEntity/Public/MassEntityBuilder.h
--- a/Engine/Source/Runtime/MassEntity/Public/MassEntityBuilder.h
+++ b/Engine/Source/Runtime/MassEntity/Public/MassEntityBuilder.h
@@ -2,6 +2,7 @@
 
 #pragma once
 
+#include "MassArchetypeTypes.h"
 #include "MassEntityConcepts.h"
 #include "MassEntityHandle.h"
 #include "MassEntityTypes.h"
@@ -74,6 +75,16 @@ struct MASSENTITY_API FEntityBuilder
 	 */
 	FMassEntityHandle Commit();
 
+	/**
+	 * A wrapper for "Commit" call that, once that's done, prepares the builder for another commit, forgetting the
+	 * handle for the entity just created, and reverting the state back to "ReadyToCommit"
+	 * @see Commit
+	 */
+	FMassEntityHandle CommitAndReprepare();
+
+	/** if the builder is in "Committed" state it will roll back to ReadyToSubmit and reset the stored entity handle */
+	void Reprepare();
+
 	/** Resets the builder to its initial state, discarding all previous entity configurations. */
 	void Reset();
 
@@ -85,7 +96,7 @@ struct MASSENTITY_API FEntityBuilder
 	 * @return A reference to the added fragment.
 	 */
 	template<typename T, typename... TArgs>
-	T& Add_GetRef(TArgs&&... InArgs) requires (!(CTag<T> || CChunkFragment<T>));
+	T& Add_GetRef(TArgs&&... InArgs) requires (CElement<T> && !(CTag<T> || CChunkFragment<T>));
 
 	/**
 	 * Adds a fragment of type T to the entity and returns a reference to it, constructing it with the provided arguments.
@@ -93,7 +104,7 @@ struct MASSENTITY_API FEntityBuilder
 	 * @return A reference to the added fragment.
 	 */
 	template<typename T, typename... TArgs>
-	T& GetOrCreate(TArgs&&... InArgs) requires (!(CTag<T> || CChunkFragment<T>));
+	T& GetOrCreate(TArgs&&... InArgs) requires (CElement<T> && !(CTag<T> || CChunkFragment<T>));
 
 	/**
 	 * Adds a tag of type T to the entity.
@@ -114,7 +125,7 @@ struct MASSENTITY_API FEntityBuilder
 	 * @return Reference to this FEntityBuilder for method chaining.
 	 */
 	template<typename T, typename... TArgs>
-	FEntityBuilder& Add(TArgs&&... InArgs) requires (!(CTag<T> || CChunkFragment<T>));
+	FEntityBuilder& Add(TArgs&&... InArgs) requires (CElement<T> && !(CTag<T> || CChunkFragment<T>));
 
 	/**
 	 * Finds and retrieves a pointer to a fragment of type T if it exists.
@@ -122,11 +133,13 @@ struct MASSENTITY_API FEntityBuilder
 	 * @return Pointer to the fragment, or nullptr if it does not exist.
 	 */
 	template<typename T>
-	T* Find();
+	T* Find() requires (CElement<T> && !(CTag<T> || CChunkFragment<T>));
 
 	/** Converts the builder to a FMassEntityHandle, reserving the entity handle if not already committed. */
 	[[nodiscard]] FMassEntityHandle GetEntityHandle() const;
 
+	[[nodiscard]] FMassArchetypeHandle GetArchetypeHandle();
+
 	/** Checks whether the builder is in a valid, expected state */
 	bool IsValid() const;
 
@@ -136,27 +149,39 @@ struct MASSENTITY_API FEntityBuilder
 	/** @return whether the builder has already committed the data */
 	bool IsCommitted() const;
 
+protected:
+	void CacheSharedFragmentValue();
+	void CacheArchetypeHandle();
+
 private:
 	TSharedRef<FMassEntityManager> EntityManager;
 	mutable FMassEntityHandle EntityHandle;
 
 	FMassArchetypeCompositionDescriptor Composition;
 
+	FMassArchetypeSharedFragmentValues CachedSharedFragmentValues;
+	FMassArchetypeHandle CachedArchetypeHandle;
+
 	template<CFragment T> 
-	TArray<FInstancedStruct>& GetInstancedStructContainer() 
+	TArray<FInstancedStruct>& GetInstancedStructContainerInternal() 
 	{ 
 		return Fragments; 
 	}
 
 	template<CSharedFragment T> 
-	TArray<FInstancedStruct>& GetInstancedStructContainer() 
-	{ 
+	TArray<FInstancedStruct>& GetInstancedStructContainerInternal() 
+	{
+		// Resetting the cached shared values because this function is always called 
+		// with the intent to modify the contents of SharedFragments, invalidating the
+		// cached data anyway 
+		CachedSharedFragmentValues.Reset();
 		return SharedFragments; 
 	}
 
 	template<CConstSharedFragment T> 
-	TArray<FInstancedStruct>& GetInstancedStructContainer() 
-	{ 
+	TArray<FInstancedStruct>& GetInstancedStructContainerInternal() 
+	{
+		CachedSharedFragmentValues.Reset();
 		return ConstSharedFragments; 
 	}
 
@@ -209,6 +234,7 @@ FEntityBuilder& FEntityBuilder::Add()
 {
 	Composition.Tags.Add<T>();
 	State = EState::ReadyToCommit;
+	CachedArchetypeHandle = FMassArchetypeHandle();
 	return *this;
 }
 
@@ -217,35 +243,38 @@ FEntityBuilder& FEntityBuilder::Add()
 {
 	Composition.ChunkFragments.Add<T>();
 	State = EState::ReadyToCommit;
+	CachedArchetypeHandle = FMassArchetypeHandle();
 	return *this;
 }
 
 template<typename T, typename... TArgs>
-FEntityBuilder& FEntityBuilder::Add(TArgs&&... InArgs) requires (!(CTag<T> || CChunkFragment<T>))
+FEntityBuilder& FEntityBuilder::Add(TArgs&&... InArgs) requires (CElement<T> && !(CTag<T> || CChunkFragment<T>))
 {
 	Add_GetRef<T>(InArgs...);
 	return *this;
 }
 
 template<typename T, typename... TArgs>
-T& FEntityBuilder::Add_GetRef(TArgs&&... InArgs) requires (!(CTag<T> || CChunkFragment<T>))
+T& FEntityBuilder::Add_GetRef(TArgs&&... InArgs) requires (CElement<T> && !(CTag<T> || CChunkFragment<T>))
 {
 	State = EState::ReadyToCommit;
 
 	if (ensureMsgf(GetBitSetBuilder<T>().template Contains<T>() == false, TEXT("Element of type %s has already been added"), *T::StaticStruct()->GetName()))
 	{
 		GetBitSetBuilder<T>().template Add<T>();
-		return GetInstancedStructContainer<T>().Add_GetRef(FInstancedStruct::Make<T>(InArgs...)).template GetMutable<T>();
+		CachedArchetypeHandle = FMassArchetypeHandle();
+
+		return GetInstancedStructContainerInternal<T>().Add_GetRef(FInstancedStruct::Make<T>(InArgs...)).template GetMutable<T>();
 	}
 
-	return GetInstancedStructContainer<T>().FindByPredicate([](const FInstancedStruct& ExistingElement)
+	return GetInstancedStructContainerInternal<T>().FindByPredicate([](const FInstancedStruct& ExistingElement)
 		{
 			return ExistingElement.GetScriptStruct() == T::StaticStruct();
 		})->template GetMutable<T>();
 }
 
 template<typename T, typename... TArgs>
-T& FEntityBuilder::GetOrCreate(TArgs&&... InArgs) requires (!(CTag<T> || CChunkFragment<T>))
+T& FEntityBuilder::GetOrCreate(TArgs&&... InArgs) requires (CElement<T> && !(CTag<T> || CChunkFragment<T>))
 {
 	State = EState::ReadyToCommit;
 
@@ -253,7 +282,7 @@ T& FEntityBuilder::GetOrCreate(TArgs&&... InArgs) requires (!(CTag<T> || CChunkF
 	if (GetBitSetBuilder<T>().template Contains<T>())
 	{
 		// replace
-		ElementFound = GetInstancedStructContainer<T>().FindByPredicate([](const FInstancedStruct& ExistingElement)
+		ElementFound = GetInstancedStructContainerInternal<T>().FindByPredicate([](const FInstancedStruct& ExistingElement)
 		{
 			return ExistingElement.GetScriptStruct() == T::StaticStruct();
 		});
@@ -264,18 +293,20 @@ T& FEntityBuilder::GetOrCreate(TArgs&&... InArgs) requires (!(CTag<T> || CChunkF
 	else
 	{
 		GetBitSetBuilder<T>().template Add<T>();
-		ElementFound = &GetInstancedStructContainer<T>().Add_GetRef(FInstancedStruct::Make<T>(InArgs...));
+		CachedArchetypeHandle = FMassArchetypeHandle();
+
+		ElementFound = &GetInstancedStructContainerInternal<T>().Add_GetRef(FInstancedStruct::Make<T>(InArgs...));
 	}
 
 	return ElementFound->GetMutable<T>();
 }
 
 template<typename T>
-T* FEntityBuilder::Find()
+T* FEntityBuilder::Find() requires (CElement<T> && !(CTag<T> || CChunkFragment<T>))
 {
 	if (GetBitSetBuilder<T>().template Contains<T>())
 	{
-		FInstancedStruct* ElementFound = GetInstancedStructContainer<T>().FindByPredicate([](const FInstancedStruct& Element)
+		FInstancedStruct* ElementFound = GetInstancedStructContainerInternal<T>().FindByPredicate([](const FInstancedStruct& Element)
 		{
 			return Element.GetScriptStruct() == T::StaticStruct();
 		});
diff --git a/Engine/Source/Runtime/MassEntity/Public/MassEntityTypes.h b/Engine/Source/Runtime/MassEntity/Public/MassEntityTypes.h
--- a/Engine/Source/Runtime/MassEntity/Public/MassEntityTypes.h
+++ b/Engine/Source/Runtime/MassEntity/Public/MassEntityTypes.h
@@ -383,11 +383,7 @@ struct MASSENTITY_API FMassArchetypeSharedFragmentValues
 	}
 
 	uint32 CalculateHash() const;
-
-	SIZE_T GetAllocatedSize() const
-	{
-		return ConstSharedFragments.GetAllocatedSize() + SharedFragments.GetAllocatedSize();
-	}
+	SIZE_T GetAllocatedSize() const;
 
 	void Sort()
 	{
@@ -399,7 +395,11 @@ struct MASSENTITY_API FMassArchetypeSharedFragmentValues
 		}
 	}
 
-	bool IsSorted() const { return bSorted; }
+	bool IsSorted() const;
+
+	bool IsEmpty() const;
+
+	void Reset();
 
 protected:
 	mutable uint32 HashCache = UINT32_MAX;
@@ -698,3 +698,31 @@ bool FMassArchetypeCompositionDescriptor::Add()
 {
 	return GetContainer<T>().template Add<T>();
 }
+
+inline SIZE_T FMassArchetypeSharedFragmentValues::GetAllocatedSize() const
+{
+	return SharedFragmentBitSet.GetAllocatedSize()
+		+ ConstSharedFragmentBitSet.GetAllocatedSize()
+		+ ConstSharedFragments.GetAllocatedSize()
+		+ SharedFragments.GetAllocatedSize();
+}
+
+inline bool FMassArchetypeSharedFragmentValues::IsSorted() const
+{
+	return bSorted;
+}
+
+inline bool FMassArchetypeSharedFragmentValues::IsEmpty() const
+{
+	return ConstSharedFragments.IsEmpty() && SharedFragments.IsEmpty();
+}
+
+inline void FMassArchetypeSharedFragmentValues::Reset()
+{
+	HashCache = UINT32_MAX;
+	bSorted = false; 
+	SharedFragmentBitSet.Reset();
+	ConstSharedFragmentBitSet.Reset();
+	ConstSharedFragments.Reset();
+	SharedFragments.Reset();
+}
