[Mass] Fixed comment style in MassEntityManagerStorage.h

#rnx

diff --git a/Engine/Source/Runtime/MassEntity/Public/MassEntityManagerStorage.h b/Engine/Source/Runtime/MassEntity/Public/MassEntityManagerStorage.h
--- a/Engine/Source/Runtime/MassEntity/Public/MassEntityManagerStorage.h
+++ b/Engine/Source/Runtime/MassEntity/Public/MassEntityManagerStorage.h
@@ -11,26 +11,34 @@
 struct FMassArchetypeData;
 struct FMassEntityHandle;
 
-// Initialization parameters to configure MassEntityManager to reserve entities only single threaded
-// Supported in all build configurations
+/**
+ * Initialization parameters to configure MassEntityManager to reserve entities only single threaded
+ * Supported in all build configurations
+ */
 struct MASSENTITY_API FMassEntityManager_InitParams_SingleThreaded {};
 
-// Initialization parameters to configure MassEntityManager to concurrently reserve entities
-// Only supported in editor builds.
-//
-// Expected static memory requirement for array of Page pointers can be computed:
-// MaxPages = MaxEntityCount / MaxEntitiesPerPage
-// MemorySize = MaxPages * sizeof(Page**)
-//
-// For default values, expectation is 128kB
+/**
+ * Initialization parameters to configure MassEntityManager to concurrently reserve entities
+ * Only supported in editor builds.
+ *
+ * Expected static memory requirement for array of Page pointers can be computed:
+ * MaxPages = MaxEntityCount / MaxEntitiesPerPage
+ * MemorySize = MaxPages * sizeof(Page**)
+ *
+ * For default values, expectation is 128kB
+ */
 struct MASSENTITY_API FMassEntityManager_InitParams_Concurrent
 {
-	// Maximum supported entities by the MassEntityManager
-	// Must be multiple of 2
+	/** 
+	 * Maximum supported entities by the MassEntityManager
+	 * Must be multiple of 2
+	 */
 	uint32 MaxEntityCount = 1 << 30; // 1 billion
 
-	// Number of entities per chunk
-	// Must be multiple of 2
+	/** 
+	 * Number of entities per chunk
+	 * Must be multiple of 2
+	 */
 	uint32 MaxEntitiesPerPage = 1 << 16; // 65536
 };
 
@@ -47,11 +55,11 @@ namespace UE::Mass
 	public:
 		enum class EEntityState
 		{
-			// Entity index refers to an entity that is free to be reserved or created
+			/** Entity index refers to an entity that is free to be reserved or created */
 			Free,
-			// Entity index refers to a reserved entity
+			/** Entity index refers to a reserved entity */
 			Reserved,
-			// Entity index refers to an entity assigned to an archetype
+			/** Entity index refers to an entity assigned to an archetype */
 			Created
 		};
 		virtual ~IEntityStorageInterface() = default;
@@ -60,46 +68,55 @@ namespace UE::Mass
 		virtual const FMassArchetypeData* GetArchetype(int32 Index) const = 0;
 		virtual TSharedPtr<FMassArchetypeData>& GetArchetypeAsShared(int32 Index) = 0;
 		virtual const TSharedPtr<FMassArchetypeData>& GetArchetypeAsShared(int32 Index) const = 0;
-	
+
 		virtual void SetArchetypeFromShared(int32 Index, TSharedPtr<FMassArchetypeData>& Archetype) = 0;
 		virtual void SetArchetypeFromShared(int32 Index, const TSharedPtr<FMassArchetypeData>& Archetype) = 0;
 
-		// Returns true if the given entity at index is currently reserved
-		// False if free or assigned an archetype
+		/**
+		 * Returns true if the given entity at index is currently reserved
+		 * False if free or assigned an archetype
+		 */
 		virtual EEntityState GetEntityState(int32 Index) const = 0;
-	
+
 		virtual int32 GetSerialNumber(int32 Index) const = 0;
 
-		// Checks if index can be used to access entity data
+		/** Checks if index can be used to access entity data */
 		virtual bool IsValidIndex(int32 Index) const = 0;
 		virtual SIZE_T GetAllocatedSize() const = 0;
 
-		// Checks if entity at Index is built
+		/** Checks if entity at Index is built */
 		virtual bool IsValid(int32 Index) const = 0;
 
 		virtual FMassEntityHandle AcquireOne() = 0;
-		// @return number of entities actually added
+		/**
+		 * @return number of entities actually added
+		 */
 		int32 Acquire(const int32 Count, TArray<FMassEntityHandle>& OutEntityHandles);
 		virtual int32 Acquire(TArrayView<FMassEntityHandle> OutEntityHandles) = 0;
 
 		virtual int32 Release(TConstArrayView<FMassEntityHandle> Handles) = 0;
 		virtual int32 ReleaseOne(FMassEntityHandle Handles) = 0;
-		
-		// Bypasses Serial Number Check
-		// Only use if caller has ensured serial number matches or for debug purposes
+
+		/**
+		 * Bypasses Serial Number Check
+		 * Only use if caller has ensured serial number matches or for debug purposes
+		 */
 		virtual int32 ForceRelease(TConstArrayView<FMassEntityHandle> Handles) = 0;
 		virtual int32 ForceReleaseOne(FMassEntityHandle Handle) = 0;
 
-		// Returns the number of entities that are not free
-		// For debug purposes only.  In multi-threaded environments, the result is going to be out of date
+		/**
+		 * Returns the number of entities that are not free
+		 * For debug purposes only. In multi-threaded environments, the result is going to be out of date
+		 */
 		virtual int32 Num() const = 0;
 
-		// Returns the number of entities that are free
-		// For debug purposes only.  In multi-threaded environments, the result is going to be out of date
+		/**
+		 * Returns the number of entities that are free
+		 * For debug purposes only. In multi-threaded environments, the result is going to be out of date
+		 */
 		virtual int32 ComputeFreeSize() const = 0;
 	};
 
-	
 	//-----------------------------------------------------------------------------
 	// FSingleThreadedEntityStorage
 	//-----------------------------------------------------------------------------
@@ -132,12 +149,12 @@ namespace UE::Mass
 		virtual int32 Num() const override;
 		virtual int32 ComputeFreeSize() const override;
 	private:
-	
+
 		struct FEntityData
 		{
 			TSharedPtr<FMassArchetypeData> CurrentArchetype;
 			int32 SerialNumber = 0;
-			
+
 			~FEntityData();
 			void Reset();
 			bool IsValid() const;
@@ -152,9 +169,9 @@ namespace UE::Mass
 	// FConcurrentEntityStorage
 	//-----------------------------------------------------------------------------
 	/**
-	 * This storage backend allows for entities to be concurrently reserved.  Reserved entities can also
+	 * This storage backend allows for entities to be concurrently reserved. Reserved entities can also
 	 * be concurrently freed.
-	 * Creation of entities (ie. assignment of an archetype and addition of data into chunks) cannot be done
+	 * Creation of entities (i.e. assignment of an archetype and addition of data into chunks) cannot be done
 	 * concurrently with this implementation.
 	 */
 	class FConcurrentEntityStorage final : public IEntityStorageInterface
@@ -189,46 +206,48 @@ namespace UE::Mass
 		/** @return whether the assumptions are still valid */
 		MASSENTITY_API static bool DebugAssumptionsSelfTest();
 #endif // WITH_MASSENTITY_DEBUG
-	private:	
-	
+	private:
+
 		struct FEntityData
 		{
 			static constexpr int MaxGenerationBits = 30;
-		
+
 			TSharedPtr<FMassArchetypeData> CurrentArchetype;
-			// Generation ID or version of the entity in this slot
+			/** Generation ID or version of the entity in this slot */
 			uint32 GenerationId : MaxGenerationBits = 0;
-			// 1 if the entity is NOT free
+			/** 1 if the entity is NOT free */
 			uint32 bIsAllocated : 1 = 0;
 
 			~FEntityData();
-			// Converts EntityData state into a SerialNumber for public usage
+			/** Converts EntityData state into a SerialNumber for public usage */
 			int32 GetSerialNumber() const;
 
 			bool operator==(const FEntityData& Other) const;
 		};
-	
+
 		FEntityData& LookupEntity(int32 Index);
 		const FEntityData& LookupEntity(int32 Index) const;
 
-		// Returns size of a page in bytes
+		/** Returns size of a page in bytes */
 		uint64 ComputePageSize() const;
 
-		// @return whether the operation was successful. Will return false when OOM
+		/**
+		 * @return whether the operation was successful. Will return false when OOM
+		 */
 		bool AddPage();
 
-		// Number of allocated Entities (only used for viewing in the debugger).
+		/** Number of allocated Entities (only used for viewing in the debugger). */
 		uint32 EntityCount = 0;
 		uint32 MaximumEntityCountShift = 0;
 		uint32 MaxEntitiesPerPage = 0;
 		uint32 PageCount = 0;
-		// ALWAYS acquire FreeListMutex before this one
+		/** ALWAYS acquire FreeListMutex before this one */
 		FTransactionallySafeMutex PageAllocateMutex;
-		// Pointer to array of pages
+		/** Pointer to array of pages */
 		FEntityData** EntityPages = nullptr;
-	
+
 		TArray<int32> EntityFreeIndexList;
 		FTransactionallySafeMutex FreeListMutex;
 	};
-	
+
 } // namespace UE::Mass
