[Mass] Added a function to MassEntityStorage directly checking entity handle's validity

#rnx

diff --git a/Engine/Source/Runtime/MassEntity/Private/MassEntityManagerStorage.cpp b/Engine/Source/Runtime/MassEntity/Private/MassEntityManagerStorage.cpp
--- a/Engine/Source/Runtime/MassEntity/Private/MassEntityManagerStorage.cpp
+++ b/Engine/Source/Runtime/MassEntity/Private/MassEntityManagerStorage.cpp
@@ -93,6 +93,12 @@ namespace UE::Mass
 		return Entities.IsValidIndex(Index);
 	}
 
+	bool FSingleThreadedEntityStorage::IsValidHandle(FMassEntityHandle EntityHandle) const
+	{
+		return Entities.IsValidIndex(EntityHandle.Index)
+			&& Entities[EntityHandle.Index].SerialNumber == EntityHandle.SerialNumber;
+	}
+
 	SIZE_T FSingleThreadedEntityStorage::GetAllocatedSize() const
 	{
 		return Entities.GetAllocatedSize() + EntityFreeIndexList.GetAllocatedSize();
@@ -322,6 +328,12 @@ namespace UE::Mass
 		return false;
 	}
 
+	bool FConcurrentEntityStorage::IsValidHandle(FMassEntityHandle EntityHandle) const
+	{
+		return IsValidIndex(EntityHandle.Index)
+			&& LookupEntity(EntityHandle.Index).GetSerialNumber() == EntityHandle.SerialNumber;
+	}
+
 	SIZE_T FConcurrentEntityStorage::GetAllocatedSize() const
 	{
 		const SIZE_T EntityFreeListSizeBytes = EntityFreeIndexList.GetAllocatedSize();
diff --git a/Engine/Source/Runtime/MassEntity/Public/MassEntityManagerStorage.h b/Engine/Source/Runtime/MassEntity/Public/MassEntityManagerStorage.h
--- a/Engine/Source/Runtime/MassEntity/Public/MassEntityManagerStorage.h
+++ b/Engine/Source/Runtime/MassEntity/Public/MassEntityManagerStorage.h
@@ -82,12 +82,21 @@ namespace UE::Mass
 
 		/** Checks if index can be used to access entity data */
 		virtual bool IsValidIndex(int32 Index) const = 0;
+
+		/**
+		 * Checks if the given handle is valid in the context od this storage, i.e. whether the
+		 * index is valid and the serial number associated with it matches the handle's
+		 */
+		virtual bool IsValidHandle(FMassEntityHandle EntityHandle) const = 0;
+
 		virtual SIZE_T GetAllocatedSize() const = 0;
 
 		/** Checks if entity at Index is built */
 		virtual bool IsValid(int32 Index) const = 0;
 
+		/** Produce a single entity handle */
 		virtual FMassEntityHandle AcquireOne() = 0;
+
 		/**
 		 * @return number of entities actually added
 		 */
@@ -137,6 +146,7 @@ namespace UE::Mass
 		virtual EEntityState GetEntityState(int32 Index) const override;
 		virtual int32 GetSerialNumber(int32 Index) const override;
 		virtual bool IsValidIndex(int32 Index) const override;
+		virtual bool IsValidHandle(FMassEntityHandle EntityHandle) const override;
 		virtual SIZE_T GetAllocatedSize() const override;
 		virtual bool IsValid(int32 Index) const override;
 		virtual FMassEntityHandle AcquireOne() override;
@@ -191,6 +201,7 @@ namespace UE::Mass
 		virtual EEntityState GetEntityState(int32 Index) const override;
 		virtual int32 GetSerialNumber(int32 Index) const override;
 		virtual bool IsValidIndex(int32 Index) const override;
+		virtual bool IsValidHandle(FMassEntityHandle EntityHandle) const override;
 		virtual SIZE_T GetAllocatedSize() const override;
 		virtual bool IsValid(int32 Index) const override;
 		virtual FMassEntityHandle AcquireOne() override;
