[Mass] Introduced a new concept, UE::Mass::CNonTag, to indicate all Mass elements that are not a Tag

diff --git a/Engine/Source/Runtime/MassEntity/Public/MassEntityConcepts.h b/Engine/Source/Runtime/MassEntity/Public/MassEntityConcepts.h
--- a/Engine/Source/Runtime/MassEntity/Public/MassEntityConcepts.h
+++ b/Engine/Source/Runtime/MassEntity/Public/MassEntityConcepts.h
@@ -23,5 +23,8 @@ namespace UE::Mass
 	concept CConstSharedFragment = TIsDerivedFrom<typename TRemoveReference<T>::Type, FMassConstSharedFragment>::Value;
 
 	template <typename T>
-	concept CElement = CFragment<T> || CTag<T> || CChunkFragment<T> || CSharedFragment<T> || CConstSharedFragment<T>;
+	concept CNonTag = CFragment<T> || CChunkFragment<T> || CSharedFragment<T> || CConstSharedFragment<T>;
+
+	template <typename T>
+	concept CElement = CNonTag<T> || CTag<T>;
 }
