[Mass] minor: refactored all MassEntityManager API parameters of format `FMassEntityHandle Entity` to FMassEntityHandle EntityHandle`, and updated the comments a bit.

#rnx

diff --git a/Engine/Source/Runtime/MassEntity/Public/MassEntityManager.h b/Engine/Source/Runtime/MassEntity/Public/MassEntityManager.h
--- a/Engine/Source/Runtime/MassEntity/Public/MassEntityManager.h
+++ b/Engine/Source/Runtime/MassEntity/Public/MassEntityManager.h
@@ -193,13 +193,13 @@ public:
 		, const FMassConstSharedFragmentBitSet& ConstSharedFragmentBitSet
 		, const FMassArchetypeCreationParams& CreationParams = FMassArchetypeCreationParams());
 
-	/** Fetches the archetype for a given Entity. If Entity is not valid it will still return a handle, just with an invalid archetype */
-	FMassArchetypeHandle GetArchetypeForEntity(FMassEntityHandle Entity) const;
+	/** Fetches the archetype for a given EntityHandle. If EntityHandle is not valid it will still return a handle, just with an invalid archetype */
+	FMassArchetypeHandle GetArchetypeForEntity(FMassEntityHandle EntityHandle) const;
 	/**
-	 * Fetches the archetype for a given Entity. Note that it's callers responsibility the given Entity handle is valid.
+	 * Fetches the archetype for a given EntityHandle. Note that it's callers responsibility the given EntityHandle handle is valid.
 	 * If you can't ensure that call GetArchetypeForEntity.
 	 */
-	FMassArchetypeHandle GetArchetypeForEntityUnsafe(FMassEntityHandle Entity) const;
+	FMassArchetypeHandle GetArchetypeForEntityUnsafe(FMassEntityHandle EntityHandle) const;
 
 	/** Method to iterate on all the fragment types of an archetype */
 	static void ForEachArchetypeFragmentType(const FMassArchetypeHandle& ArchetypeHandle, TFunction< void(const UScriptStruct* /*FragmentType*/)> Function);
@@ -306,8 +306,8 @@ public:
 
 	/**
 	 * Destroys a fully built entity, use ReleaseReservedEntity if entity was not yet built.
-	 * @param Entity to destroy */
-	void DestroyEntity(FMassEntityHandle Entity);
+	 * @param EntityHandle identifying the entity to destroy */
+	void DestroyEntity(FMassEntityHandle EntityHandle);
 
 	/**
 	 * Reserves an entity in the subsystem, the entity is still not ready to be used by the subsystem, need to call BuildEntity()
@@ -316,24 +316,24 @@ public:
 
 	/**
 	 * Builds an entity for it to be ready to be used by the subsystem
-	 * @param Entity to build which was retrieved with ReserveEntity() method
+	 * @param EntityHandle identifying the entity to build, which was retrieved with ReserveEntity() method
 	 * @param ArchetypeHandle you want this entity to be
 	 * @param SharedFragmentValues to be associated with the entity
 	 */
-	void BuildEntity(FMassEntityHandle Entity, const FMassArchetypeHandle& ArchetypeHandle, const FMassArchetypeSharedFragmentValues& SharedFragmentValues = {});
+	void BuildEntity(FMassEntityHandle EntityHandle, const FMassArchetypeHandle& ArchetypeHandle, const FMassArchetypeSharedFragmentValues& SharedFragmentValues = {});
 
 	/**
 	 * Builds an entity for it to be ready to be used by the subsystem
-	 * @param Entity to build which was retrieved with ReserveEntity() method
+	 * @param EntityHandle identifying the entity to build, which was retrieved with ReserveEntity() method
 	 * @param FragmentInstanceList is the fragments to create the entity from and initialize values
 	 * @param SharedFragmentValues to be associated with the entity
 	 */
-	void BuildEntity(FMassEntityHandle Entity, TConstArrayView<FInstancedStruct> FragmentInstanceList, const FMassArchetypeSharedFragmentValues& SharedFragmentValues = {});
+	void BuildEntity(FMassEntityHandle EntityHandle, TConstArrayView<FInstancedStruct> FragmentInstanceList, const FMassArchetypeSharedFragmentValues& SharedFragmentValues = {});
 
 	/*
-	 * Releases a previously reserved entity that was not yet built, otherwise call DestroyEntity
-	 * @param Entity to release */
-	void ReleaseReservedEntity(FMassEntityHandle Entity);
+	 * Releases a previously reserved entity handle that was not yet built, otherwise call DestroyEntity
+	 * @param EntityHandle to release */
+	void ReleaseReservedEntity(FMassEntityHandle EntityHandle);
 
 	/**
 	 * Destroys all the entities in the provided array of entities. The function will also gracefully handle entities
@@ -373,56 +373,56 @@ public:
 
 	const UE::Mass::FArchetypeGroups& GetGroupsForArchetype(const FMassArchetypeHandle& ArchetypeHandle) const;
 
-	void AddFragmentToEntity(FMassEntityHandle Entity, const UScriptStruct* FragmentType);
-	void AddFragmentToEntity(FMassEntityHandle Entity, const UScriptStruct* FragmentType, const FStructInitializationCallback& Initializer);
+	void AddFragmentToEntity(FMassEntityHandle EntityHandle, const UScriptStruct* FragmentType);
+	void AddFragmentToEntity(FMassEntityHandle EntityHandle, const UScriptStruct* FragmentType, const FStructInitializationCallback& Initializer);
 
 	/** 
 	 *  Ensures that only unique fragments are added. 
-	 *  @note It's caller's responsibility to ensure Entity's and FragmentList's validity. 
+	 *  @note It's caller's responsibility to ensure EntityHandle's and FragmentList's validity. 
 	 */
-	void AddFragmentListToEntity(FMassEntityHandle Entity, TConstArrayView<const UScriptStruct*> FragmentList);
+	void AddFragmentListToEntity(FMassEntityHandle EntityHandle, TConstArrayView<const UScriptStruct*> FragmentList);
 
-	void AddFragmentInstanceListToEntity(FMassEntityHandle Entity, TConstArrayView<FInstancedStruct> FragmentInstanceList);
-	void RemoveFragmentFromEntity(FMassEntityHandle Entity, const UScriptStruct* FragmentType);
-	void RemoveFragmentListFromEntity(FMassEntityHandle Entity, TConstArrayView<const UScriptStruct*> FragmentList);
+	void AddFragmentInstanceListToEntity(FMassEntityHandle EntityHandle, TConstArrayView<FInstancedStruct> FragmentInstanceList);
+	void RemoveFragmentFromEntity(FMassEntityHandle EntityHandle, const UScriptStruct* FragmentType);
+	void RemoveFragmentListFromEntity(FMassEntityHandle EntityHandle, TConstArrayView<const UScriptStruct*> FragmentList);
 
-	void AddTagToEntity(FMassEntityHandle Entity, const UScriptStruct* TagType);
-	void RemoveTagFromEntity(FMassEntityHandle Entity, const UScriptStruct* TagType);
-	void SwapTagsForEntity(FMassEntityHandle Entity, const UScriptStruct* FromFragmentType, const UScriptStruct* ToFragmentType);
+	void AddTagToEntity(FMassEntityHandle EntityHandle, const UScriptStruct* TagType);
+	void RemoveTagFromEntity(FMassEntityHandle EntityHandle, const UScriptStruct* TagType);
+	void SwapTagsForEntity(FMassEntityHandle EntityHandle, const UScriptStruct* FromFragmentType, const UScriptStruct* ToFragmentType);
 
 	/** 
 	 * Adds a new const shared fragment to the given entity. Note that it only works if the given entity doesn't have
 	 * a shared fragment of the given type. The function will give a soft "pass" if the entity has the shared fragment
 	 * of the same value. Setting shared fragment value (i.e. changing) is not supported and the function will log
 	 * a warning if that's attempted.
-	 * @return whether the Entity has the Fragment value assigned to it, regardless of its original state (i.e. the function will
-	 *	return true also if the Entity already had the same values associated with it)
+	 * @return whether the entity has the Fragment value assigned to it, regardless of its original state (i.e. the function will
+	 *	return true also if the entity already had the same values associated with it)
 	 */
-	bool AddConstSharedFragmentToEntity(const FMassEntityHandle Entity, const FConstSharedStruct& InConstSharedFragment);
+	bool AddConstSharedFragmentToEntity(const FMassEntityHandle EntityHandle, const FConstSharedStruct& InConstSharedFragment);
 
 	/**
 	 * Removes a const shared fragment of the given type from the entity.
 	 * Will do nothing if entity did not have the shared fragment.
 	 * @return True if fragment removed from entity, false otherwise.
 	 */
-	bool RemoveConstSharedFragmentFromEntity(const FMassEntityHandle Entity, const UScriptStruct& ConstSharedFragmentType);
+	bool RemoveConstSharedFragmentFromEntity(const FMassEntityHandle EntityHandle, const UScriptStruct& ConstSharedFragmentType);
 
 	/**
 	 * Adds a new shared fragment to the given entity. Note that it only works if the given entity doesn't have
 	 * a shared fragment of the given type. The function will give a soft "pass" if the entity has the shared fragment
 	 * of the same value. Setting shared fragment value (i.e. changing) is not supported and the function will log
 	 * a warning if that's attempted.
-	 * @return whether the Entity has the Fragment value assigned to it, regardless of its original state (i.e. the function will
-	 *	return true also if the Entity already had the same values associated with it)
+	 * @return whether the entity has the Fragment value assigned to it, regardless of its original state (i.e. the function will
+	 *	return true also if the entity already had the same values associated with it)
 	 */
-	bool AddSharedFragmentToEntity(const FMassEntityHandle Entity, const FSharedStruct& InSharedFragment);
+	bool AddSharedFragmentToEntity(const FMassEntityHandle EntityHandle, const FSharedStruct& InSharedFragment);
 
 	/**
 	 * Removes a shared fragment of the given type from the entity.
 	 * Will do nothing if entity did not have the shared fragment.
 	 * @return True if fragment removed from entity, false otherwise.
 	 */
-	bool RemoveSharedFragmentFromEntity(const FMassEntityHandle Entity, const UScriptStruct& SharedFragmentType);
+	bool RemoveSharedFragmentFromEntity(const FMassEntityHandle EntityHandle, const UScriptStruct& SharedFragmentType);
 
 	/**
 	 * Removes EntityHandle from any-and-all groups of given type - i.e. the entity will be moved to an archetype
@@ -461,26 +461,29 @@ public:
 	void BatchAddSharedFragmentsForEntities(TConstArrayView<FMassArchetypeEntityCollection> EntityCollections, const FMassArchetypeSharedFragmentValues& AddedFragmentValues);
 
 	/**
-	 * Adds fragments and tags indicated by InOutDescriptor to the Entity. The function also figures out which elements
+	 * Adds fragments and tags indicated by InOutDescriptor to the entity indicated by EntityHandle. The function also figures out which elements
 	 * in InOutDescriptor are missing from the current composition of the given entity and then returns the resulting 
 	 * delta via InOutDescriptor.
 	 */
-	void AddCompositionToEntity_GetDelta(FMassEntityHandle Entity, FMassArchetypeCompositionDescriptor& InOutDescriptor);
-	void RemoveCompositionFromEntity(FMassEntityHandle Entity, const FMassArchetypeCompositionDescriptor& InDescriptor);
+	void AddCompositionToEntity_GetDelta(FMassEntityHandle EntityHandle, FMassArchetypeCompositionDescriptor& InOutDescriptor);
+	void RemoveCompositionFromEntity(FMassEntityHandle EntityHandle, const FMassArchetypeCompositionDescriptor& InDescriptor);
 
 	const FMassArchetypeCompositionDescriptor& GetArchetypeComposition(const FMassArchetypeHandle& ArchetypeHandle) const;
 
 	/** 
 	 * Moves an entity over to a new archetype by copying over fragments common to both archetypes
-	 * @param Entity is the entity to move 
+	 * @param EntityHandle idicates the entity to move 
 	 * @param NewArchetypeHandle the handle to the new archetype
-	 * @param SharedFragmentValuesOverride if provided will override all given Entity's shared fragment values
+	 * @param SharedFragmentValuesOverride if provided will override all given entity's shared fragment values
 	 */
-	void MoveEntityToAnotherArchetype(FMassEntityHandle Entity, FMassArchetypeHandle NewArchetypeHandle, const FMassArchetypeSharedFragmentValues* SharedFragmentValuesOverride = nullptr);
+	void MoveEntityToAnotherArchetype(FMassEntityHandle EntityHandle, FMassArchetypeHandle NewArchetypeHandle, const FMassArchetypeSharedFragmentValues* SharedFragmentValuesOverride = nullptr);
 
-	/** Copies values from FragmentInstanceList over to Entity's fragment. Caller is responsible for ensuring that 
-	 *  the given entity does have given fragments. Failing this assumption will cause a check-fail.*/
-	void SetEntityFragmentValues(FMassEntityHandle Entity, TArrayView<const FInstancedStruct> FragmentInstanceList);
+	/** 
+	 *  Copies values from FragmentInstanceList over to target entity's fragment. Caller is responsible for ensuring that 
+	 *  the given entity does have given fragments. Failing this assumption will cause a check-fail.
+	 *  @param EntityHandle idicates the target entity
+	 */
+	void SetEntityFragmentValues(FMassEntityHandle EntityHandle, TArrayView<const FInstancedStruct> FragmentInstanceList);
 
 	/** Copies values from FragmentInstanceList over to fragments of given entities collection. The caller is responsible 
 	 *  for ensuring that the given entity archetype (FMassArchetypeEntityCollection .Archetype) does have given fragments. 
@@ -490,67 +493,67 @@ public:
 	void BatchSetEntityFragmentValues(TConstArrayView<FMassArchetypeEntityCollection> EntityCollections, TArrayView<const FInstancedStruct> FragmentInstanceList);
 
 	// Return true if it is an valid built entity
-	bool IsEntityActive(FMassEntityHandle Entity) const 
+	bool IsEntityActive(FMassEntityHandle EntityHandle) const 
 	{
-		return IsEntityValid(Entity) && IsEntityBuilt(Entity);
+		return IsEntityValid(EntityHandle) && IsEntityBuilt(EntityHandle);
 	}
 
-	// Returns true if Entity is valid
-	bool IsEntityValid(FMassEntityHandle Entity) const;
+	// Returns true if EntityHandle is valid
+	bool IsEntityValid(FMassEntityHandle EntityHandle) const;
 
-	// Returns true if Entity is has been fully built (expecting a valid Entity)
-	bool IsEntityBuilt(FMassEntityHandle Entity) const;
+	// Returns true if EntityHandle is has been fully built (expecting a valid EntityHandle)
+	bool IsEntityBuilt(FMassEntityHandle EntityHandle) const;
 
 	// Asserts that IsEntityValid
-	void CheckIfEntityIsValid(FMassEntityHandle Entity) const;
+	void CheckIfEntityIsValid(FMassEntityHandle EntityHandle) const;
 
 	// Asserts that IsEntityBuilt
-	void CheckIfEntityIsActive(FMassEntityHandle Entity) const;
+	void CheckIfEntityIsActive(FMassEntityHandle EntityHandle) const;
 
 	template<typename FragmentType>
-	FragmentType& GetFragmentDataChecked(FMassEntityHandle Entity) const
+	FragmentType& GetFragmentDataChecked(FMassEntityHandle EntityHandle) const
 	{
 		static_assert(TIsDerivedFrom<FragmentType, FMassFragment>::IsDerived
 			, "Given struct doesn't represent a valid fragment type. Make sure to inherit from FMassFragment or one of its child-types.");
-		return *((FragmentType*)InternalGetFragmentDataChecked(Entity, FragmentType::StaticStruct()));
+		return *((FragmentType*)InternalGetFragmentDataChecked(EntityHandle, FragmentType::StaticStruct()));
 	}
 
 	template<typename FragmentType>
-	FragmentType* GetFragmentDataPtr(FMassEntityHandle Entity) const
+	FragmentType* GetFragmentDataPtr(FMassEntityHandle EntityHandle) const
 	{
 		static_assert(TIsDerivedFrom<FragmentType, FMassFragment>::IsDerived
 			, "Given struct doesn't represent a valid fragment type. Make sure to inherit from FMassFragment or one of its child-types.");
-		return (FragmentType*)InternalGetFragmentDataPtr(Entity, FragmentType::StaticStruct());
+		return (FragmentType*)InternalGetFragmentDataPtr(EntityHandle, FragmentType::StaticStruct());
 	}
 
-	FStructView GetFragmentDataStruct(FMassEntityHandle Entity, const UScriptStruct* FragmentType) const
+	FStructView GetFragmentDataStruct(FMassEntityHandle EntityHandle, const UScriptStruct* FragmentType) const
 	{
 		checkf((FragmentType != nullptr) && FragmentType->IsChildOf(FMassFragment::StaticStruct())
 			, TEXT("GetFragmentDataStruct called with an invalid fragment type '%s'"), *GetPathNameSafe(FragmentType));
-		return FStructView(FragmentType, static_cast<uint8*>(InternalGetFragmentDataPtr(Entity, FragmentType)));
+		return FStructView(FragmentType, static_cast<uint8*>(InternalGetFragmentDataPtr(EntityHandle, FragmentType)));
 	}
 
 	template<typename ConstSharedFragmentType>
-	ConstSharedFragmentType* GetConstSharedFragmentDataPtr(FMassEntityHandle Entity) const
+	ConstSharedFragmentType* GetConstSharedFragmentDataPtr(FMassEntityHandle EntityHandle) const
 	{
 		static_assert(TIsDerivedFrom<ConstSharedFragmentType, FMassConstSharedFragment>::IsDerived, "Given struct doesn't represent a valid const shared fragment type. Make sure to inherit from FMassConstSharedFragment or one of its child-types.");
-		const FConstSharedStruct* ConstSharedStruct = InternalGetConstSharedFragmentPtr(Entity, ConstSharedFragmentType::StaticStruct());
+		const FConstSharedStruct* ConstSharedStruct = InternalGetConstSharedFragmentPtr(EntityHandle, ConstSharedFragmentType::StaticStruct());
 		return (ConstSharedFragmentType*)(ConstSharedStruct ? ConstSharedStruct->GetMemory() : nullptr);
 	}
 
 	template<typename ConstSharedFragmentType>
-	ConstSharedFragmentType& GetConstSharedFragmentDataChecked(FMassEntityHandle Entity) const
+	ConstSharedFragmentType& GetConstSharedFragmentDataChecked(FMassEntityHandle EntityHandle) const
 	{
-		ConstSharedFragmentType* TypePtr = GetConstSharedFragmentDataPtr<ConstSharedFragmentType>(Entity);
+		ConstSharedFragmentType* TypePtr = GetConstSharedFragmentDataPtr<ConstSharedFragmentType>(EntityHandle);
 		check(TypePtr);
 		return *TypePtr;
 	}
 
-	FConstStructView GetConstSharedFragmentDataStruct(FMassEntityHandle Entity, const UScriptStruct* ConstSharedFragmentType) const
+	FConstStructView GetConstSharedFragmentDataStruct(FMassEntityHandle EntityHandle, const UScriptStruct* ConstSharedFragmentType) const
 	{
 		checkf((ConstSharedFragmentType != nullptr) && ConstSharedFragmentType->IsChildOf(FMassConstSharedFragment::StaticStruct())
 			, TEXT("GetConstSharedFragmentDataStruct called with an invalid fragment type '%s'"), *GetPathNameSafe(ConstSharedFragmentType));
-		const FConstSharedStruct* ConstSharedStruct = InternalGetConstSharedFragmentPtr(Entity, ConstSharedFragmentType);
+		const FConstSharedStruct* ConstSharedStruct = InternalGetConstSharedFragmentPtr(EntityHandle, ConstSharedFragmentType);
 		return ConstSharedStruct
 			? FConstStructView(*ConstSharedStruct)
 			: FConstStructView();
@@ -567,27 +570,27 @@ public:
 	}
 
 	template<typename SharedFragmentType>
-	SharedFragmentType* GetSharedFragmentDataPtr(FMassEntityHandle Entity) const
+	SharedFragmentType* GetSharedFragmentDataPtr(FMassEntityHandle EntityHandle) const
 	{
 		static_assert(TIsDerivedFrom<SharedFragmentType, FMassSharedFragment>::IsDerived
 			, "Given struct doesn't represent a valid shared fragment type. Make sure to inherit from FMassSharedFragment or one of its child-types.");
-		const FSharedStruct* FragmentPtr = InternalGetSharedFragmentPtr(Entity, SharedFragmentType::StaticStruct());
+		const FSharedStruct* FragmentPtr = InternalGetSharedFragmentPtr(EntityHandle, SharedFragmentType::StaticStruct());
 		return (SharedFragmentType*)(FragmentPtr ? FragmentPtr->GetMemory() : nullptr);
 	}
 
 	template<typename SharedFragmentType>
-	SharedFragmentType& GetSharedFragmentDataChecked(FMassEntityHandle Entity) const
+	SharedFragmentType& GetSharedFragmentDataChecked(FMassEntityHandle EntityHandle) const
 	{
-		SharedFragmentType* TypePtr = GetSharedFragmentDataPtr<SharedFragmentType>(Entity);
+		SharedFragmentType* TypePtr = GetSharedFragmentDataPtr<SharedFragmentType>(EntityHandle);
 		check(TypePtr);
 		return *TypePtr;
 	}
 
-	FConstStructView GetSharedFragmentDataStruct(FMassEntityHandle Entity, const UScriptStruct* SharedFragmentType) const
+	FConstStructView GetSharedFragmentDataStruct(FMassEntityHandle EntityHandle, const UScriptStruct* SharedFragmentType) const
 	{
 		checkf((SharedFragmentType != nullptr) && SharedFragmentType->IsChildOf(FMassSharedFragment::StaticStruct())
 			, TEXT("GetSharedFragmentDataStruct called with an invalid fragment type '%s'"), *GetPathNameSafe(SharedFragmentType));
-		const FSharedStruct* FragmentPtr = InternalGetSharedFragmentPtr(Entity, SharedFragmentType);
+		const FSharedStruct* FragmentPtr = InternalGetSharedFragmentPtr(EntityHandle, SharedFragmentType);
 		return FragmentPtr
 			? FConstStructView(*FragmentPtr)
 			: FConstStructView();
@@ -1013,28 +1016,28 @@ protected:
 		TConstArrayView<const UScriptStruct*> FragmentsAndTagsList) const;
 
 private:
-	void InternalBuildEntity(FMassEntityHandle Entity, const FMassArchetypeHandle& ArchetypeHandle, const FMassArchetypeSharedFragmentValues& SharedFragmentValues);
-	void InternalReleaseEntity(FMassEntityHandle Entity);
+	void InternalBuildEntity(FMassEntityHandle EntityHandle, const FMassArchetypeHandle& ArchetypeHandle, const FMassArchetypeSharedFragmentValues& SharedFragmentValues);
+	void InternalReleaseEntity(FMassEntityHandle EntityHandle);
 
 	/** 
-	 *  Adds fragments in FragmentList to Entity. Only the unique fragments will be added.
+	 *  Adds fragments in FragmentList to the entity indicated by EntityHandle. Only the unique fragments will be added.
 	 *  @return Bitset for the added fragments (might be empty or a subset of `InFragments` depending on the current archetype fragments)
 	 */
-	FMassFragmentBitSet InternalAddFragmentListToEntityChecked(FMassEntityHandle Entity, const FMassFragmentBitSet& InFragments);
+	FMassFragmentBitSet InternalAddFragmentListToEntityChecked(FMassEntityHandle EntityHandle, const FMassFragmentBitSet& InFragments);
 
 	/** 
 	 *  Similar to InternalAddFragmentListToEntity but expects NewFragmentList not overlapping with current entity's
 	 *  fragment list. It's callers responsibility to ensure that's true. Failing this will cause a `check` fail.
 	 */
-	void InternalAddFragmentListToEntity(FMassEntityHandle Entity, const FMassFragmentBitSet& InFragments);
+	void InternalAddFragmentListToEntity(FMassEntityHandle EntityHandle, const FMassFragmentBitSet& InFragments);
 	/** Note that it's the caller's responsibility to ensure `FragmentType` is a kind of FMassFragment */
-	void* InternalGetFragmentDataChecked(FMassEntityHandle Entity, const UScriptStruct* FragmentType) const;
+	void* InternalGetFragmentDataChecked(FMassEntityHandle EntityHandle, const UScriptStruct* FragmentType) const;
 	/** Note that it's the caller's responsibility to ensure `FragmentType` is a kind of FMassFragment */
-	void* InternalGetFragmentDataPtr(FMassEntityHandle Entity, const UScriptStruct* FragmentType) const;
+	void* InternalGetFragmentDataPtr(FMassEntityHandle EntityHandle, const UScriptStruct* FragmentType) const;
 	/** Note that it's the caller's responsibility to ensure `ConstSharedFragmentType` is a kind of FMassSharedFragment */
-	const FConstSharedStruct* InternalGetConstSharedFragmentPtr(FMassEntityHandle Entity, const UScriptStruct* ConstSharedFragmentType) const;
+	const FConstSharedStruct* InternalGetConstSharedFragmentPtr(FMassEntityHandle EntityHandle, const UScriptStruct* ConstSharedFragmentType) const;
 	/** Note that it's the caller's responsibility to ensure `SharedFragmentType` is a kind of FMassSharedFragment */
-	const FSharedStruct* InternalGetSharedFragmentPtr(FMassEntityHandle Entity, const UScriptStruct* SharedFragmentType) const;
+	const FSharedStruct* InternalGetSharedFragmentPtr(FMassEntityHandle EntityHandle, const UScriptStruct* SharedFragmentType) const;
 
 	TSharedRef<FEntityCreationContext> InternalBatchCreateReservedEntities(const FMassArchetypeHandle& ArchetypeHandle,
 		const FMassArchetypeSharedFragmentValues& SharedFragmentValues, TConstArrayView<FMassEntityHandle> ReservedEntities);
@@ -1244,7 +1247,7 @@ protected:
 	FMassArchetypeHandle InternalCreateSimilarArchetype(const FMassArchetypeData& SourceArchetypeRef, FMassArchetypeCompositionDescriptor&& NewComposition);
 public:
 	UE_DEPRECATED(5.6, "SetEntityFragmentsValues is deprecated. Use SetEntityFragmentValues instead (note the slight change in name).")
-	void SetEntityFragmentsValues(FMassEntityHandle Entity, TArrayView<const FInstancedStruct> FragmentInstanceList);
+	void SetEntityFragmentsValues(FMassEntityHandle EntityHandle, TArrayView<const FInstancedStruct> FragmentInstanceList);
 
 	UE_DEPRECATED(5.6, "Static BatchSetEntityFragmentsValues is deprecated. Use EntityManager's member function BatchSetEntityFragmentValues (note the slight change in name).")
 	static void BatchSetEntityFragmentsValues(const FMassArchetypeEntityCollection& SparseEntities, TArrayView<const FInstancedStruct> FragmentInstanceList);
