Reduced the cost of setting the debug execution description in Mass.

This change also introduces the ability to selectively turn debug features on/off through a central flag on a per entity manager basis.

diff --git a/Engine/Source/Runtime/MassEntity/Private/MassEntityManager.cpp b/Engine/Source/Runtime/MassEntity/Private/MassEntityManager.cpp
--- a/Engine/Source/Runtime/MassEntity/Private/MassEntityManager.cpp
+++ b/Engine/Source/Runtime/MassEntity/Private/MassEntityManager.cpp
@@ -2237,6 +2237,21 @@ const FString& FMassEntityManager::DebugGetName() const
 	return DebugName;
 }
 
+void FMassEntityManager::DebugEnableDebugFeature(EDebugFeatures Features)
+{
+	EnumAddFlags(EnabledDebugFeatures, Features);
+}
+
+void FMassEntityManager::DebugDisableDebugFeature(EDebugFeatures Features)
+{
+	EnumRemoveFlags(EnabledDebugFeatures, Features);
+}
+
+bool FMassEntityManager::DebugHasAllDebugFeatures(EDebugFeatures Features) const
+{
+	return EnumHasAllFlags(EnabledDebugFeatures, Features);
+}
+
 FMassRequirementAccessDetector& FMassEntityManager::GetRequirementAccessDetector()
 {
 	return RequirementAccessDetector;
diff --git a/Engine/Source/Runtime/MassEntity/Private/MassProcessor.cpp b/Engine/Source/Runtime/MassEntity/Private/MassProcessor.cpp
--- a/Engine/Source/Runtime/MassEntity/Private/MassProcessor.cpp
+++ b/Engine/Source/Runtime/MassEntity/Private/MassProcessor.cpp
@@ -235,7 +235,17 @@ void UMassProcessor::CallExecute(FMassEntityManager& EntityManager, FMassExecuti
 		// Not using a more specific scope by default (i.e. LLM_SCOPE_BYNAME(*StatId)) since LLM is more strict regarding the provided string (no spaces or '_')
 
 #if WITH_MASSENTITY_DEBUG
-		Context.DebugSetExecutionDesc(FString::Printf(TEXT("%s (%s)"), *GetProcessorName(), EntityManager.GetWorld() ? *ToString(EntityManager.GetWorld()->GetNetMode()) : TEXT("No World")));
+		if (EntityManager.DebugHasAllDebugFeatures(FMassEntityManager::EDebugFeatures::TraceProcessors))
+		{
+			FString DebugDescription = *GetProcessorName();
+			FString NetMode = EntityManager.GetWorld() ? ToString(EntityManager.GetWorld()->GetNetMode()) : TEXT("None");
+			//                       DebugDescription       " ("  NetMode        ")"
+			DebugDescription.Reserve(DebugDescription.Len() + 2 + NetMode.Len() + 1);
+			DebugDescription.Append(TEXT(" ("));
+			DebugDescription.Append(*NetMode);
+			DebugDescription.Append(TEXT(")"));
+			Context.DebugSetExecutionDesc(MoveTemp(DebugDescription));
+		}
 #endif
 		// CacheSubsystemRequirements will return true only if all requirements declared with ProcessorRequirements are met
 		// meaning if it fails there's no point in calling Execute.
diff --git a/Engine/Source/Runtime/MassEntity/Public/MassEntityManager.h b/Engine/Source/Runtime/MassEntity/Public/MassEntityManager.h
--- a/Engine/Source/Runtime/MassEntity/Public/MassEntityManager.h
+++ b/Engine/Source/Runtime/MassEntity/Public/MassEntityManager.h
@@ -18,6 +18,10 @@
 #include "MassEntityManagerStorage.h"
 #include "MassArchetypeGroup.h"
 
+#if WITH_MASSENTITY_DEBUG
+#include "Misc/EnumClassFlags.h"
+#endif
+
 
 struct FInstancedStruct;
 struct FMassEntityQuery;
@@ -989,6 +993,17 @@ PRAGMA_ENABLE_DEPRECATION_WARNINGS
 	FMassEntityHandle DebugGetEntityIndexHandle(const int32 EntityIndex) const;
 	const FString& DebugGetName() const;
 
+	enum class EDebugFeatures
+	{
+		None,
+		TraceProcessors = 1 << 0, // Used to track information about processors such as their name.
+		All = TraceProcessors
+	};
+
+	void DebugEnableDebugFeature(EDebugFeatures Features);
+	void DebugDisableDebugFeature(EDebugFeatures Features);
+	bool DebugHasAllDebugFeatures(EDebugFeatures Features) const;
+
 	FMassRequirementAccessDetector& GetRequirementAccessDetector();
 
 	// For use by the friend MassDebugger
@@ -1228,6 +1243,7 @@ private:
 #if WITH_MASSENTITY_DEBUG
 	FMassRequirementAccessDetector RequirementAccessDetector;
 	FString DebugName;
+	EDebugFeatures EnabledDebugFeatures = EDebugFeatures::All;
 #endif // WITH_MASSENTITY_DEBUG
 
 	TWeakObjectPtr<UObject> Owner;
@@ -1271,6 +1287,11 @@ public:
 //-----------------------------------------------------------------------------
 // INLINE
 //-----------------------------------------------------------------------------
+
+#if WITH_MASSENTITY_DEBUG
+ENUM_CLASS_FLAGS(FMassEntityManager::EDebugFeatures);
+#endif
+
 template<UE::Mass::CSharedFragment T>
 void FMassEntityManager::ForEachSharedFragment(TFunctionRef< void(T& /*SharedFragment*/) > ExecuteFunction)
 {
diff --git a/Engine/Source/Runtime/MassEntity/Public/MassExecutionContext.h b/Engine/Source/Runtime/MassEntity/Public/MassExecutionContext.h
--- a/Engine/Source/Runtime/MassEntity/Public/MassExecutionContext.h
+++ b/Engine/Source/Runtime/MassEntity/Public/MassExecutionContext.h
@@ -125,6 +125,7 @@ public:
 #if WITH_MASSENTITY_DEBUG
 	const FString& DebugGetExecutionDesc() const { return DebugExecutionDescription; }
 	void DebugSetExecutionDesc(const FString& Description) { DebugExecutionDescription = Description; }
+	void DebugSetExecutionDesc(FString&& Description) { DebugExecutionDescription = MoveTemp(Description); }
 #endif
 
 	void PushQuery(FMassEntityQuery& InQuery);
