[Mass] Made FMassProcessingPhase.bRunInParallelMode `true` by default. Also switched application of changes to `mass.FullyParallel` from taking place in FMassProcessingPhaseManager::OnPhaseEnd to FMassProcessingPhaseManager::OnPhaseStart

diff --git a/Engine/Source/Runtime/MassEntity/Private/MassProcessingPhaseManager.cpp b/Engine/Source/Runtime/MassEntity/Private/MassProcessingPhaseManager.cpp
--- a/Engine/Source/Runtime/MassEntity/Private/MassProcessingPhaseManager.cpp
+++ b/Engine/Source/Runtime/MassEntity/Private/MassProcessingPhaseManager.cpp
@@ -451,11 +451,24 @@ void FMassProcessingPhaseManager::Stop()
 	}
 }
 
-void FMassProcessingPhaseManager::OnPhaseStart(const FMassProcessingPhase& Phase)
+void FMassProcessingPhaseManager::OnPhaseStart(FMassProcessingPhase& Phase)
 {
 	ensure(CurrentPhase == EMassProcessingPhase::MAX);
 	CurrentPhase = Phase.Phase;
 
+	// switch between parallel and single-thread versions only after a given batch of processing has been wrapped up	
+	if (Phase.IsConfiguredForParallelMode() != UE::Mass::Tweakables::bFullyParallel)
+	{
+		if (UE::Mass::Tweakables::bFullyParallel)
+		{
+			Phase.ConfigureForParallelMode();
+		}
+		else
+		{
+			Phase.ConfigureForSingleThreadMode();
+		}
+	}
+
 	const int32 PhaseAsInt = int32(Phase.Phase);
 	if (Owner.IsValid()
 		&& ensure(Phase.Phase != EMassProcessingPhase::MAX)
@@ -501,19 +514,6 @@ void FMassProcessingPhaseManager::OnPhaseEnd(FMassProcessingPhase& Phase)
 {
 	ensure(CurrentPhase == Phase.Phase);
 	CurrentPhase = EMassProcessingPhase::MAX;
-
-	// switch between parallel and single-thread versions only after a given batch of processing has been wrapped up	
-	if (Phase.IsConfiguredForParallelMode() != UE::Mass::Tweakables::bFullyParallel)
-	{
-		if (UE::Mass::Tweakables::bFullyParallel)
-		{
-			Phase.ConfigureForParallelMode();
-		}
-		else
-		{
-			Phase.ConfigureForSingleThreadMode();
-		}
-	}
 }
 
 FString FMassProcessingPhaseManager::GetName() const
diff --git a/Engine/Source/Runtime/MassEntity/Public/MassProcessingPhaseManager.h b/Engine/Source/Runtime/MassEntity/Public/MassProcessingPhaseManager.h
--- a/Engine/Source/Runtime/MassEntity/Public/MassProcessingPhaseManager.h
+++ b/Engine/Source/Runtime/MassEntity/Public/MassProcessingPhaseManager.h
@@ -94,7 +94,7 @@ protected:
 private:
 	FMassProcessingPhaseManager* PhaseManager = nullptr;
 	std::atomic<bool> bIsDuringMassProcessing = false;
-	bool bRunInParallelMode = false;
+	bool bRunInParallelMode = true;
 	uint8 SupportedTickTypes = 0;
 };
 
@@ -219,7 +219,7 @@ protected:
 	 *  Called by the given Phase at the very start of its execution function (the FMassProcessingPhase::ExecuteTick),
 	 *  even before the FMassProcessingPhase.OnPhaseStart broadcast delegate
 	 */
-	void OnPhaseStart(const FMassProcessingPhase& Phase);
+	void OnPhaseStart(FMassProcessingPhase& Phase);
 
 	/**
 	 *  Called by the given Phase at the very end of its execution function (the FMassProcessingPhase::ExecuteTick),
