[Mass] Per-entity/processor breakpoint support

diff --git a/Engine/Source/Editor/MassEntityDebugger/Private/MassDebuggerModel.cpp b/Engine/Source/Editor/MassEntityDebugger/Private/MassDebuggerModel.cpp
--- a/Engine/Source/Editor/MassEntityDebugger/Private/MassDebuggerModel.cpp
+++ b/Engine/Source/Editor/MassEntityDebugger/Private/MassDebuggerModel.cpp
@@ -143,19 +143,20 @@ FMassDebuggerProcessorData::FMassDebuggerProcessorData(const UMassProcessor& InP
 #endif // WITH_MASSENTITY_DEBUG
 }
 
-FMassDebuggerProcessorData::FMassDebuggerProcessorData(const FMassEntityManager& EntityManager, const UMassProcessor& InProcessor
+FMassDebuggerProcessorData::FMassDebuggerProcessorData(const FMassEntityManager& InEntityManager, const UMassProcessor& InProcessor
 	, const TMap<FMassArchetypeHandle, TSharedPtr<FMassDebuggerArchetypeData>>& InTransientArchetypesMap)
 {
 	SetProcessor(InProcessor);
 #if WITH_MASSENTITY_DEBUG
+	EntityManager = &InEntityManager;
 	// yeah, it's ugly. But it's debugging code, so... 
 	UMassProcessor& MutableProcessor = const_cast<UMassProcessor&>(InProcessor);
-	TConstArrayView<FMassEntityQuery*> ProcessorQueries = FMassDebugger::GetUpToDateProcessorQueries(EntityManager, MutableProcessor);
+	TConstArrayView<FMassEntityQuery*> ProcessorQueries = FMassDebugger::GetUpToDateProcessorQueries(InEntityManager, MutableProcessor);
 
 	ProcessorRequirements = MakeShareable(new FMassDebuggerQueryData(InProcessor.GetProcessorRequirements(), LOCTEXT("MassProcessorRequirementsLabel", "Processor Requirements")));
 
 	const FMassEntityHandle SelectedEntityHandle = UE::Mass::Debug::bTestSelectedEntityAgainstProcessorQueries
-		? FMassDebugger::GetSelectedEntity(EntityManager)
+		? FMassDebugger::GetSelectedEntity(InEntityManager)
 		: FMassEntityHandle();
 	FStringOutputDevice SelectedEntityFailureJustificationLog;
 	SelectedEntityFailureJustificationLog.SetAutoEmitLineTerminator(true);
@@ -171,7 +172,7 @@ FMassDebuggerProcessorData::FMassDebuggerProcessorData(const FMassEntityManager&
 
 		if (SelectedEntityHandle.IsValid())
 		{
-			const FMassArchetypeHandle ArchetypeHandle = EntityManager.GetArchetypeForEntity(SelectedEntityHandle);
+			const FMassArchetypeHandle ArchetypeHandle = InEntityManager.GetArchetypeForEntity(SelectedEntityHandle);
 			if (ArchetypeHandle.IsValid() && Query->GetArchetypes().Contains(ArchetypeHandle) == false)
 			{
 				if (FMassArchetypeHelper::DoesArchetypeMatchRequirements(FMassArchetypeHelper::ArchetypeDataFromHandleChecked(ArchetypeHandle), *Query
@@ -201,6 +202,7 @@ void FMassDebuggerProcessorData::SetProcessor(const UMassProcessor& InProcessor)
 	Name = InProcessor.GetProcessorName();
 	UE::Mass::Debugger::Private::MakeDisplayName(Name, Label);
 
+	Processor = &InProcessor;
 	ProcessorHash = UE::Mass::Debugger::Private::CalcProcessorHash(InProcessor);
 	bIsActive = InProcessor.IsActive();
 	if (bIsActive == false)
@@ -235,6 +237,7 @@ FMassDebuggerArchetypeData::FMassDebuggerArchetypeData(const FMassArchetypeHandl
 	HashLabel = FText::FromString(FullHashAsString);
 
 	FMassDebugger::GetArchetypeEntityStats(ArchetypeHandle, ArchetypeStats);
+	EntityHandles = FMassDebugger::GetEntitiesOfArchetype(ArchetypeHandle);
 
 	const TConstArrayView<FName> DebugNames = FMassDebugger::GetArchetypeDebugNames(ArchetypeHandle);
 
diff --git a/Engine/Source/Editor/MassEntityDebugger/Private/MassDebuggerModel.h b/Engine/Source/Editor/MassEntityDebugger/Private/MassDebuggerModel.h
--- a/Engine/Source/Editor/MassEntityDebugger/Private/MassDebuggerModel.h
+++ b/Engine/Source/Editor/MassEntityDebugger/Private/MassDebuggerModel.h
@@ -2,6 +2,7 @@
 #pragma once
 
 #include "MassArchetypeTypes.h"
+#include "MassEntityHandle.h"
 #include "MassEntityQuery.h"
 #include "MassProcessorDependencySolver.h"
 #include "MassProcessingTypes.h"
@@ -88,6 +89,8 @@ struct FMassDebuggerArchetypeData
 	/** Primary debug name, used for grouping derived archetypes. */
 	FString PrimaryDebugName;
 
+	TArray<FMassEntityHandle> EntityHandles;
+
 	/** True if the archetype is selected. */
 	bool bIsSelected = false;
 
@@ -97,7 +100,7 @@ struct FMassDebuggerArchetypeData
 struct FMassDebuggerProcessorData
 {
 	FMassDebuggerProcessorData(const UMassProcessor& InProcessor);
-	FMassDebuggerProcessorData(const FMassEntityManager& EntitySubsystem, const UMassProcessor& InProcessor, const TMap<FMassArchetypeHandle, TSharedPtr<FMassDebuggerArchetypeData>>& InTransientArchetypesMap);
+	FMassDebuggerProcessorData(const FMassEntityManager& InEntityManager, const UMassProcessor& InProcessor, const TMap<FMassArchetypeHandle, TSharedPtr<FMassDebuggerArchetypeData>>& InTransientArchetypesMap);
 
 private:
 	void SetProcessor(const UMassProcessor& InProcessor);
@@ -107,6 +110,8 @@ public:
 	FString Label;
 	uint32 ProcessorHash = 0;
 	bool bIsActive = true;
+	const FMassEntityManager* EntityManager = nullptr;
+	TWeakObjectPtr<const UMassProcessor> Processor;
 	EMassDebuggerProcessorSelection Selection = EMassDebuggerProcessorSelection::None;
 	TSharedPtr<FMassDebuggerQueryData> ProcessorRequirements;
 	TArray<TSharedPtr<FMassDebuggerQueryData>> Queries;
diff --git a/Engine/Source/Runtime/CoreUObject/Private/StructUtils/StructTypeBitSet.cpp b/Engine/Source/Runtime/CoreUObject/Private/StructUtils/StructTypeBitSet.cpp
--- a/Engine/Source/Runtime/CoreUObject/Private/StructUtils/StructTypeBitSet.cpp
+++ b/Engine/Source/Runtime/CoreUObject/Private/StructUtils/StructTypeBitSet.cpp
@@ -291,3 +291,17 @@ void FStructTracker::Serialize(FArchive& Ar, FStructTypeBitSet::FBitSetContainer
 		}
 	}
 }
+
+#if WITH_STRUCTUTILS_DEBUG
+const UStruct* FStructTracker::DebugFindTypeByPartialName(const FString& PartialFragmentName) const
+{
+	for (int32 TypeIndex = 0; TypeIndex < DebugStructTypeNamesList.Num(); ++TypeIndex)
+	{
+		if (DebugStructTypeNamesList[TypeIndex].ToString().Contains(PartialFragmentName))
+		{
+			return StructTypesList[TypeIndex].Get();
+		}
+	}
+	return nullptr;
+}
+#endif // WITH_STRUCTUTILS_DEBUG
diff --git a/Engine/Source/Runtime/CoreUObject/Public/StructUtils/StructTypeBitSet.h b/Engine/Source/Runtime/CoreUObject/Public/StructUtils/StructTypeBitSet.h
--- a/Engine/Source/Runtime/CoreUObject/Public/StructUtils/StructTypeBitSet.h
+++ b/Engine/Source/Runtime/CoreUObject/Public/StructUtils/StructTypeBitSet.h
@@ -359,6 +359,8 @@ struct COREUOBJECT_API FStructTracker
 		DebugStructTypeNamesList.Reset();
 	}
 
+	[[nodiscard]] const UStruct* DebugFindTypeByPartialName(const FString& PartialFragmentName) const;
+
 protected:
 	mutable TArray<FName, TInlineAllocator<64>> DebugStructTypeNamesList;
 #endif // WITH_STRUCTUTILS_DEBUG
@@ -1144,6 +1146,11 @@ public:
 	{
 		TStructTrackerWrapper::StructTracker.DebugResetStructTypeMappingInfo();
 	}
+
+	[[nodiscard]] static const TUStructType* DebugFindTypeByPartialName(const FString& PartialFragmentName)
+	{
+		return reinterpret_cast<const TUStructType*>(TStructTrackerWrapper::StructTracker.DebugFindTypeByPartialName(PartialFragmentName));
+	}
 #endif // WITH_STRUCTUTILS_DEBUG
 
 public:
diff --git a/Engine/Source/Runtime/MassEntity/Private/MassDebugger.cpp b/Engine/Source/Runtime/MassEntity/Private/MassDebugger.cpp
--- a/Engine/Source/Runtime/MassEntity/Private/MassDebugger.cpp
+++ b/Engine/Source/Runtime/MassEntity/Private/MassDebugger.cpp
@@ -332,39 +332,132 @@ namespace UE::Mass::Debug
 	static FAutoConsoleCommandWithWorldAndArgs SetDebugEntity(
 		TEXT("mass.debug.DebugEntity"),
 		TEXT("ID of a Mass entity that we want to debug.")
+		TEXT("Note that this call results in the same behavior as if the entity was picked via the Mass GameplayDebugger's category.")
 		TEXT("Usage: \"mass.debug.DebugEntity <Entity>\""),
 		FConsoleCommandWithWorldAndArgsDelegate::CreateLambda([](const TArray<FString>& Args, UWorld* World)
 		{
-			if (Args.Num() != 1)
+			if (!World)
 			{
-				UE_LOG(LogConsoleResponse, Display, TEXT("Error: Expecting 1 parameter"));
+				UE_LOG(LogConsoleResponse, Display, TEXT("Error: invalid world"));
 				return;
 			}
 
 			int32 ID = INDEX_NONE;
-			if (!LexTryParseString<int32>(ID, *Args[0]))
+			if (Args.Num() > 0)
+			{
+				LexTryParseString<int32>(ID, *Args[0]);
+			}
+
+			SetDebugEntityRange(ID, ID);
+
+			FMassEntityManager& EntityManager = UE::Mass::Utils::GetEntityManagerChecked(*World);
+			FMassEntityHandle EntityToDebug = EntityManager.DebugGetEntityIndexHandle(ID);
+			if (!EntityToDebug.IsSet() && ID != INDEX_NONE)
 			{
-				UE_LOG(LogConsoleResponse, Display, TEXT("Error: parameter must be an integer"));
+				UE_LOG(LogConsoleResponse, Display, TEXT("Cannot find entity for this index, clearing current selection"));
 				return;
 			}
 
+			FMassDebugger::SelectEntity(EntityManager, EntityToDebug);
+		}
+	));
+
+	const UScriptStruct* FindElementTypeByName(const FString& PartialFragmentName)
+	{
+		const UScriptStruct* Result = nullptr;
+#if WITH_STRUCTUTILS_DEBUG
+		Result = FMassFragmentBitSet::DebugFindTypeByPartialName(PartialFragmentName);
+		if (Result == nullptr)
+		{
+			Result = FMassSharedFragmentBitSet::DebugFindTypeByPartialName(PartialFragmentName);
+		}
+		if (Result == nullptr)
+		{
+			Result = FMassConstSharedFragmentBitSet::DebugFindTypeByPartialName(PartialFragmentName);
+		}
+#endif // WITH_STRUCTUTILS_DEBUG
+		return Result;
+	}
+
+	static FAutoConsoleCommandWithWorldAndArgs SetFragmentBreakpoint(
+		TEXT("mass.debug.SetFragmentBreakpoint"),
+		TEXT("Enables fragment write break-point on an arbitrary number of fragment types, on the selected entity (see `mass.debug.DebugEntity`).")
+		TEXT("Usage: `mass.debug.SetFragmentBreakpoint <FragmentTypeName> <FragmentTypeName2> <FragmentTypeName3> <...>`"),
+		FConsoleCommandWithWorldAndArgsDelegate::CreateLambda([](const TArray<FString>& Args, UWorld* World)
+		{
 			if (!World)
 			{
 				UE_LOG(LogConsoleResponse, Display, TEXT("Error: invalid world"));
 				return;
 			}
 
-			SetDebugEntityRange(ID, ID);
+			if (Args.Num() == 0)
+			{
+				UE_LOG(LogConsoleResponse, Display, TEXT("No fragment types indicated"));
+			}
+			else
+			{
+				FMassEntityManager& EntityManager = Utils::GetEntityManagerChecked(*World);
+				FMassEntityHandle SelectedEntity = FMassDebugger::GetSelectedEntity(EntityManager);
+				if (SelectedEntity.IsValid())
+				{
+					for (const FString& PartialFragmentName : Args)
+					{
+						if (const UScriptStruct* FragmentType = FindElementTypeByName(PartialFragmentName))
+						{
+							FMassDebugger::SetFragmentWriteBreak(EntityManager, FragmentType, SelectedEntity);
+						}
+						else
+						{
+							UE_LOG(LogConsoleResponse, Display, TEXT("Warning: Unable to find element type %s"), *PartialFragmentName);
+						}
+					}
+				}
+				else
+				{
+					UE_LOG(LogConsoleResponse, Display, TEXT("Warning: No entity selected, no break points set"));
+				}
+			}
+		}
+	));
 
-			FMassEntityManager& EntityManager = UE::Mass::Utils::GetEntityManagerChecked(*World);
-			FMassEntityHandle EntityToDebug = EntityManager.DebugGetEntityIndexHandle(ID);
-			if (!EntityToDebug.IsSet())
+	static FAutoConsoleCommandWithWorldAndArgs ClearFragmentBreakpoint(
+		TEXT("mass.debug.ClearFragmentBreakpoint"),
+		TEXT("Clears fragment write break-point on an arbitrary number of fragment types, on the selected entity (see `mass.debug.DebugEntity`).")
+		TEXT("If no entity is currently selected then the call will clear the type breakpoints on all entities.")
+		TEXT("Usage: `mass.debug.ClearFragmentBreakpoint <FragmentTypeName> <FragmentTypeName2> <FragmentTypeName3> <...>`"),
+		FConsoleCommandWithWorldAndArgsDelegate::CreateLambda([](const TArray<FString>& Args, UWorld* World)
+		{
+			if (!World)
 			{
-				UE_LOG(LogConsoleResponse, Display, TEXT("Error: cannot find entity for this index"));
+				UE_LOG(LogConsoleResponse, Display, TEXT("Error: invalid world"));
 				return;
 			}
 
-			FMassDebugger::SelectEntity(EntityManager, EntityToDebug);
+			if (Args.Num() == 0)
+			{
+				UE_LOG(LogConsoleResponse, Display, TEXT("No fragment types indicated"));
+			}
+			else
+			{
+				FMassEntityManager& EntityManager = Utils::GetEntityManagerChecked(*World);
+				FMassEntityHandle SelectedEntity = FMassDebugger::GetSelectedEntity(EntityManager);
+				const bool bEntityValid = SelectedEntity.IsValid();
+
+				for (const FString& PartialFragmentName : Args)
+				{
+					if (const UScriptStruct* FragmentType = FindElementTypeByName(PartialFragmentName))
+					{
+						bEntityValid
+							? FMassDebugger::ClearFragmentWriteBreak(EntityManager, FragmentType, SelectedEntity)
+							: FMassDebugger::ClearAllFragmentWriteBreaks(EntityManager, FragmentType);
+					}
+					else
+					{
+						UE_LOG(LogConsoleResponse, Display, TEXT("Warning: Unable to find element type %s"), *PartialFragmentName);
+					}
+				}				
+			}
 		}
 	));
 
@@ -381,6 +474,10 @@ FMassDebugger::FOnEnvironmentEvent FMassDebugger::OnProcessorProviderRegistered;
 FMassDebugger::FOnDebugEvent FMassDebugger::OnDebugEvent;
 TArray<FMassDebugger::FEnvironment> FMassDebugger::ActiveEnvironments;
 UE::FSpinLock FMassDebugger::EntityManagerRegistrationLock;
+bool FMassDebugger::HasBreakPoint = false;
+FName FMassDebugger::SelectedFragment;
+FOnFragmentSelectionChanged FMassDebugger::OnFragmentSelectionChanged;
+TMap<FName, const UScriptStruct*> FMassDebugger::FragmentsByName;
 
 TConstArrayView<FMassEntityQuery*> FMassDebugger::GetProcessorQueries(const UMassProcessor& Processor)
 {
@@ -456,6 +553,19 @@ const TConstArrayView<FName> FMassDebugger::GetArchetypeDebugNames(const FMassAr
 	return ArchetypeData.GetDebugNames();
 }
 
+TArray<FMassEntityHandle> FMassDebugger::GetEntitiesOfArchetype(const FMassArchetypeHandle& ArchetypeHandle)
+{
+	TArray<FMassEntityHandle> EntitiesOfArchetype;
+	FMassArchetypeData& ArchetypeData = FMassArchetypeHelper::ArchetypeDataFromHandleChecked(ArchetypeHandle);
+	EntitiesOfArchetype.Reserve(ArchetypeData.GetNumEntities());
+	for (FMassArchetypeChunk& Chunk : ArchetypeData.Chunks)
+	{
+		TArrayView<FMassEntityHandle> EntityListView = TArrayView<FMassEntityHandle>(&Chunk.GetEntityArrayElementRef(ArchetypeData.EntityListOffsetWithinChunk, 0), Chunk.GetNumInstances());
+		EntitiesOfArchetype.Append(EntityListView);
+	}
+	return EntitiesOfArchetype;
+}
+
 TConstArrayView<struct UMassCompositeProcessor::FDependencyNode> FMassDebugger::GetProcessingGraph(const UMassCompositeProcessor& GraphOwner)
 {
 	return GraphOwner.FlatProcessingGraph;
@@ -662,9 +772,9 @@ void FMassDebugger::SelectEntity(const FMassEntityManager& EntityManager, const
 	UE::Mass::Debug::SetDebugEntityRange(EntityHandle.Index, EntityHandle.Index);
 
 	const int32 Index = ActiveEnvironments.IndexOfByPredicate([WeakManager = EntityManager.AsWeak()](const FEnvironment& Element)
-		{
-			return Element.EntityManager == WeakManager;
-		});
+	{
+		return Element.EntityManager == WeakManager;
+	});
 	if (ensure(Index != INDEX_NONE))
 	{
 		ActiveEnvironments[Index].SelectedEntity = EntityHandle;
@@ -676,13 +786,35 @@ void FMassDebugger::SelectEntity(const FMassEntityManager& EntityManager, const
 FMassEntityHandle FMassDebugger::GetSelectedEntity(const FMassEntityManager& EntityManager)
 {
 	const int32 Index = ActiveEnvironments.IndexOfByPredicate([WeakManager = EntityManager.AsWeak()](const FEnvironment& Element)
-		{
-			return Element.EntityManager == WeakManager;
-		});
+	{
+		return Element.EntityManager == WeakManager;
+	});
 
 	return Index != INDEX_NONE ? ActiveEnvironments[Index].SelectedEntity : FMassEntityHandle();
 }
 
+void FMassDebugger::HighlightEntity(const FMassEntityManager& EntityManager, const FMassEntityHandle EntityHandle)
+{
+	const int32 Index = ActiveEnvironments.IndexOfByPredicate([WeakManager = EntityManager.AsWeak()](const FEnvironment& Element)
+	{
+		return Element.EntityManager == WeakManager;
+	});
+	if (ensure(Index != INDEX_NONE))
+	{
+		ActiveEnvironments[Index].HighlightedEntity = EntityHandle;
+	}
+}
+
+FMassEntityHandle FMassDebugger::GetHighlightedEntity(const FMassEntityManager& EntityManager)
+{
+	const int32 Index = ActiveEnvironments.IndexOfByPredicate([WeakManager = EntityManager.AsWeak()](const FEnvironment& Element)
+	{
+		return Element.EntityManager == WeakManager;
+	});
+
+	return Index != INDEX_NONE ? ActiveEnvironments[Index].HighlightedEntity : FMassEntityHandle();
+}
+
 int32 FMassDebugger::RegisterEntityManager(FMassEntityManager& EntityManager)
 {
 	int32 NewEnvironmentIndex = INDEX_NONE;
@@ -746,4 +878,298 @@ bool FMassDebugger::DoesArchetypeMatchRequirements(const FMassArchetypeHandle& A
 	return false;
 }
 
+bool FMassDebugger::ShouldProcessorBreak(const FMassEntityManager& EntityManager, const UMassProcessor* Processor, FMassEntityHandle Entity)
+{
+	if (LIKELY(!HasBreakPoint))
+	{
+		return false;
+	}
+
+	const int32 Index = ActiveEnvironments.IndexOfByPredicate([WeakManager = EntityManager.AsWeak()](const FEnvironment& Element)
+	{
+		return Element.EntityManager == WeakManager;
+	});
+
+	if (LIKELY(!ActiveEnvironments[Index].HasBreakPoint))
+	{
+		return false;
+	}
+
+	const FMassEntityHandle* Found = ActiveEnvironments[Index].ProcessorBreakPoints.FindPair(Processor, Entity);
+	
+	return Found != nullptr;
+}
+
+bool FMassDebugger::HasAnyProcessorBreakpoints(const FMassEntityManager& EntityManager, const UMassProcessor* Processor)
+{
+	if (LIKELY(!HasBreakPoint))
+	{
+		return false;
+	}
+
+	const int32 Index = ActiveEnvironments.IndexOfByPredicate([WeakManager = EntityManager.AsWeak()](const FEnvironment& Element)
+	{
+		return Element.EntityManager == WeakManager;
+	});
+
+	if (LIKELY(!ActiveEnvironments[Index].HasBreakPoint))
+	{
+		return false;
+	}
+
+	return ActiveEnvironments[Index].ProcessorBreakPoints.Contains(Processor);
+}
+
+bool FMassDebugger::ShouldFragmentWriteBreak(const FMassEntityManager& EntityManager, const UScriptStruct* FragmentType, FMassEntityHandle Entity)
+{
+	if (LIKELY(!HasBreakPoint))
+	{
+		return false;
+	}
+
+	const int32 Index = ActiveEnvironments.IndexOfByPredicate([WeakManager = EntityManager.AsWeak()](const FEnvironment& Element)
+	{
+		return Element.EntityManager == WeakManager;
+	});
+
+	if (LIKELY(!ActiveEnvironments[Index].HasBreakPoint))
+	{
+		return false;
+	}
+
+	const FMassEntityHandle* Found = ActiveEnvironments[Index].FragmentWriteBreakPoints.FindPair(FragmentType, Entity);
+
+	return Found != nullptr;
+}
+
+bool FMassDebugger::HasAnyFragmentWriteBreakpoints(const FMassEntityManager& EntityManager, const UScriptStruct* FragmentType)
+{
+	if (LIKELY(!HasBreakPoint))
+	{
+		return false;
+	}
+
+	const int32 Index = ActiveEnvironments.IndexOfByPredicate([WeakManager = EntityManager.AsWeak()](const FEnvironment& Element)
+	{
+		return Element.EntityManager == WeakManager;
+	});
+
+	if (LIKELY(!ActiveEnvironments[Index].HasBreakPoint))
+	{
+		return false;
+	}
+
+	if (FragmentType == nullptr)
+	{
+		return ActiveEnvironments[Index].FragmentWriteBreakPoints.Num() > 0;
+	}
+
+	return ActiveEnvironments[Index].FragmentWriteBreakPoints.Contains(FragmentType);
+}
+
+void FMassDebugger::SetProcessorBreakPoint(const FMassEntityManager& EntityManager, const UMassProcessor* Processor, FMassEntityHandle Entity)
+{
+	const int32 Index = ActiveEnvironments.IndexOfByPredicate([WeakManager = EntityManager.AsWeak()](const FEnvironment& Element)
+	{
+		return Element.EntityManager == WeakManager;
+	});
+
+	ActiveEnvironments[Index].HasBreakPoint = true;
+	HasBreakPoint = true;
+
+	ActiveEnvironments[Index].ProcessorBreakPoints.Add(Processor, Entity);
+}
+
+void FMassDebugger::SetFragmentWriteBreak(const FMassEntityManager& EntityManager, const UScriptStruct* FragmentType, FMassEntityHandle Entity)
+{
+	const int32 Index = ActiveEnvironments.IndexOfByPredicate([WeakManager = EntityManager.AsWeak()](const FEnvironment& Element)
+	{
+		return Element.EntityManager == WeakManager;
+	});
+
+	ActiveEnvironments[Index].HasBreakPoint = true;
+	HasBreakPoint = true;
+
+	ActiveEnvironments[Index].FragmentWriteBreakPoints.Add(FragmentType, Entity);
+}
+
+void FMassDebugger::ClearProcessorBreakPoint(const FMassEntityManager& EntityManager, const UMassProcessor* Processor, FMassEntityHandle Entity)
+{
+	const int32 Index = ActiveEnvironments.IndexOfByPredicate([WeakManager = EntityManager.AsWeak()](const FEnvironment& Element)
+	{
+		return Element.EntityManager == WeakManager;
+	});
+
+	ActiveEnvironments[Index].ProcessorBreakPoints.Remove(Processor, Entity);
+	UpdateHasBreakpoint();
+}
+
+void FMassDebugger::ClearFragmentWriteBreak(const FMassEntityManager& EntityManager, const UScriptStruct* FragmentType, FMassEntityHandle Entity)
+{
+	const int32 Index = ActiveEnvironments.IndexOfByPredicate([WeakManager = EntityManager.AsWeak()](const FEnvironment& Element)
+	{
+		return Element.EntityManager == WeakManager;
+	});
+
+	ActiveEnvironments[Index].FragmentWriteBreakPoints.Remove(FragmentType, Entity);
+	UpdateHasBreakpoint();
+}
+
+void FMassDebugger::ClearProcessorBreakPoints(const FMassEntityManager& EntityManager, TNotNull<const UMassProcessor*> Processor)
+{
+	const int32 Index = ActiveEnvironments.IndexOfByPredicate([WeakManager = EntityManager.AsWeak()](const FEnvironment& Element)
+	{
+		return Element.EntityManager == WeakManager;
+	});
+
+	ActiveEnvironments[Index].ProcessorBreakPoints.Remove(Processor);
+	UpdateHasBreakpoint();
+}
+
+void FMassDebugger::ClearAllFragmentWriteBreaks(const FMassEntityManager& EntityManager, TNotNull<const UScriptStruct*> FragmentType)
+{
+	const int32 Index = ActiveEnvironments.IndexOfByPredicate([WeakManager = EntityManager.AsWeak()](const FEnvironment& Element)
+	{
+		return Element.EntityManager == WeakManager;
+	});
+
+	ActiveEnvironments[Index].FragmentWriteBreakPoints.Remove(FragmentType);
+	UpdateHasBreakpoint();	
+} 
+ 
+void FMassDebugger::BreakOnSelectedFragmentWrite()
+{
+	if (!SelectedFragment.IsNone())
+	{
+		for (FEnvironment& Environment : ActiveEnvironments)
+		{
+			if (Environment.EntityManager.Pin()->IsEntityValid(Environment.SelectedEntity))
+			{
+				SetFragmentWriteBreak(*Environment.EntityManager.Pin(), GetFragmentTypeFromName(SelectedFragment), Environment.SelectedEntity);
+			}
+		}
+	}
+}
+
+void FMassDebugger::ClearAllBreakpoints()
+{
+	for (FEnvironment& Environment : ActiveEnvironments)
+	{
+		Environment.ProcessorBreakPoints.Reset();
+		Environment.FragmentWriteBreakPoints.Reset();
+		Environment.HasBreakPoint = false;
+	}
+	HasBreakPoint = false;
+}
+
+const UScriptStruct* FMassDebugger::GetFragmentTypeFromName(FName FragmentName)
+{
+	const UScriptStruct** FoundType = FragmentsByName.Find(FragmentName);
+	if (FoundType)
+	{
+		return *FoundType;
+	}
+
+	for (FEnvironment& Environment : ActiveEnvironments)
+	{
+		TArray<FMassArchetypeHandle> ArchetypeHandles = GetAllArchetypes(*Environment.EntityManager.Pin());
+		for (FMassArchetypeHandle& ArchetypeHandle : ArchetypeHandles)
+		{
+			FMassArchetypeCompositionDescriptor Composition = GetArchetypeComposition(ArchetypeHandle);
+
+			FMassFragmentBitSet::FIndexIterator It = Composition.Fragments.GetIndexIterator();
+			while (It)
+			{
+				FName StructName = Composition.Fragments.DebugGetStructTypeName(*It);
+				const UScriptStruct* StructType = Composition.Fragments.GetTypeAtIndex(*It);
+				
+				FragmentsByName.Add(StructName, StructType);
+				++It;
+			}
+
+			FMassChunkFragmentBitSet::FIndexIterator ChunkIt = Composition.ChunkFragments.GetIndexIterator();
+			while (ChunkIt)
+			{
+				FName StructName = Composition.ChunkFragments.DebugGetStructTypeName(*ChunkIt);
+				const UScriptStruct* StructType = Composition.ChunkFragments.GetTypeAtIndex(*ChunkIt);
+
+				FragmentsByName.Add(StructName, StructType);
+				++ChunkIt;
+			}
+
+			FMassSharedFragmentBitSet::FIndexIterator SharedFragIt = Composition.SharedFragments.GetIndexIterator();
+			while (SharedFragIt)
+			{
+				FName StructName = Composition.SharedFragments.DebugGetStructTypeName(*SharedFragIt);
+				const UScriptStruct* StructType = Composition.SharedFragments.GetTypeAtIndex(*SharedFragIt);
+
+				FragmentsByName.Add(StructName, StructType);
+				++SharedFragIt;
+			}
+
+			FMassConstSharedFragmentBitSet::FIndexIterator ConstSharedFragIt = Composition.ConstSharedFragments.GetIndexIterator();
+			while (ConstSharedFragIt)
+			{
+				FName StructName = Composition.ConstSharedFragments.DebugGetStructTypeName(*ConstSharedFragIt);
+				const UScriptStruct* StructType = Composition.ConstSharedFragments.GetTypeAtIndex(*ConstSharedFragIt);
+
+				FragmentsByName.Add(StructName, StructType);
+				++ConstSharedFragIt;
+			}
+		}
+	}
+
+	FoundType = FragmentsByName.Find(FragmentName);
+	if (FoundType)
+	{
+		return *FoundType;
+	}
+
+	return nullptr;
+}
+
+void FMassDebugger::UpdateHasBreakpoint()
+{
+	HasBreakPoint = false;
+	for (FEnvironment& Environment : ActiveEnvironments)
+	{
+		Environment.HasBreakPoint = Environment.ProcessorBreakPoints.Num() != 0
+			|| Environment.FragmentWriteBreakPoints.Num() != 0;
+		HasBreakPoint |= Environment.HasBreakPoint;
+	}
+}
+
+void FMassDebugger::SelectProcessor(const FMassEntityManager& EntityManager, const UMassProcessor* SelectedProcessor)
+{
+	const int32 Index = ActiveEnvironments.IndexOfByPredicate([WeakManager = EntityManager.AsWeak()](const FEnvironment& Element)
+	{
+		return Element.EntityManager == WeakManager;
+	});
+
+	if (Index >= 0)
+	{
+		if (ActiveEnvironments[Index].SelectedProcessor.Get() != SelectedProcessor)
+		{
+			ActiveEnvironments[Index].SelectedProcessor = SelectedProcessor;
+			ActiveEnvironments[Index].OnProcessorSelectionChanged.Broadcast(SelectedProcessor);
+		}
+	}
+}
+
+FOnProcessorSelectionChanged FMassDebugger::OnProcessorSelectionChanged(const FMassEntityManager& EntityManager)
+{
+	const int32 Index = ActiveEnvironments.IndexOfByPredicate([WeakManager = EntityManager.AsWeak()](const FEnvironment& Element)
+	{
+		return Element.EntityManager == WeakManager;
+	});
+
+	if (Index >= 0)
+	{
+		return ActiveEnvironments[Index].OnProcessorSelectionChanged;
+	}
+
+	static FOnProcessorSelectionChanged DefaultDelegate;
+	return DefaultDelegate;
+}
+
 #endif // WITH_MASSENTITY_DEBUG
\ No newline at end of file
diff --git a/Engine/Source/Runtime/MassEntity/Private/MassExecutionContext.cpp b/Engine/Source/Runtime/MassEntity/Private/MassExecutionContext.cpp
--- a/Engine/Source/Runtime/MassEntity/Private/MassExecutionContext.cpp
+++ b/Engine/Source/Runtime/MassEntity/Private/MassExecutionContext.cpp
@@ -3,6 +3,10 @@
 #include "MassExecutionContext.h"
 #include "MassArchetypeData.h"
 #include "MassEntityManager.h"
+#if WITH_MASSENTITY_DEBUG
+#include "MassDebugger.h"
+#endif // WITH_MASSENTITY_DEBUG
+
 
 
 //-----------------------------------------------------------------------------
@@ -28,7 +32,8 @@ FMassExecutionContext::FMassExecutionContext(const FMassExecutionContext& Other)
 FMassExecutionContext::FMassExecutionContext(const FMassExecutionContext& Other, FMassEntityQuery& Query, const TSharedPtr<FMassCommandBuffer>& InCommandBuffer)
 	: FMassExecutionContext(Other)
 {
-	QueriesStack.Add({&Query});
+	ensure(Other.QueriesStack.Last().Query == &Query);
+	QueriesStack.Add(Other.QueriesStack.Last());
 	SetDeferredCommandBuffer(InCommandBuffer);
 }
 
@@ -121,6 +126,37 @@ void FMassExecutionContext::PushQuery(FMassEntityQuery& InQuery)
 {
 	FQueryTransientRuntime& RuntimeData = QueriesStack.Add_GetRef({&InQuery});
 	GetSubsystemRequirementBits(RuntimeData.ConstSubsystemsBitSet, RuntimeData.MutableSubsystemsBitSet);
+
+#if WITH_MASSENTITY_DEBUG
+	// check if this could possibly trigger a break before iterating to avoid extraneous breakpoint checks
+	FMassEntityManager& EntityManagerRef = GetEntityManagerChecked();
+	RuntimeData.bCheckProcessorBreaks = FMassDebugger::HasAnyProcessorBreakpoints(EntityManagerRef, DebugGetProcessor());
+	
+	if (UNLIKELY(FMassDebugger::HasAnyFragmentWriteBreakpoints(EntityManagerRef)))
+	{
+		auto CheckFragmentRequirement = [&RuntimeData, &EntityManagerRef](TConstArrayView<FMassFragmentRequirementDescription> Requirements) -> void
+			{
+				for (const FMassFragmentRequirementDescription& Req : Requirements)
+				{
+					if (Req.AccessMode == EMassFragmentAccess::ReadWrite &&
+						FMassDebugger::HasAnyFragmentWriteBreakpoints(EntityManagerRef, Req.StructType))
+					{
+						if (RuntimeData.BreakFragmentsCount < FQueryTransientRuntime::MaxBreakFragmentCount)
+						{
+							RuntimeData.BreakFragments[RuntimeData.BreakFragmentsCount++] = Req.StructType;
+						}
+					}
+				}
+			};
+
+		// don't need to check ConstSharedFragmentRequirements because those can't write
+		CheckFragmentRequirement(InQuery.GetFragmentRequirements());
+		CheckFragmentRequirement(InQuery.GetChunkFragmentRequirements());
+		CheckFragmentRequirement(InQuery.GetSharedFragmentRequirements());
+	}
+#endif // WITH_MASSENTITY_DEBUG
+
+	RuntimeData.IteratorSerialNumber = ++IteratorSerialNumberGenerator;
 }
 
 void FMassExecutionContext::PopQuery(const FMassEntityQuery& InQuery)
@@ -132,3 +168,98 @@ void FMassExecutionContext::PopQuery(const FMassEntityQuery& InQuery)
 
 	QueriesStack.RemoveAt(QueriesStack.Num() - 1, EAllowShrinking::No);
 }
+
+FMassExecutionContext::FEntityIterator FMassExecutionContext::CreateEntityIterator()
+{
+	if (!ensureMsgf(QueriesStack.Num(), TEXT("")))
+	{
+		return FEntityIterator(*this);
+	}
+
+	return FEntityIterator(*this, QueriesStack.Last());
+}
+
+//-----------------------------------------------------------------------------
+// FMassExecutionContext::FQueryTransientRuntime
+//-----------------------------------------------------------------------------
+FMassExecutionContext::FQueryTransientRuntime& FMassExecutionContext::FQueryTransientRuntime::GetDummyInstance()
+{
+	static FMassEntityQuery DummyQuery;
+	static FQueryTransientRuntime DummyInstance = { &DummyQuery };
+	return DummyInstance;
+}
+
+//-----------------------------------------------------------------------------
+// FMassExecutionContext::FEntityIterator
+//-----------------------------------------------------------------------------
+FMassExecutionContext::FEntityIterator::FEntityIterator(FMassExecutionContext& InExecutionContext)
+	: ExecutionContext(InExecutionContext)
+	, QueryRuntime(FQueryTransientRuntime::GetDummyInstance())
+{
+	
+}
+
+FMassExecutionContext::FEntityIterator::FEntityIterator(FMassExecutionContext& InExecutionContext, FQueryTransientRuntime& InQueryRuntime)
+	: ExecutionContext(InExecutionContext)
+	, QueryRuntime(InQueryRuntime)
+	, EntitiesNum(InExecutionContext.GetNumEntities())
+	, SerialNumber(InQueryRuntime.IteratorSerialNumber)
+{
+	this->operator++();
+}
+
+#if WITH_MASSENTITY_DEBUG
+
+UE_DISABLE_OPTIMIZATION_SHIP
+void FMassExecutionContext::FEntityIterator::TestBreakPoints()
+{
+	FMassEntityManager& EntityManagerRef = ExecutionContext.GetEntityManagerChecked();
+	FMassEntityHandle Entity = ExecutionContext.GetEntity(EntityIndex);
+	if (QueryRuntime.bCheckProcessorBreaks && FMassDebugger::ShouldProcessorBreak(EntityManagerRef, ExecutionContext.DebugGetProcessor(), Entity))
+	{
+		bool bDisableThisBreakpoint = false;
+		//====================================================================
+		//= A breakpoint for this entity set in the MassDebugger has triggered
+		//= Step out of this function to debug the actual code being run for the entity
+		//=
+		//= To disable this specific breakpoint use the Watch window do set
+		//= bDisableThisBreakpoint to `true` or 1
+		//====================================================================
+		UE_DEBUG_BREAK();
+
+		if (bDisableThisBreakpoint)
+		{
+			FMassDebugger::ClearProcessorBreakPoint(EntityManagerRef, ExecutionContext.DebugGetProcessor(), Entity);
+		}
+
+		// bailing out, no point to hit multiple breakpoints for the given entity/processor pair
+		return;
+	}
+
+	for (const UScriptStruct* Fragment : QueryRuntime.BreakFragments)
+	{
+		if (FMassDebugger::ShouldFragmentWriteBreak(EntityManagerRef, Fragment, Entity))
+		{
+			bool bDisableThisBreakpoint = false;
+			//====================================================================
+			//= A breakpoint for this entity set in the MassDebugger has triggered
+			//= Step out of this function to debug the actual code being run for the entity
+			// 
+			//= To disable this specific breakpoint use the Watch window do set
+			//= bDisableThisBreakpoint to `true` or 1
+			//====================================================================
+			UE_DEBUG_BREAK();
+
+			if (bDisableThisBreakpoint)
+			{
+				FMassDebugger::ClearFragmentWriteBreak(EntityManagerRef, Fragment, Entity);
+			}
+
+			// bailing out, no point to hit multiple breakpoints for the given entity/processor pair
+			return;
+		}
+	}
+}
+UE_ENABLE_OPTIMIZATION_SHIP
+
+#endif // WITH_MASSENTITY_DEBUG
\ No newline at end of file
diff --git a/Engine/Source/Runtime/MassEntity/Private/MassProcessor.cpp b/Engine/Source/Runtime/MassEntity/Private/MassProcessor.cpp
--- a/Engine/Source/Runtime/MassEntity/Private/MassProcessor.cpp
+++ b/Engine/Source/Runtime/MassEntity/Private/MassProcessor.cpp
@@ -235,6 +235,7 @@ void UMassProcessor::CallExecute(FMassEntityManager& EntityManager, FMassExecuti
 		// Not using a more specific scope by default (i.e. LLM_SCOPE_BYNAME(*StatId)) since LLM is more strict regarding the provided string (no spaces or '_')
 
 #if WITH_MASSENTITY_DEBUG
+		Context.DebugSetExecutionDesc(FString::Printf(TEXT("%s (%s)"), *GetProcessorName(), EntityManager.GetWorld() ? *ToString(EntityManager.GetWorld()->GetNetMode()) : TEXT("No World")));
 		if (EntityManager.DebugHasAllDebugFeatures(FMassEntityManager::EDebugFeatures::TraceProcessors))
 		{
 			FString DebugDescription = *GetProcessorName();
@@ -246,6 +247,7 @@ void UMassProcessor::CallExecute(FMassEntityManager& EntityManager, FMassExecuti
 			DebugDescription.Append(TEXT(")"));
 			Context.DebugSetExecutionDesc(MoveTemp(DebugDescription));
 		}
+		Context.DebugSetProcessor(this);
 #endif
 		// CacheSubsystemRequirements will return true only if all requirements declared with ProcessorRequirements are met
 		// meaning if it fails there's no point in calling Execute.
diff --git a/Engine/Source/Runtime/MassEntity/Public/MassArchetypeTypes.h b/Engine/Source/Runtime/MassEntity/Public/MassArchetypeTypes.h
--- a/Engine/Source/Runtime/MassEntity/Public/MassArchetypeTypes.h
+++ b/Engine/Source/Runtime/MassEntity/Public/MassArchetypeTypes.h
@@ -24,6 +24,7 @@ struct FMassEntityView;
 struct FMassDebugger;
 struct FMassArchetypeHelper;
 
+using FMassEntityExecuteFunction = TFunction< void(FMassExecutionContext& /*ExecutionContext*/, int32 /*EntityIndex*/) >;
 using FMassExecuteFunction = TFunction< void(FMassExecutionContext& /*ExecutionContext*/) >;
 using FMassChunkConditionFunction = TFunction< bool(const FMassExecutionContext& /*ExecutionContext*/) >;
 
diff --git a/Engine/Source/Runtime/MassEntity/Public/MassDebugger.h b/Engine/Source/Runtime/MassEntity/Public/MassDebugger.h
--- a/Engine/Source/Runtime/MassEntity/Public/MassDebugger.h
+++ b/Engine/Source/Runtime/MassEntity/Public/MassDebugger.h
@@ -25,6 +25,9 @@ enum class EMassFragmentPresence : uint8;
 #endif // WITH_MASSENTITY_DEBUG
 #include "MassDebugger.generated.h"
 
+DECLARE_MULTICAST_DELEGATE_OneParam(FOnFragmentSelectionChanged, FName);
+DECLARE_MULTICAST_DELEGATE_OneParam(FOnProcessorSelectionChanged, const UMassProcessor*);
+
 namespace UE::Mass::Debug
 {
 	struct FArchetypeStats
@@ -139,6 +142,14 @@ struct MASSENTITY_API FMassDebugger
 		TWeakPtr<const FMassEntityManager> EntityManager;
 		TMap<FName, UE::Mass::Debug::FProcessorProviderFunction> ProcessorProviders;
 		FMassEntityHandle SelectedEntity;
+		FMassEntityHandle HighlightedEntity;
+		TWeakObjectPtr<const UMassProcessor> SelectedProcessor;
+		FOnProcessorSelectionChanged OnProcessorSelectionChanged;
+
+		bool HasBreakPoint = false;
+
+		TMultiMap<const UMassProcessor*, FMassEntityHandle> ProcessorBreakPoints;
+		TMultiMap<const UScriptStruct*, FMassEntityHandle> FragmentWriteBreakPoints;
 
 		explicit FEnvironment(const FMassEntityManager& InEntityManager)
 			: EntityManager(InEntityManager.AsWeak())
@@ -163,8 +174,9 @@ struct MASSENTITY_API FMassDebugger
 	static TArray<FMassArchetypeHandle> GetAllArchetypes(const FMassEntityManager& EntitySubsystem);
 	static const FMassArchetypeCompositionDescriptor& GetArchetypeComposition(const FMassArchetypeHandle& ArchetypeHandle);
 
-	static void GetArchetypeEntityStats(const FMassArchetypeHandle& ArchetypeHandle, UE::Mass::Debug::FArchetypeStats& OutStats);
+	static void GetArchetypeEntityStats(const FMassArchetypeHandle & ArchetypeHandle, UE::Mass::Debug::FArchetypeStats & OutStats);
 	static const TConstArrayView<FName> GetArchetypeDebugNames(const FMassArchetypeHandle& ArchetypeHandle);
+	static TArray<FMassEntityHandle> GetEntitiesOfArchetype(const FMassArchetypeHandle& ArchetypeHandle);
 
 	static TConstArrayView<UMassCompositeProcessor::FDependencyNode> GetProcessingGraph(const UMassCompositeProcessor& GraphOwner);
 	static TConstArrayView<TObjectPtr<UMassProcessor>> GetHostedProcessors(const UMassCompositeProcessor& GraphOwner);
@@ -181,6 +193,9 @@ struct MASSENTITY_API FMassDebugger
 	static void SelectEntity(const FMassEntityManager& EntityManager, const FMassEntityHandle EntityHandle);
 	static FMassEntityHandle GetSelectedEntity(const FMassEntityManager& EntityManager);
 
+	static void HighlightEntity(const FMassEntityManager& EntityManager, const FMassEntityHandle EntityHandle);
+	static FMassEntityHandle GetHighlightedEntity(const FMassEntityManager& EntityManager);
+
 	static FOnEntitySelected OnEntitySelectedDelegate;
 
 	static FOnMassEntityManagerEvent OnEntityManagerInitialized;
@@ -223,9 +238,89 @@ struct MASSENTITY_API FMassDebugger
 	 */
 	static bool DoesArchetypeMatchRequirements(const FMassArchetypeHandle& ArchetypeHandle, const FMassFragmentRequirements& Requirements, FOutputDevice& OutputDevice);
 
+	/**
+	* Checks if a processor should break on execute for a given entity.
+	*/
+	static bool ShouldProcessorBreak(const FMassEntityManager& EntityManager, const UMassProcessor* Processor, FMassEntityHandle Entity);
+
+	/**
+	* Checks if a processor has any breakpoints set for any entity.
+	*/
+	static bool HasAnyProcessorBreakpoints(const FMassEntityManager& EntityManager, const UMassProcessor* Processor);
+
+	/**
+	* Checks if a break shoud be triggered for a processor that's about to write a given fragment on an entity.
+	*/
+	static bool ShouldFragmentWriteBreak(const FMassEntityManager& EntityManager, const UScriptStruct* FragmentType, FMassEntityHandle Entity);
+
+	/**
+	* Checks if there are any breakpoints set for writing a fragment for any entity
+	* Use FragmentType = nullptr (default) to check for ANY fragment breaks.
+	*/
+	static bool HasAnyFragmentWriteBreakpoints(const FMassEntityManager& EntityManager, const UScriptStruct* FragmentType = nullptr);
+
+	/**
+	* Set a break to be triggered on processor execute for an entity.
+	*/
+	static void SetProcessorBreakPoint(const FMassEntityManager& EntityManager, const UMassProcessor* Processor, FMassEntityHandle Entity);
+
+	/**
+	* Set a break to be triggered for a processor that's about to write a given fragment on an entity.
+	*/
+	static void SetFragmentWriteBreak(const FMassEntityManager& EntityManager, const UScriptStruct* FragmentType, FMassEntityHandle Entity);
+
+	/**
+	* Clear a break triggered on processor execute for an entity.
+	*/
+	static void ClearProcessorBreakPoint(const FMassEntityManager& EntityManager, const UMassProcessor* Processor, FMassEntityHandle Entity);
+
+	/**
+	* Clear a break triggered for a processor that's about to write a given fragment on an entity.
+	*/
+	static void ClearFragmentWriteBreak(const FMassEntityManager& EntityManager, const UScriptStruct* FragmentType, FMassEntityHandle Entity);
+
+	/**
+	 * Clear all breakpoints set for the given processor
+	 */
+	static void ClearProcessorBreakPoints(const FMassEntityManager& EntityManager, TNotNull<const UMassProcessor*> Processor);
+
+	/**
+	 * Clear all write breakpoints for the given element type
+	 */
+	static void ClearAllFragmentWriteBreaks(const FMassEntityManager& EntityManager, TNotNull<const UScriptStruct*> FragmentType);
+
+	static void BreakOnSelectedFragmentWrite();
+	static const UScriptStruct* GetFragmentTypeFromName(FName FragmentName);
+	static void ClearAllBreakpoints();
+
+	/**
+	* Select a fragment name for highlight
+	*/
+	static void SelectFragment(FName FragmentName)
+	{
+		if (SelectedFragment != FragmentName)
+		{
+			SelectedFragment = FragmentName;
+			OnFragmentSelectionChanged.Broadcast(FragmentName);  // Broadcast when the fragment changes
+		}
+	}
+	static FName GetSelectedFragment() { return SelectedFragment; }
+	// Delegate to notify when the selected fragment changes
+	static FOnFragmentSelectionChanged& OnSelectedFragmentChanged() { return OnFragmentSelectionChanged; }
+
+	// Delegate to notify when the selected processor for the environment changes
+	static FOnProcessorSelectionChanged OnProcessorSelectionChanged(const FMassEntityManager& EntityManager);
+
+	static void SelectProcessor(const FMassEntityManager& EntityManager, const UMassProcessor* SelectedProcessor);
+	
 private:
 	static TArray<FEnvironment> ActiveEnvironments;
 	static UE::FSpinLock EntityManagerRegistrationLock;
+	static void UpdateHasBreakpoint();
+	static bool HasBreakPoint;
+	static FName SelectedFragment;
+	static FOnFragmentSelectionChanged OnFragmentSelectionChanged;
+	static TMap<FName, const UScriptStruct*> FragmentsByName;
 };
 
 #else
diff --git a/Engine/Source/Runtime/MassEntity/Public/MassExecutionContext.h b/Engine/Source/Runtime/MassEntity/Public/MassExecutionContext.h
--- a/Engine/Source/Runtime/MassEntity/Public/MassExecutionContext.h
+++ b/Engine/Source/Runtime/MassEntity/Public/MassExecutionContext.h
@@ -75,13 +75,31 @@ private:
 		TNotNull<FMassEntityQuery*> Query;
 		FMassExternalSubsystemBitSet ConstSubsystemsBitSet;
 		FMassExternalSubsystemBitSet MutableSubsystemsBitSet;
+#if WITH_MASSENTITY_DEBUG
+		// MaxBreakFragmentCount needs to be bigger than the greatest number of fragments a processor could write to that have a breakpoint set
+		static constexpr uint32 MaxBreakFragmentCount = 8;
+		TStaticArray<const UScriptStruct*, MaxBreakFragmentCount> BreakFragments;
+
+		bool bCheckProcessorBreaks = false;
+		int32 BreakFragmentsCount = 0;
+#endif // WITH_MASSENTITY_DEBUG
+
+		uint32 IteratorSerialNumber = 0;
+
+		static FQueryTransientRuntime& GetDummyInstance();
 	};
 	/** We usually expect the queries to go only a single layer deep, so 2 elements here should suffice most of the time */
 	TArray<FQueryTransientRuntime, TInlineAllocator<2>> QueriesStack;
+	uint32 IteratorSerialNumberGenerator = 0;
 
 #if WITH_MASSENTITY_DEBUG
 	FString DebugExecutionDescription;
-#endif
+
+	/** Currently executing processor, used for debugger breakpoint checking */
+	// js@todo make this more generic
+	// mz@todo this could be moved to FQueryTransientRuntime now 
+	UMassProcessor* DebugProcessor = nullptr;
+#endif // WITH_MASSENTITY_DEBUG
 	
 	/** Used to control when the context is allowed to flush commands collected in DeferredCommandBuffer. This mechanism 
 	 * is mainly utilized to avoid numerous small flushes in favor of fewer larger ones. */
@@ -126,6 +144,9 @@ public:
 	const FString& DebugGetExecutionDesc() const { return DebugExecutionDescription; }
 	void DebugSetExecutionDesc(const FString& Description) { DebugExecutionDescription = Description; }
 	void DebugSetExecutionDesc(FString&& Description) { DebugExecutionDescription = MoveTemp(Description); }
+
+	UMassProcessor* DebugGetProcessor() const { return DebugProcessor; }
+	void DebugSetProcessor(UMassProcessor* Processor) { DebugProcessor = Processor; }
 #endif
 
 	void PushQuery(FMassEntityQuery& InQuery);
@@ -135,6 +156,87 @@ public:
 	void ApplyFragmentRequirements(const FMassEntityQuery& RequestingQuery);
 	void ClearFragmentViews(const FMassEntityQuery& RequestingQuery);
 
+	struct FEntityIterator
+	{
+		FORCEINLINE int operator*() const
+		{
+			return EntityIndex;
+		}
+
+		FORCEINLINE bool operator!=(const int& other) const
+		{
+			return EntityIndex != other;
+		}
+
+		FORCEINLINE operator int() const
+		{
+			return EntityIndex;
+		}
+
+		FORCEINLINE operator bool() const
+		{
+			return SerialNumber && EntityIndex < EntitiesNum;
+		}
+
+		FORCEINLINE bool operator<(const int32 Other) const
+		{
+			return SerialNumber && EntityIndex != INDEX_NONE && EntityIndex < Other;
+		}
+
+		FORCEINLINE FEntityIterator& operator++()
+		{
+			++EntityIndex;
+#if WITH_MASSENTITY_DEBUG
+			if (UNLIKELY(QueryRuntime.bCheckProcessorBreaks || QueryRuntime.BreakFragmentsCount != 0) 
+				&& EntityIndex < EntitiesNum)
+			{
+				TestBreakPoints();
+			}
+#endif //!WITH_MASSENTITY_DEBUG
+			return *this;
+		}
+
+		FORCEINLINE FEntityIterator operator++(int)
+		{
+			FEntityIterator Temp  = *this;
+			++(*this);
+			return Temp;
+		}
+
+		FEntityIterator& begin()
+		{
+			return *this;
+		}
+
+		FEntityIterator end() const
+		{
+			FEntityIterator End = *this;
+			End.EntityIndex = EntitiesNum;
+			return End;
+		}
+
+		FEntityIterator& operator=(const FEntityIterator&) = delete;
+		FEntityIterator& operator=(FEntityIterator&&) = delete;
+		FEntityIterator(const FEntityIterator&) = default;
+		FEntityIterator(FEntityIterator&&) = default;
+	private:
+		friend FMassExecutionContext;
+		FEntityIterator(FMassExecutionContext& InExecutionContext);
+		FEntityIterator(FMassExecutionContext& InExecutionContext, FQueryTransientRuntime& InQueryRuntime);
+		
+
+#if WITH_MASSENTITY_DEBUG
+		MASSENTITY_API void TestBreakPoints();
+#endif //!WITH_MASSENTITY_DEBUG
+
+		FMassExecutionContext& ExecutionContext;
+		FQueryTransientRuntime& QueryRuntime;
+		int32 EntityIndex = INDEX_NONE;
+		const int32 EntitiesNum = INDEX_NONE;
+		const uint32 SerialNumber = 0;
+	};
+	FEntityIterator CreateEntityIterator();
+
 	/** Sets bFlushDeferredCommands. Note that setting to True while the system is being executed doesn't result in
 	 *  immediate commands flushing */
 	void SetFlushDeferredCommands(const bool bNewFlushDeferredCommands);
@@ -162,6 +264,14 @@ public:
 		return EntityListView[Index];
 	}
 
+	void ForEachEntityInChunk(const FMassEntityExecuteFunction& EntityExecuteFunction)
+	{
+		for (FMassExecutionContext::FEntityIterator EntityIterator = CreateEntityIterator(); EntityIterator; ++EntityIterator)
+		{
+			EntityExecuteFunction(*this, EntityIterator);
+		}
+	}
+
 	bool DoesArchetypeHaveFragment(const UScriptStruct& FragmentType) const
 	{
 		return CurrentArchetypeCompositionDescriptor.Fragments.Contains(FragmentType);
