[Mass] Made a minor modification how a MassProcessor can declare it doesn't want to be pruned based on queries - no longer needing to override ShouldAllowQueryBasedPruning

diff --git a/Engine/Source/Runtime/MassEntity/Private/MassProcessor.cpp b/Engine/Source/Runtime/MassEntity/Private/MassProcessor.cpp
--- a/Engine/Source/Runtime/MassEntity/Private/MassProcessor.cpp
+++ b/Engine/Source/Runtime/MassEntity/Private/MassProcessor.cpp
@@ -271,7 +271,7 @@ void UMassProcessor::CallExecute(FMassEntityManager& EntityManager, FMassExecuti
 
 bool UMassProcessor::ShouldAllowQueryBasedPruning(const bool bRuntimeMode) const
 {
-	return bRuntimeMode;
+	return bRuntimeMode && QueryBasedPruning == EMassQueryBasedPruning::Prune;
 }
 
 EMassProcessingPhase UMassProcessor::GetProcessingPhase() const
diff --git a/Engine/Source/Runtime/MassEntity/Public/MassProcessor.h b/Engine/Source/Runtime/MassEntity/Public/MassProcessor.h
--- a/Engine/Source/Runtime/MassEntity/Public/MassProcessor.h
+++ b/Engine/Source/Runtime/MassEntity/Public/MassProcessor.h
@@ -50,6 +50,19 @@ enum class EActivationState : uint8
 	OneShot,	// one-shot processor will auto-disable itself after the next CallExecute call
 };
 
+/**
+ * Values determining whether a processor wants to be pruned at runtime. The value is not used when
+ * processing graph is generated for project configuration purposes or debug-time graph visualization purposes
+ * This behavior can be overridden by UMassProcessor::ShouldAllowQueryBasedPruning child class overrides
+ */
+UENUM()
+enum class EMassQueryBasedPruning : uint8
+{
+	Prune,		// pruning will always be applied at runtime
+	Never,		// pruning will never be applied at runtime
+	Default = Prune
+};
+
 UCLASS(abstract, EditInlineNew, CollapseCategories, config = Mass, defaultconfig, ConfigDoNotCheckDefaults)
 class MASSENTITY_API UMassProcessor : public UObject
 {
@@ -86,11 +99,14 @@ public:
 	//----------------------------------------------------------------------//
 	// Ordering functions 
 	//----------------------------------------------------------------------//
-	/** Indicates whether this processor can ever be pruned while considered for a phase processing graph. A processor
-	 *  can get pruned if none of its registered queries interact with archetypes instantiated at the moment of graph
-	 *  building. This can also happen for special processors that don't register any queries - if that's the case override 
-	 *  this function to return an appropriate value
-	 *  @param bRuntimeMode indicates whether the pruning is being done for game runtime (true) or editor-time presentation (false) */
+	/**
+	 * Indicates whether this processor can ever be pruned while considered for a phase processing graph. A processor
+	 * can get pruned if none of its registered queries interact with archetypes instantiated at the moment of graph
+	 * building. This can also happen for special processors that don't register any queries - if that's the case override 
+	 * this function to return an appropriate value or use QueryBasedPruning to configure the expected behavior.
+	 * By default, the processor will be the subject of pruning when bRuntimeMode == true.
+	 * @param bRuntimeMode indicates whether the pruning is being done for game runtime (true) or editor-time presentation (false)
+	 */
 	virtual bool ShouldAllowQueryBasedPruning(const bool bRuntimeMode = true) const;
 
 	virtual EMassProcessingPhase GetProcessingPhase() const;
@@ -224,6 +240,13 @@ private:
 
 protected:
 
+	/**
+	 * Determines whether given processor wants to be pruned from the execution graph when there are
+	 * no archetypes matching its requirements.
+	 * Defaults to EMassQueryBasedPruning::Prune, @see EMassQueryBasedPruning
+	 */
+	EMassQueryBasedPruning QueryBasedPruning = EMassQueryBasedPruning::Default;
+
 	friend UMassCompositeProcessor;
 	friend FMassDebugger;
 
