[Mass] moved Mass's testable ensured handling to a dedicated header, MassTestableEnsures.h

diff --git a/Engine/Source/Runtime/MassEntity/Public/MassEntityTypes.h b/Engine/Source/Runtime/MassEntity/Public/MassEntityTypes.h
--- a/Engine/Source/Runtime/MassEntity/Public/MassEntityTypes.h
+++ b/Engine/Source/Runtime/MassEntity/Public/MassEntityTypes.h
@@ -12,13 +12,7 @@
 #include "StructUtils/SharedStruct.h"
 #include "MassEntityElementTypes.h"
 #include "MassEntityConcepts.h"
-#ifdef WITH_AITESTSUITE
-#include "TestableEnsures.h"
-#else
-#define testableEnsureMsgf ensureMsgf
-#define testableCheckf checkf
-#define testableCheckfReturn(InExpression, ReturnValue, InFormat, ... ) checkf(InExpression, InFormat, ##__VA_ARGS__)
-#endif 
+#include "MassTestableEnsures.h"
 #if UE_ENABLE_INCLUDE_ORDER_DEPRECATED_IN_5_6
 #include "MassExternalSubsystemTraits.h"
 #include "MassEntityHandle.h"
diff --git a/Engine/Source/Runtime/MassEntity/Public/MassTestableEnsures.h b/Engine/Source/Runtime/MassEntity/Public/MassTestableEnsures.h
new file mode 100644
--- /dev/null
+++ b/Engine/Source/Runtime/MassEntity/Public/MassTestableEnsures.h
@@ -0,0 +1,11 @@
+// Copyright Epic Games, Inc. All Rights Reserved.
+
+#pragma once
+
+#ifdef WITH_AITESTSUITE
+#include "TestableEnsures.h"
+#else
+#define testableEnsureMsgf ensureMsgf
+#define testableCheckf checkf
+#define testableCheckfReturn(InExpression, ReturnValue, InFormat, ... ) checkf(InExpression, InFormat, ##__VA_ARGS__)
+#endif 
