[Mass] Moved all one-line functions to the back of the MassRequirements.h header
Also, fixed comments wording a bit

#rnx

diff --git a/Engine/Source/Runtime/MassEntity/Public/MassRequirements.h b/Engine/Source/Runtime/MassEntity/Public/MassRequirements.h
--- a/Engine/Source/Runtime/MassEntity/Public/MassRequirements.h
+++ b/Engine/Source/Runtime/MassEntity/Public/MassRequirements.h
@@ -31,7 +31,7 @@ enum class EMassFragmentAccess : uint8
 UENUM()
 enum class EMassFragmentPresence : uint8
 {
-	/** All of the required fragments must be present */
+	/** All the required fragments must be present */
 	All,
 
 	/** One of the required fragments must be present */
@@ -40,49 +40,36 @@ enum class EMassFragmentPresence : uint8
 	/** None of the required fragments can be present */
 	None,
 
-	/** If fragment is present we'll use it, but it missing stop processing of a given archetype */
+	/** If fragment is present we'll use it */
 	Optional,
 
 	MAX
 };
 
 
-struct MASSENTITY_API FMassFragmentRequirementDescription
+struct FMassFragmentRequirementDescription
 {
-	const UScriptStruct* StructType = nullptr;
-	EMassFragmentAccess AccessMode = EMassFragmentAccess::None;
-	EMassFragmentPresence Presence = EMassFragmentPresence::Optional;
-
-public:
-	FMassFragmentRequirementDescription(){}
-	FMassFragmentRequirementDescription(const UScriptStruct* InStruct, const EMassFragmentAccess InAccessMode, const EMassFragmentPresence InPresence)
-		: StructType(InStruct)
-		, AccessMode(InAccessMode)
-		, Presence(InPresence)
-	{
-		check(InStruct);
-	}
+	FMassFragmentRequirementDescription() = default;
+	FMassFragmentRequirementDescription(const UScriptStruct* InStruct, const EMassFragmentAccess InAccessMode, const EMassFragmentPresence InPresence);
 
-	bool RequiresBinding() const { return (AccessMode != EMassFragmentAccess::None); }
-	bool IsOptional() const { return (Presence == EMassFragmentPresence::Optional || Presence == EMassFragmentPresence::Any); }
+	bool RequiresBinding() const;
+	bool IsOptional() const;
 
 	/** these functions are used for sorting. See FScriptStructSortOperator */
-	int32 GetStructureSize() const
-	{
-		return StructType->GetStructureSize();
-	}
+	int32 GetStructureSize() const;
 
-	FName GetFName() const
-	{
-		return StructType->GetFName();
-	}
+	FName GetFName() const;
+
+	const UScriptStruct* StructType = nullptr;
+	EMassFragmentAccess AccessMode = EMassFragmentAccess::None;
+	EMassFragmentPresence Presence = EMassFragmentPresence::Optional;
 };
 
 /**
  *  FMassSubsystemRequirements is a structure that declares runtime subsystem access type given calculations require.
  */
 USTRUCT()
-struct MASSENTITY_API FMassSubsystemRequirements
+struct FMassSubsystemRequirements
 {
 	GENERATED_BODY()
 
@@ -136,14 +123,14 @@ struct MASSENTITY_API FMassSubsystemRequirements
 		return *this;
 	}
 
-	void Reset();
+	MASSENTITY_API void Reset();
 
-	const FMassExternalSubsystemBitSet& GetRequiredConstSubsystems() const { return RequiredConstSubsystems; }
-	const FMassExternalSubsystemBitSet& GetRequiredMutableSubsystems() const { return RequiredMutableSubsystems; }
-	bool IsEmpty() const { return RequiredConstSubsystems.IsEmpty() && RequiredMutableSubsystems.IsEmpty(); }
+	const FMassExternalSubsystemBitSet& GetRequiredConstSubsystems() const;
+	const FMassExternalSubsystemBitSet& GetRequiredMutableSubsystems() const;
+	bool IsEmpty() const;
 
-	bool DoesRequireGameThreadExecution() const { return bRequiresGameThreadExecution; }
-	void ExportRequirements(FMassExecutionRequirements& OutRequirements) const;
+	bool DoesRequireGameThreadExecution() const;
+	MASSENTITY_API void ExportRequirements(FMassExecutionRequirements& OutRequirements) const;
 
 protected:
 	FMassExternalSubsystemBitSet RequiredConstSubsystems;
@@ -157,7 +144,7 @@ private:
  *  FMassFragmentRequirements is a structure that describes properties required of an archetype that's a subject of calculations.
  */
 USTRUCT()
-struct MASSENTITY_API FMassFragmentRequirements
+struct FMassFragmentRequirements
 {
 	GENERATED_BODY()
 
@@ -166,10 +153,10 @@ struct MASSENTITY_API FMassFragmentRequirements
 
 public:
 	FMassFragmentRequirements() = default;
-	explicit FMassFragmentRequirements(const TSharedPtr<FMassEntityManager>& EntityManager);
-	explicit FMassFragmentRequirements(const TSharedRef<FMassEntityManager>& EntityManager);
+	MASSENTITY_API explicit FMassFragmentRequirements(const TSharedPtr<FMassEntityManager>& EntityManager);
+	MASSENTITY_API explicit FMassFragmentRequirements(const TSharedRef<FMassEntityManager>& EntityManager);
 
-	void Initialize(const TSharedRef<FMassEntityManager>& EntityManager);
+	MASSENTITY_API void Initialize(const TSharedRef<FMassEntityManager>& EntityManager);
 
 	FMassFragmentRequirements& AddRequirement(const UScriptStruct* FragmentType, const EMassFragmentAccess AccessMode, const EMassFragmentPresence Presence = EMassFragmentPresence::All)
 	{
@@ -294,7 +281,7 @@ public:
 	}
 
 	/** Clears given tags out of all collected requirements, including negative ones */
-	FMassFragmentRequirements& ClearTagRequirements(const FMassTagBitSet& TagsToRemoveBitSet);
+	MASSENTITY_API FMassFragmentRequirements& ClearTagRequirements(const FMassTagBitSet& TagsToRemoveBitSet);
 
 	template<typename T>
 	FMassFragmentRequirements& AddChunkRequirement(const EMassFragmentAccess AccessMode, const EMassFragmentPresence Presence = EMassFragmentPresence::All)
@@ -323,6 +310,31 @@ public:
 		return *this;
 	}
 
+	FMassFragmentRequirements& AddChunkRequirement(TNotNull<const UScriptStruct*> ChunkFragmentType, const EMassFragmentAccess AccessMode, const EMassFragmentPresence Presence = EMassFragmentPresence::All)
+	{
+		checkf(bInitialized, TEXT("Modifying requirements before initialization is not supported."));
+		checkf(ChunkFragmentRequirements.FindByPredicate([&ChunkFragmentType](const FMassFragmentRequirementDescription& Item) { return Item.StructType == ChunkFragmentType; }) == nullptr
+			, TEXT("Duplicated requirements are not supported. %s already present"), *ChunkFragmentType->GetName());
+		checkf(Presence != EMassFragmentPresence::Any, TEXT("\'Any\' is not a valid Presence value for AddChunkRequirement."));
+
+		switch (Presence)
+		{
+			case EMassFragmentPresence::All:
+				RequiredAllChunkFragments.Add(*ChunkFragmentType);
+				ChunkFragmentRequirements.Emplace(ChunkFragmentType, AccessMode, Presence);
+				break;
+			case EMassFragmentPresence::Optional:
+				RequiredOptionalChunkFragments.Add(*ChunkFragmentType);
+				ChunkFragmentRequirements.Emplace(ChunkFragmentType, AccessMode, Presence);
+				break;
+			case EMassFragmentPresence::None:
+				RequiredNoneChunkFragments.Add(*ChunkFragmentType);
+				break;
+		}
+		IncrementChangeCounter();
+		return *this;
+	}
+
 	template<typename T>
 	FMassFragmentRequirements& AddConstSharedRequirement(const EMassFragmentPresence Presence = EMassFragmentPresence::All)
 	{
@@ -416,61 +428,57 @@ public:
 		return *this;
 	}
 
-	void Reset();
+	MASSENTITY_API void Reset();
 
 	/** 
 	 * The function validates requirements we make for queries. See the FMassFragmentRequirements struct description for details.
-	 * Even though the code of the function is non trivial the consecutive calls will be essentially free due to the result 
+	 * Even though the code of the function is non-trivial the consecutive calls will be essentially free due to the result 
 	 * being cached (note that the caching gets invalidated if the composition changes).
 	 * @return whether this query's requirements follow the rules.
 	 */
-	bool CheckValidity() const;
-
-	TConstArrayView<FMassFragmentRequirementDescription> GetFragmentRequirements() const { return FragmentRequirements; }
-	TConstArrayView<FMassFragmentRequirementDescription> GetChunkFragmentRequirements() const { return ChunkFragmentRequirements; }
-	TConstArrayView<FMassFragmentRequirementDescription> GetConstSharedFragmentRequirements() const { return ConstSharedFragmentRequirements; }
-	TConstArrayView<FMassFragmentRequirementDescription> GetSharedFragmentRequirements() const { return SharedFragmentRequirements; }
-	const FMassFragmentBitSet& GetRequiredAllFragments() const { return RequiredAllFragments; }
-	const FMassFragmentBitSet& GetRequiredAnyFragments() const { return RequiredAnyFragments; }
-	const FMassFragmentBitSet& GetRequiredOptionalFragments() const { return RequiredOptionalFragments; }
-	const FMassFragmentBitSet& GetRequiredNoneFragments() const { return RequiredNoneFragments; }
-	const FMassTagBitSet& GetRequiredAllTags() const { return RequiredAllTags; }
-	const FMassTagBitSet& GetRequiredAnyTags() const { return RequiredAnyTags; }
-	const FMassTagBitSet& GetRequiredNoneTags() const { return RequiredNoneTags; }
-	const FMassTagBitSet& GetRequiredOptionalTags() const { return RequiredOptionalTags; }
-	const FMassChunkFragmentBitSet& GetRequiredAllChunkFragments() const { return RequiredAllChunkFragments; }
-	const FMassChunkFragmentBitSet& GetRequiredOptionalChunkFragments() const { return RequiredOptionalChunkFragments; }
-	const FMassChunkFragmentBitSet& GetRequiredNoneChunkFragments() const { return RequiredNoneChunkFragments; }
-	const FMassSharedFragmentBitSet& GetRequiredAllSharedFragments() const { return RequiredAllSharedFragments; }
-	const FMassSharedFragmentBitSet& GetRequiredOptionalSharedFragments() const { return RequiredOptionalSharedFragments; }
-	const FMassSharedFragmentBitSet& GetRequiredNoneSharedFragments() const { return RequiredNoneSharedFragments; }
-	const FMassConstSharedFragmentBitSet& GetRequiredAllConstSharedFragments() const { return RequiredAllConstSharedFragments; }
-	const FMassConstSharedFragmentBitSet& GetRequiredOptionalConstSharedFragments() const { return RequiredOptionalConstSharedFragments; }
-	const FMassConstSharedFragmentBitSet& GetRequiredNoneConstSharedFragments() const { return RequiredNoneConstSharedFragments; }
-
-	bool IsInitialized() const { return bInitialized; }
-	bool IsEmpty() const;
-	bool HasPositiveRequirements() const { return bHasPositiveRequirements; }
-	bool HasNegativeRequirements() const { return bHasNegativeRequirements; }
-	bool HasOptionalRequirements() const { return bHasOptionalRequirements; }
-
-	bool DoesArchetypeMatchRequirements(const FMassArchetypeHandle& ArchetypeHandle) const;
-	bool DoesArchetypeMatchRequirements(const FMassArchetypeCompositionDescriptor& ArchetypeComposition) const;
-	bool DoesMatchAnyOptionals(const FMassArchetypeCompositionDescriptor& ArchetypeComposition) const;
-
-	bool DoesRequireGameThreadExecution() const { return bRequiresGameThreadExecution; }
-	void ExportRequirements(FMassExecutionRequirements& OutRequirements) const;
+	MASSENTITY_API bool CheckValidity() const;
+
+	TConstArrayView<FMassFragmentRequirementDescription> GetFragmentRequirements() const;
+	TConstArrayView<FMassFragmentRequirementDescription> GetChunkFragmentRequirements() const;
+	TConstArrayView<FMassFragmentRequirementDescription> GetConstSharedFragmentRequirements() const;
+	TConstArrayView<FMassFragmentRequirementDescription> GetSharedFragmentRequirements() const;
+	const FMassFragmentBitSet& GetRequiredAllFragments() const;
+	const FMassFragmentBitSet& GetRequiredAnyFragments() const;
+	const FMassFragmentBitSet& GetRequiredOptionalFragments() const;
+	const FMassFragmentBitSet& GetRequiredNoneFragments() const;
+	const FMassTagBitSet& GetRequiredAllTags() const;
+	const FMassTagBitSet& GetRequiredAnyTags() const;
+	const FMassTagBitSet& GetRequiredNoneTags() const;
+	const FMassTagBitSet& GetRequiredOptionalTags() const;
+	const FMassChunkFragmentBitSet& GetRequiredAllChunkFragments() const;
+	const FMassChunkFragmentBitSet& GetRequiredOptionalChunkFragments() const;
+	const FMassChunkFragmentBitSet& GetRequiredNoneChunkFragments() const;
+	const FMassSharedFragmentBitSet& GetRequiredAllSharedFragments() const;
+	const FMassSharedFragmentBitSet& GetRequiredOptionalSharedFragments() const;
+	const FMassSharedFragmentBitSet& GetRequiredNoneSharedFragments() const;
+	const FMassConstSharedFragmentBitSet& GetRequiredAllConstSharedFragments() const;
+	const FMassConstSharedFragmentBitSet& GetRequiredOptionalConstSharedFragments() const;
+	const FMassConstSharedFragmentBitSet& GetRequiredNoneConstSharedFragments() const;
+
+	bool IsInitialized() const;
+	MASSENTITY_API bool IsEmpty() const;
+	bool HasPositiveRequirements() const;
+	bool HasNegativeRequirements() const;
+	bool HasOptionalRequirements() const;
+
+	MASSENTITY_API bool DoesArchetypeMatchRequirements(const FMassArchetypeHandle& ArchetypeHandle) const;
+	MASSENTITY_API bool DoesArchetypeMatchRequirements(const FMassArchetypeCompositionDescriptor& ArchetypeComposition) const;
+	MASSENTITY_API bool DoesMatchAnyOptionals(const FMassArchetypeCompositionDescriptor& ArchetypeComposition) const;
+
+	bool DoesRequireGameThreadExecution() const;
+	MASSENTITY_API void ExportRequirements(FMassExecutionRequirements& OutRequirements) const;
 
 protected:
-	void SortRequirements();
+	MASSENTITY_API void SortRequirements();
 
-	FORCEINLINE void IncrementChangeCounter() 
-	{ 
-		++IncrementalChangesCount; 
-		bPropertiesCached = false;
-	}
-	void ConsumeIncrementalChangesCount() { IncrementalChangesCount = 0; }
-	bool HasIncrementalChanges() const{ return IncrementalChangesCount > 0; }
+	void IncrementChangeCounter();
+	void ConsumeIncrementalChangesCount();
+	bool HasIncrementalChanges() const;
 
 	friend FMassRequirementAccessDetector;
 
@@ -503,7 +511,7 @@ protected:
 	TSharedPtr<FMassEntityManager> CachedEntityManager;
 
 private:
-	void CacheProperties() const;
+	MASSENTITY_API void CacheProperties() const;
 	mutable uint16 bPropertiesCached : 1 = false;
 	mutable uint16 bHasPositiveRequirements : 1 = false;
 	mutable uint16 bHasNegativeRequirements : 1 = false;
@@ -521,17 +529,68 @@ private:
 
 public:
 	UE_DEPRECATED(5.6, "This type of FMassFragmentRequirements is no longer supported. Use one of the other constructors instead.")
-	FMassFragmentRequirements(std::initializer_list<UScriptStruct*> InitList);
+	MASSENTITY_API FMassFragmentRequirements(std::initializer_list<UScriptStruct*> InitList);
 
 	UE_DEPRECATED(5.6, "This type of FMassFragmentRequirements is no longer supported. Use one of the other constructors instead.")
-	FMassFragmentRequirements(TConstArrayView<const UScriptStruct*> InitList);
+	MASSENTITY_API FMassFragmentRequirements(TConstArrayView<const UScriptStruct*> InitList);
 };
 
+//-----------------------------------------------------------------------------
+// INLINE
+//-----------------------------------------------------------------------------
+inline FMassFragmentRequirementDescription::FMassFragmentRequirementDescription(const UScriptStruct* InStruct, const EMassFragmentAccess InAccessMode, const EMassFragmentPresence InPresence)
+	: StructType(InStruct)
+	, AccessMode(InAccessMode)
+	, Presence(InPresence)
+{
+	check(InStruct);
+}
+
+inline bool FMassFragmentRequirementDescription::RequiresBinding() const
+{
+	return (AccessMode != EMassFragmentAccess::None);
+}
+
+inline bool FMassFragmentRequirementDescription::IsOptional() const
+{
+	return (Presence == EMassFragmentPresence::Optional || Presence == EMassFragmentPresence::Any);
+}
+
+inline int32 FMassFragmentRequirementDescription::GetStructureSize() const
+{
+	return StructType->GetStructureSize();
+}
+
+inline FName FMassFragmentRequirementDescription::GetFName() const
+{
+	return StructType->GetFName();
+}
+
+inline const FMassExternalSubsystemBitSet& FMassSubsystemRequirements::GetRequiredConstSubsystems() const
+{
+	return RequiredConstSubsystems;
+}
+
+inline const FMassExternalSubsystemBitSet& FMassSubsystemRequirements::GetRequiredMutableSubsystems() const
+{
+	return RequiredMutableSubsystems;
+}
+
+inline bool FMassSubsystemRequirements::IsEmpty() const
+{
+	return RequiredConstSubsystems.IsEmpty() && RequiredMutableSubsystems.IsEmpty();
+}
+
+inline bool FMassSubsystemRequirements::DoesRequireGameThreadExecution() const
+{
+	return bRequiresGameThreadExecution;
+}
+
 template<>
 FORCEINLINE FMassFragmentRequirements& FMassFragmentRequirements::AddTagRequirements<EMassFragmentPresence::All>(const FMassTagBitSet& TagBitSet)
 {
 	RequiredAllTags += TagBitSet;
-	// force recaching the next time this query is used or the following CacheArchetypes call.
+	// force re-caching the next time this query is used or the following CacheArchetypes call.
 	IncrementChangeCounter();
 	return *this;
 }
@@ -540,7 +599,7 @@ template<>
 FORCEINLINE FMassFragmentRequirements& FMassFragmentRequirements::AddTagRequirements<EMassFragmentPresence::Any>(const FMassTagBitSet& TagBitSet)
 {
 	RequiredAnyTags += TagBitSet;
-	// force recaching the next time this query is used or the following CacheArchetypes call.
+	// force re-caching the next time this query is used or the following CacheArchetypes call.
 	IncrementChangeCounter();
 	return *this;
 }
@@ -549,7 +608,7 @@ template<>
 FORCEINLINE FMassFragmentRequirements& FMassFragmentRequirements::AddTagRequirements<EMassFragmentPresence::None>(const FMassTagBitSet& TagBitSet)
 {
 	RequiredNoneTags += TagBitSet;
-	// force recaching the next time this query is used or the following CacheArchetypes call.
+	// force re-caching the next time this query is used or the following CacheArchetypes call.
 	IncrementChangeCounter();
 	return *this;
 }
@@ -558,7 +617,153 @@ template<>
 FORCEINLINE FMassFragmentRequirements& FMassFragmentRequirements::AddTagRequirements<EMassFragmentPresence::Optional>(const FMassTagBitSet& TagBitSet)
 {
 	RequiredOptionalTags += TagBitSet;
-	// force recaching the next time this query is used or the following CacheArchetypes call.
+	// force re-caching the next time this query is used or the following CacheArchetypes call.
 	IncrementChangeCounter();
 	return *this;
 }
+
+inline TConstArrayView<FMassFragmentRequirementDescription> FMassFragmentRequirements::GetFragmentRequirements() const
+{ 
+	return FragmentRequirements; 
+}
+
+inline TConstArrayView<FMassFragmentRequirementDescription> FMassFragmentRequirements::GetChunkFragmentRequirements() const
+{ 
+	return ChunkFragmentRequirements; 
+}
+
+inline TConstArrayView<FMassFragmentRequirementDescription> FMassFragmentRequirements::GetConstSharedFragmentRequirements() const
+{ 
+	return ConstSharedFragmentRequirements; 
+}
+
+inline TConstArrayView<FMassFragmentRequirementDescription> FMassFragmentRequirements::GetSharedFragmentRequirements() const
+{ 
+	return SharedFragmentRequirements; 
+}
+
+inline const FMassFragmentBitSet& FMassFragmentRequirements::GetRequiredAllFragments() const
+{ 
+	return RequiredAllFragments; 
+}
+
+inline const FMassFragmentBitSet& FMassFragmentRequirements::GetRequiredAnyFragments() const
+{ 
+	return RequiredAnyFragments; 
+}
+
+inline const FMassFragmentBitSet& FMassFragmentRequirements::GetRequiredOptionalFragments() const
+{ 
+	return RequiredOptionalFragments; 
+}
+
+inline const FMassFragmentBitSet& FMassFragmentRequirements::GetRequiredNoneFragments() const
+{ 
+	return RequiredNoneFragments; 
+}
+
+inline const FMassTagBitSet& FMassFragmentRequirements::GetRequiredAllTags() const
+{ 
+	return RequiredAllTags; 
+}
+
+inline const FMassTagBitSet& FMassFragmentRequirements::GetRequiredAnyTags() const
+{ 
+	return RequiredAnyTags; 
+}
+
+inline const FMassTagBitSet& FMassFragmentRequirements::GetRequiredNoneTags() const
+{ 
+	return RequiredNoneTags; 
+}
+
+inline const FMassTagBitSet& FMassFragmentRequirements::GetRequiredOptionalTags() const
+{ 
+	return RequiredOptionalTags; 
+}
+
+inline const FMassChunkFragmentBitSet& FMassFragmentRequirements::GetRequiredAllChunkFragments() const
+{ 
+	return RequiredAllChunkFragments; 
+}
+
+inline const FMassChunkFragmentBitSet& FMassFragmentRequirements::GetRequiredOptionalChunkFragments() const
+{ 
+	return RequiredOptionalChunkFragments; 
+}
+
+inline const FMassChunkFragmentBitSet& FMassFragmentRequirements::GetRequiredNoneChunkFragments() const
+{ 
+	return RequiredNoneChunkFragments; 
+}
+
+inline const FMassSharedFragmentBitSet& FMassFragmentRequirements::GetRequiredAllSharedFragments() const
+{ 
+	return RequiredAllSharedFragments; 
+}
+
+inline const FMassSharedFragmentBitSet& FMassFragmentRequirements::GetRequiredOptionalSharedFragments() const
+{ 
+	return RequiredOptionalSharedFragments; 
+}
+
+inline const FMassSharedFragmentBitSet& FMassFragmentRequirements::GetRequiredNoneSharedFragments() const
+{ 
+	return RequiredNoneSharedFragments; 
+}
+
+inline const FMassConstSharedFragmentBitSet& FMassFragmentRequirements::GetRequiredAllConstSharedFragments() const
+{ 
+	return RequiredAllConstSharedFragments; 
+}
+
+inline const FMassConstSharedFragmentBitSet& FMassFragmentRequirements::GetRequiredOptionalConstSharedFragments() const
+{ 
+	return RequiredOptionalConstSharedFragments; 
+}
+
+inline const FMassConstSharedFragmentBitSet& FMassFragmentRequirements::GetRequiredNoneConstSharedFragments() const
+{ 
+	return RequiredNoneConstSharedFragments; 
+}
+
+inline bool FMassFragmentRequirements::IsInitialized() const 
+{ 
+	return bInitialized; 
+}
+
+inline bool FMassFragmentRequirements::HasPositiveRequirements() const 
+{ 
+	return bHasPositiveRequirements; 
+}
+
+inline bool FMassFragmentRequirements::HasNegativeRequirements() const 
+{ 
+	return bHasNegativeRequirements; 
+}
+
+inline bool FMassFragmentRequirements::HasOptionalRequirements() const 
+{ 
+	return bHasOptionalRequirements; 
+}
+
+inline bool FMassFragmentRequirements::DoesRequireGameThreadExecution() const
+{
+	return bRequiresGameThreadExecution;
+}
+
+FORCEINLINE void FMassFragmentRequirements::IncrementChangeCounter()
+{ 
+	++IncrementalChangesCount; 
+	bPropertiesCached = false;
+}
+
+inline void FMassFragmentRequirements::ConsumeIncrementalChangesCount()
+{
+	IncrementalChangesCount = 0;
+}
+
+inline bool FMassFragmentRequirements::HasIncrementalChanges() const
+{
+	return IncrementalChangesCount > 0;
+}
