[Mass] Implemented GetTypeHash for FMassSubsystemRequirements and FMassFragmentRequirements

diff --git a/Engine/Source/Runtime/MassEntity/Private/MassRequirements.cpp b/Engine/Source/Runtime/MassEntity/Private/MassRequirements.cpp
--- a/Engine/Source/Runtime/MassEntity/Private/MassRequirements.cpp
+++ b/Engine/Source/Runtime/MassEntity/Private/MassRequirements.cpp
@@ -106,6 +106,33 @@ FMassFragmentRequirements& FMassFragmentRequirements::ClearTagRequirements(const
 	return *this;
 }
 
+uint32 GetTypeHash(const FMassFragmentRequirements& Instance)
+{
+	// @todo consider calculating hash only for non-empty elements, or any other optimization
+	uint32 Hash = FCrc::MemCrc32(Instance.FragmentRequirements.GetData(), Instance.FragmentRequirements.Num() * sizeof(FMassFragmentRequirementDescription));
+	Hash = HashCombine(Hash, FCrc::MemCrc32(Instance.ChunkFragmentRequirements.GetData(), Instance.ChunkFragmentRequirements.Num() * sizeof(FMassFragmentRequirementDescription)));
+	Hash = HashCombine(Hash, FCrc::MemCrc32(Instance.ConstSharedFragmentRequirements.GetData(), Instance.ConstSharedFragmentRequirements.Num() * sizeof(FMassFragmentRequirementDescription)));
+	Hash = HashCombine(Hash, FCrc::MemCrc32(Instance.SharedFragmentRequirements.GetData(), Instance.SharedFragmentRequirements.Num() * sizeof(FMassFragmentRequirementDescription)));
+	Hash = HashCombine(Hash, GetTypeHash(Instance.RequiredAllTags));
+	Hash = HashCombine(Hash, GetTypeHash(Instance.RequiredAnyTags));
+	Hash = HashCombine(Hash, GetTypeHash(Instance.RequiredNoneTags));
+	Hash = HashCombine(Hash, GetTypeHash(Instance.RequiredOptionalTags));
+	Hash = HashCombine(Hash, GetTypeHash(Instance.RequiredAllFragments));
+	Hash = HashCombine(Hash, GetTypeHash(Instance.RequiredAnyFragments));
+	Hash = HashCombine(Hash, GetTypeHash(Instance.RequiredOptionalFragments));
+	Hash = HashCombine(Hash, GetTypeHash(Instance.RequiredNoneFragments));
+	Hash = HashCombine(Hash, GetTypeHash(Instance.RequiredAllChunkFragments));
+	Hash = HashCombine(Hash, GetTypeHash(Instance.RequiredOptionalChunkFragments));
+	Hash = HashCombine(Hash, GetTypeHash(Instance.RequiredNoneChunkFragments));
+	Hash = HashCombine(Hash, GetTypeHash(Instance.RequiredAllSharedFragments));
+	Hash = HashCombine(Hash, GetTypeHash(Instance.RequiredOptionalSharedFragments));
+	Hash = HashCombine(Hash, GetTypeHash(Instance.RequiredNoneSharedFragments));
+	Hash = HashCombine(Hash, GetTypeHash(Instance.RequiredAllConstSharedFragments));
+	Hash = HashCombine(Hash, GetTypeHash(Instance.RequiredOptionalConstSharedFragments));
+	Hash = HashCombine(Hash, GetTypeHash(Instance.RequiredNoneConstSharedFragments));
+	return HashCombine(Hash, PointerHash(Instance.CachedEntityManager.Get()));
+}
+
 void FMassFragmentRequirements::SortRequirements()
 {
 	// we're sorting the Requirements the same way ArchetypeData's FragmentConfig is sorted (see FMassArchetypeData::Initialize)
diff --git a/Engine/Source/Runtime/MassEntity/Public/MassRequirements.h b/Engine/Source/Runtime/MassEntity/Public/MassRequirements.h
--- a/Engine/Source/Runtime/MassEntity/Public/MassRequirements.h
+++ b/Engine/Source/Runtime/MassEntity/Public/MassRequirements.h
@@ -132,6 +132,8 @@ struct FMassSubsystemRequirements
 	bool DoesRequireGameThreadExecution() const;
 	MASSENTITY_API void ExportRequirements(FMassExecutionRequirements& OutRequirements) const;
 
+	friend uint32 GetTypeHash(const FMassSubsystemRequirements& Instance);
+
 protected:
 	FMassExternalSubsystemBitSet RequiredConstSubsystems;
 	FMassExternalSubsystemBitSet RequiredMutableSubsystems;
@@ -473,6 +475,8 @@ public:
 	bool DoesRequireGameThreadExecution() const;
 	MASSENTITY_API void ExportRequirements(FMassExecutionRequirements& OutRequirements) const;
 
+	MASSENTITY_API friend uint32 GetTypeHash(const FMassFragmentRequirements& Instance);
+
 protected:
 	MASSENTITY_API void SortRequirements();
 
@@ -586,6 +590,11 @@ inline bool FMassSubsystemRequirements::DoesRequireGameThreadExecution() const
 	return bRequiresGameThreadExecution;
 }
 
+inline uint32 GetTypeHash(const FMassSubsystemRequirements& Instance)
+{
+	return HashCombine(GetTypeHash(Instance.RequiredConstSubsystems), GetTypeHash(Instance.RequiredMutableSubsystems));
+}
+
 template<>
 FORCEINLINE FMassFragmentRequirements& FMassFragmentRequirements::AddTagRequirements<EMassFragmentPresence::All>(const FMassTagBitSet& TagBitSet)
 {
