[Mass]
Debugger Cleanup

diff --git a/Engine/Source/Runtime/CoreUObject/Private/StructUtils/StructTypeBitSet.cpp b/Engine/Source/Runtime/CoreUObject/Private/StructUtils/StructTypeBitSet.cpp
--- a/Engine/Source/Runtime/CoreUObject/Private/StructUtils/StructTypeBitSet.cpp
+++ b/Engine/Source/Runtime/CoreUObject/Private/StructUtils/StructTypeBitSet.cpp
@@ -293,11 +293,11 @@ void FStructTracker::Serialize(FArchive& Ar, FStructTypeBitSet::FBitSetContainer
 }
 
 #if WITH_STRUCTUTILS_DEBUG
-const UStruct* FStructTracker::DebugFindTypeByPartialName(const FString& PartialFragmentName) const
+const UStruct* FStructTracker::DebugFindTypeByPartialName(const FString& PartialName) const
 {
 	for (int32 TypeIndex = 0; TypeIndex < DebugStructTypeNamesList.Num(); ++TypeIndex)
 	{
-		if (DebugStructTypeNamesList[TypeIndex].ToString().Contains(PartialFragmentName))
+		if (DebugStructTypeNamesList[TypeIndex].ToString().Contains(PartialName))
 		{
 			return StructTypesList[TypeIndex].Get();
 		}
diff --git a/Engine/Source/Runtime/CoreUObject/Public/StructUtils/StructTypeBitSet.h b/Engine/Source/Runtime/CoreUObject/Public/StructUtils/StructTypeBitSet.h
--- a/Engine/Source/Runtime/CoreUObject/Public/StructUtils/StructTypeBitSet.h
+++ b/Engine/Source/Runtime/CoreUObject/Public/StructUtils/StructTypeBitSet.h
@@ -359,7 +359,7 @@ struct COREUOBJECT_API FStructTracker
 		DebugStructTypeNamesList.Reset();
 	}
 
-	[[nodiscard]] const UStruct* DebugFindTypeByPartialName(const FString& PartialFragmentName) const;
+	[[nodiscard]] const UStruct* DebugFindTypeByPartialName(const FString& PartialName) const;
 
 protected:
 	mutable TArray<FName, TInlineAllocator<64>> DebugStructTypeNamesList;
@@ -1147,9 +1147,9 @@ public:
 		TStructTrackerWrapper::StructTracker.DebugResetStructTypeMappingInfo();
 	}
 
-	[[nodiscard]] static const TUStructType* DebugFindTypeByPartialName(const FString& PartialFragmentName)
+	[[nodiscard]] static const TUStructType* DebugFindTypeByPartialName(const FString& PartialName)
 	{
-		return reinterpret_cast<const TUStructType*>(TStructTrackerWrapper::StructTracker.DebugFindTypeByPartialName(PartialFragmentName));
+		return reinterpret_cast<const TUStructType*>(TStructTrackerWrapper::StructTracker.DebugFindTypeByPartialName(PartialName));
 	}
 #endif // WITH_STRUCTUTILS_DEBUG
 
diff --git a/Engine/Source/Runtime/MassEntity/Private/MassDebugger.cpp b/Engine/Source/Runtime/MassEntity/Private/MassDebugger.cpp
--- a/Engine/Source/Runtime/MassEntity/Private/MassDebugger.cpp
+++ b/Engine/Source/Runtime/MassEntity/Private/MassDebugger.cpp
@@ -450,13 +450,13 @@ namespace UE::Mass::Debug
 					{
 						bEntityValid
 							? FMassDebugger::ClearFragmentWriteBreak(EntityManager, FragmentType, SelectedEntity)
-							: FMassDebugger::ClearAllFragmentWriteBreaks(EntityManager, FragmentType);
+							: FMassDebugger::ClearFragmentWriteBreak(EntityManager, FragmentType, FMassEntityHandle());
 					}
 					else
 					{
 						UE_LOG(LogConsoleResponse, Display, TEXT("Warning: Unable to find element type %s"), *PartialFragmentName);
 					}
-				}				
+				}
 			}
 		}
 	));
@@ -474,9 +474,7 @@ FMassDebugger::FOnEnvironmentEvent FMassDebugger::OnProcessorProviderRegistered;
 FMassDebugger::FOnDebugEvent FMassDebugger::OnDebugEvent;
 TArray<FMassDebugger::FEnvironment> FMassDebugger::ActiveEnvironments;
 UE::FSpinLock FMassDebugger::EntityManagerRegistrationLock;
-bool FMassDebugger::HasBreakPoint = false;
-FName FMassDebugger::SelectedFragment;
-FOnFragmentSelectionChanged FMassDebugger::OnFragmentSelectionChanged;
+bool FMassDebugger::bHasBreakpoint = false;
 TMap<FName, const UScriptStruct*> FMassDebugger::FragmentsByName;
 
 TConstArrayView<FMassEntityQuery*> FMassDebugger::GetProcessorQueries(const UMassProcessor& Processor)
@@ -771,48 +769,24 @@ void FMassDebugger::SelectEntity(const FMassEntityManager& EntityManager, const
 {
 	UE::Mass::Debug::SetDebugEntityRange(EntityHandle.Index, EntityHandle.Index);
 
-	const int32 Index = ActiveEnvironments.IndexOfByPredicate([WeakManager = EntityManager.AsWeak()](const FEnvironment& Element)
-	{
-		return Element.EntityManager == WeakManager;
-	});
-	if (ensure(Index != INDEX_NONE))
-	{
-		ActiveEnvironments[Index].SelectedEntity = EntityHandle;
-	}
+	GetActiveEnvironment(EntityManager).SelectedEntity = EntityHandle;
 
 	OnEntitySelectedDelegate.Broadcast(EntityManager, EntityHandle);
 }
 
 FMassEntityHandle FMassDebugger::GetSelectedEntity(const FMassEntityManager& EntityManager)
 {
-	const int32 Index = ActiveEnvironments.IndexOfByPredicate([WeakManager = EntityManager.AsWeak()](const FEnvironment& Element)
-	{
-		return Element.EntityManager == WeakManager;
-	});
-
-	return Index != INDEX_NONE ? ActiveEnvironments[Index].SelectedEntity : FMassEntityHandle();
+	return GetActiveEnvironment(EntityManager).SelectedEntity;
 }
 
 void FMassDebugger::HighlightEntity(const FMassEntityManager& EntityManager, const FMassEntityHandle EntityHandle)
 {
-	const int32 Index = ActiveEnvironments.IndexOfByPredicate([WeakManager = EntityManager.AsWeak()](const FEnvironment& Element)
-	{
-		return Element.EntityManager == WeakManager;
-	});
-	if (ensure(Index != INDEX_NONE))
-	{
-		ActiveEnvironments[Index].HighlightedEntity = EntityHandle;
-	}
+	GetActiveEnvironment(EntityManager).HighlightedEntity = EntityHandle;
 }
 
 FMassEntityHandle FMassDebugger::GetHighlightedEntity(const FMassEntityManager& EntityManager)
 {
-	const int32 Index = ActiveEnvironments.IndexOfByPredicate([WeakManager = EntityManager.AsWeak()](const FEnvironment& Element)
-	{
-		return Element.EntityManager == WeakManager;
-	});
-
-	return Index != INDEX_NONE ? ActiveEnvironments[Index].HighlightedEntity : FMassEntityHandle();
+	return GetActiveEnvironment(EntityManager).HighlightedEntity;
 }
 
 int32 FMassDebugger::RegisterEntityManager(FMassEntityManager& EntityManager)
@@ -880,173 +854,185 @@ bool FMassDebugger::DoesArchetypeMatchRequirements(const FMassArchetypeHandle& A
 
 bool FMassDebugger::ShouldProcessorBreak(const FMassEntityManager& EntityManager, const UMassProcessor* Processor, FMassEntityHandle Entity)
 {
-	if (LIKELY(!HasBreakPoint))
+	if (LIKELY(!bHasBreakpoint))
 	{
 		return false;
 	}
 
-	const int32 Index = ActiveEnvironments.IndexOfByPredicate([WeakManager = EntityManager.AsWeak()](const FEnvironment& Element)
-	{
-		return Element.EntityManager == WeakManager;
-	});
+	FEnvironment& ActiveEnvironment = GetActiveEnvironment(EntityManager);
 
-	if (LIKELY(!ActiveEnvironments[Index].HasBreakPoint))
+	if (LIKELY(!ActiveEnvironment.bHasBreakpoint))
 	{
 		return false;
 	}
 
-	const FMassEntityHandle* Found = ActiveEnvironments[Index].ProcessorBreakPoints.FindPair(Processor, Entity);
+	const FMassEntityHandle* Found = ActiveEnvironment.ProcessorBreakpoints.FindPair(Processor, Entity);
 	
 	return Found != nullptr;
 }
 
 bool FMassDebugger::HasAnyProcessorBreakpoints(const FMassEntityManager& EntityManager, const UMassProcessor* Processor)
 {
-	if (LIKELY(!HasBreakPoint))
+	if (LIKELY(!bHasBreakpoint))
 	{
 		return false;
 	}
 
-	const int32 Index = ActiveEnvironments.IndexOfByPredicate([WeakManager = EntityManager.AsWeak()](const FEnvironment& Element)
-	{
-		return Element.EntityManager == WeakManager;
-	});
+	FEnvironment& ActiveEnvironment = GetActiveEnvironment(EntityManager);
 
-	if (LIKELY(!ActiveEnvironments[Index].HasBreakPoint))
+	if (LIKELY(!ActiveEnvironment.bHasBreakpoint))
 	{
 		return false;
 	}
 
-	return ActiveEnvironments[Index].ProcessorBreakPoints.Contains(Processor);
+	return ActiveEnvironment.ProcessorBreakpoints.Contains(Processor);
 }
 
-bool FMassDebugger::ShouldFragmentWriteBreak(const FMassEntityManager& EntityManager, const UScriptStruct* FragmentType, FMassEntityHandle Entity)
+bool FMassDebugger::ShouldBreakOnFragmentWrite(const FMassEntityManager& EntityManager, const UScriptStruct* FragmentType, FMassEntityHandle Entity)
 {
-	if (LIKELY(!HasBreakPoint))
+	if (LIKELY(!bHasBreakpoint))
 	{
 		return false;
 	}
 
-	const int32 Index = ActiveEnvironments.IndexOfByPredicate([WeakManager = EntityManager.AsWeak()](const FEnvironment& Element)
-	{
-		return Element.EntityManager == WeakManager;
-	});
+	FEnvironment& ActiveEnvironment = GetActiveEnvironment(EntityManager);
 
-	if (LIKELY(!ActiveEnvironments[Index].HasBreakPoint))
+	if (LIKELY(!ActiveEnvironment.bHasBreakpoint))
 	{
 		return false;
 	}
 
-	const FMassEntityHandle* Found = ActiveEnvironments[Index].FragmentWriteBreakPoints.FindPair(FragmentType, Entity);
+	const FMassEntityHandle* Found = ActiveEnvironment.FragmentWriteBreakpoints.FindPair(FragmentType, Entity);
 
 	return Found != nullptr;
 }
 
 bool FMassDebugger::HasAnyFragmentWriteBreakpoints(const FMassEntityManager& EntityManager, const UScriptStruct* FragmentType)
 {
-	if (LIKELY(!HasBreakPoint))
+	if (LIKELY(!bHasBreakpoint))
 	{
 		return false;
 	}
 
-	const int32 Index = ActiveEnvironments.IndexOfByPredicate([WeakManager = EntityManager.AsWeak()](const FEnvironment& Element)
-	{
-		return Element.EntityManager == WeakManager;
-	});
+	FEnvironment& ActiveEnvironment = GetActiveEnvironment(EntityManager);
 
-	if (LIKELY(!ActiveEnvironments[Index].HasBreakPoint))
+	if (LIKELY(!ActiveEnvironment.bHasBreakpoint))
 	{
 		return false;
 	}
 
 	if (FragmentType == nullptr)
 	{
-		return ActiveEnvironments[Index].FragmentWriteBreakPoints.Num() > 0;
+		return ActiveEnvironment.FragmentWriteBreakpoints.Num() > 0;
 	}
 
-	return ActiveEnvironments[Index].FragmentWriteBreakPoints.Contains(FragmentType);
+	return ActiveEnvironment.FragmentWriteBreakpoints.Contains(FragmentType);
 }
 
-void FMassDebugger::SetProcessorBreakPoint(const FMassEntityManager& EntityManager, const UMassProcessor* Processor, FMassEntityHandle Entity)
+void FMassDebugger::SetProcessorBreakpoint(const FMassEntityManager& EntityManager, TNotNull<const UMassProcessor*> Processor, FMassEntityHandle Entity)
 {
-	const int32 Index = ActiveEnvironments.IndexOfByPredicate([WeakManager = EntityManager.AsWeak()](const FEnvironment& Element)
-	{
-		return Element.EntityManager == WeakManager;
-	});
+	FEnvironment& ActiveEnvironment = GetActiveEnvironment(EntityManager);
 
-	ActiveEnvironments[Index].HasBreakPoint = true;
-	HasBreakPoint = true;
+	ActiveEnvironment.bHasBreakpoint = true;
+	bHasBreakpoint = true;
 
-	ActiveEnvironments[Index].ProcessorBreakPoints.Add(Processor, Entity);
+	ActiveEnvironment.ProcessorBreakpoints.Add(Processor, Entity);
 }
 
 void FMassDebugger::SetFragmentWriteBreak(const FMassEntityManager& EntityManager, const UScriptStruct* FragmentType, FMassEntityHandle Entity)
 {
-	const int32 Index = ActiveEnvironments.IndexOfByPredicate([WeakManager = EntityManager.AsWeak()](const FEnvironment& Element)
-	{
-		return Element.EntityManager == WeakManager;
-	});
+	FEnvironment& ActiveEnvironment = GetActiveEnvironment(EntityManager);
 
-	ActiveEnvironments[Index].HasBreakPoint = true;
-	HasBreakPoint = true;
+	ActiveEnvironment.bHasBreakpoint = true;
+	bHasBreakpoint = true;
 
-	ActiveEnvironments[Index].FragmentWriteBreakPoints.Add(FragmentType, Entity);
+	ActiveEnvironment.FragmentWriteBreakpoints.Add(FragmentType, Entity);
 }
 
-void FMassDebugger::ClearProcessorBreakPoint(const FMassEntityManager& EntityManager, const UMassProcessor* Processor, FMassEntityHandle Entity)
+void FMassDebugger::ClearProcessorBreakpoint(const FMassEntityManager& EntityManager, const UMassProcessor* Processor, FMassEntityHandle Entity)
 {
-	const int32 Index = ActiveEnvironments.IndexOfByPredicate([WeakManager = EntityManager.AsWeak()](const FEnvironment& Element)
-	{
-		return Element.EntityManager == WeakManager;
-	});
+	FEnvironment& ActiveEnvironment = GetActiveEnvironment(EntityManager);
+
+	ActiveEnvironment.ProcessorBreakpoints.Remove(Processor, Entity);
 
-	ActiveEnvironments[Index].ProcessorBreakPoints.Remove(Processor, Entity);
 	UpdateHasBreakpoint();
 }
 
+void FMassDebugger::ClearAllProcessorBreakpoints(const FMassEntityManager& EntityManager, const UMassProcessor* Processor)
+{
+	FEnvironment& ActiveEnvironment = GetActiveEnvironment(EntityManager);
+	ActiveEnvironment.ProcessorBreakpoints.Remove(Processor);
+}
+
 void FMassDebugger::ClearFragmentWriteBreak(const FMassEntityManager& EntityManager, const UScriptStruct* FragmentType, FMassEntityHandle Entity)
 {
-	const int32 Index = ActiveEnvironments.IndexOfByPredicate([WeakManager = EntityManager.AsWeak()](const FEnvironment& Element)
+	FEnvironment& ActiveEnvironment = GetActiveEnvironment(EntityManager);
+
+	if (FragmentType)
 	{
-		return Element.EntityManager == WeakManager;
-	});
+		if (Entity.IsSet())
+		{
+			ActiveEnvironment.FragmentWriteBreakpoints.Remove(FragmentType, Entity);
+		}
+		else
+		{
+			ActiveEnvironment.FragmentWriteBreakpoints.Remove(FragmentType);
+		}
+	}
+	else
+	{
+		if (Entity.IsSet())
+		{
+			for (TPair<const UScriptStruct*, FMassEntityHandle>& Pair : ActiveEnvironment.FragmentWriteBreakpoints)
+			{
+				ActiveEnvironment.FragmentWriteBreakpoints.Remove(Pair.Key, Entity);
+			}
+		}
+		else
+		{
+			ActiveEnvironment.FragmentWriteBreakpoints.Reset();
+		}
+	}
 
-	ActiveEnvironments[Index].FragmentWriteBreakPoints.Remove(FragmentType, Entity);
 	UpdateHasBreakpoint();
 }
 
-void FMassDebugger::ClearProcessorBreakPoints(const FMassEntityManager& EntityManager, TNotNull<const UMassProcessor*> Processor)
+void FMassDebugger::ClearAllFragmentWriteBreak(const FMassEntityManager& EntityManager, const UScriptStruct* FragmentType)
 {
-	const int32 Index = ActiveEnvironments.IndexOfByPredicate([WeakManager = EntityManager.AsWeak()](const FEnvironment& Element)
-	{
-		return Element.EntityManager == WeakManager;
-	});
-
-	ActiveEnvironments[Index].ProcessorBreakPoints.Remove(Processor);
-	UpdateHasBreakpoint();
+	FEnvironment& ActiveEnvironment = GetActiveEnvironment(EntityManager);
+	ActiveEnvironment.FragmentWriteBreakpoints.Remove(FragmentType);
 }
 
-void FMassDebugger::ClearAllFragmentWriteBreaks(const FMassEntityManager& EntityManager, TNotNull<const UScriptStruct*> FragmentType)
+void FMassDebugger::ClearAllEntityBreakpoints(const FMassEntityManager& EntityManager, FMassEntityHandle Entity)
 {
-	const int32 Index = ActiveEnvironments.IndexOfByPredicate([WeakManager = EntityManager.AsWeak()](const FEnvironment& Element)
+	FEnvironment& ActiveEnvironment = GetActiveEnvironment(EntityManager);
+	TMultiMap<const UMassProcessor*, FMassEntityHandle>::TIterator ProccessorIterator = ActiveEnvironment.ProcessorBreakpoints.CreateIterator();
+
+	for (; ProccessorIterator; ++ProccessorIterator)
 	{
-		return Element.EntityManager == WeakManager;
-	});
+		if (ProccessorIterator.Value() == Entity)
+		{
+			ProccessorIterator.RemoveCurrent();
+		}
+	}
 
-	ActiveEnvironments[Index].FragmentWriteBreakPoints.Remove(FragmentType);
-	UpdateHasBreakpoint();	
-} 
+	TMultiMap<const UScriptStruct*, FMassEntityHandle>::TIterator FragmentIterator = ActiveEnvironment.FragmentWriteBreakpoints.CreateIterator();
+	for (; FragmentIterator; ++FragmentIterator)
+	{
+		if (FragmentIterator.Value() == Entity)
+		{
+			FragmentIterator.RemoveCurrent();
+		}
+	}
+}
  
-void FMassDebugger::BreakOnSelectedFragmentWrite()
+void FMassDebugger::BreakOnFragmentWriteForSelectedEntity(FName FragmentName)
 {
-	if (!SelectedFragment.IsNone())
+	for (FEnvironment& Environment : ActiveEnvironments)
 	{
-		for (FEnvironment& Environment : ActiveEnvironments)
+		if (Environment.EntityManager.Pin()->IsEntityValid(Environment.SelectedEntity))
 		{
-			if (Environment.EntityManager.Pin()->IsEntityValid(Environment.SelectedEntity))
-			{
-				SetFragmentWriteBreak(*Environment.EntityManager.Pin(), GetFragmentTypeFromName(SelectedFragment), Environment.SelectedEntity);
-			}
+			SetFragmentWriteBreak(*Environment.EntityManager.Pin(), GetFragmentTypeFromName(FragmentName), Environment.SelectedEntity);
 		}
 	}
 }
@@ -1055,11 +1041,9 @@ void FMassDebugger::ClearAllBreakpoints()
 {
 	for (FEnvironment& Environment : ActiveEnvironments)
 	{
-		Environment.ProcessorBreakPoints.Reset();
-		Environment.FragmentWriteBreakPoints.Reset();
-		Environment.HasBreakPoint = false;
+		Environment.ClearBreakpoints();
 	}
-	HasBreakPoint = false;
+	bHasBreakpoint = false;
 }
 
 const UScriptStruct* FMassDebugger::GetFragmentTypeFromName(FName FragmentName)
@@ -1130,46 +1114,33 @@ const UScriptStruct* FMassDebugger::GetFragmentTypeFromName(FName FragmentName)
 
 void FMassDebugger::UpdateHasBreakpoint()
 {
-	HasBreakPoint = false;
+	bHasBreakpoint = false;
 	for (FEnvironment& Environment : ActiveEnvironments)
 	{
-		Environment.HasBreakPoint = Environment.ProcessorBreakPoints.Num() != 0
-			|| Environment.FragmentWriteBreakPoints.Num() != 0;
-		HasBreakPoint |= Environment.HasBreakPoint;
+		Environment.bHasBreakpoint = Environment.ProcessorBreakpoints.Num() != 0
+			|| Environment.FragmentWriteBreakpoints.Num() != 0;
+		bHasBreakpoint |= Environment.bHasBreakpoint;
 	}
 }
 
-void FMassDebugger::SelectProcessor(const FMassEntityManager& EntityManager, const UMassProcessor* SelectedProcessor)
+FMassDebugger::FEnvironment& FMassDebugger::GetActiveEnvironment(const FMassEntityManager& EntityManager)
 {
 	const int32 Index = ActiveEnvironments.IndexOfByPredicate([WeakManager = EntityManager.AsWeak()](const FEnvironment& Element)
 	{
 		return Element.EntityManager == WeakManager;
 	});
 
-	if (Index >= 0)
-	{
-		if (ActiveEnvironments[Index].SelectedProcessor.Get() != SelectedProcessor)
-		{
-			ActiveEnvironments[Index].SelectedProcessor = SelectedProcessor;
-			ActiveEnvironments[Index].OnProcessorSelectionChanged.Broadcast(SelectedProcessor);
-		}
-	}
+	checkf(Index != INDEX_NONE, TEXT("Mass Debug Environment not found for specified EntitManager"));
+
+	return ActiveEnvironments[Index];
 }
 
-FOnProcessorSelectionChanged FMassDebugger::OnProcessorSelectionChanged(const FMassEntityManager& EntityManager)
+void FMassDebugger::FEnvironment::ClearBreakpoints()
 {
-	const int32 Index = ActiveEnvironments.IndexOfByPredicate([WeakManager = EntityManager.AsWeak()](const FEnvironment& Element)
-	{
-		return Element.EntityManager == WeakManager;
-	});
-
-	if (Index >= 0)
-	{
-		return ActiveEnvironments[Index].OnProcessorSelectionChanged;
-	}
-
-	static FOnProcessorSelectionChanged DefaultDelegate;
-	return DefaultDelegate;
+	ProcessorBreakpoints.Reset();
+	FragmentWriteBreakpoints.Reset();
+	bHasBreakpoint = false;
 }
 
+
 #endif // WITH_MASSENTITY_DEBUG
\ No newline at end of file
diff --git a/Engine/Source/Runtime/MassEntity/Private/MassExecutionContext.cpp b/Engine/Source/Runtime/MassEntity/Private/MassExecutionContext.cpp
--- a/Engine/Source/Runtime/MassEntity/Private/MassExecutionContext.cpp
+++ b/Engine/Source/Runtime/MassEntity/Private/MassExecutionContext.cpp
@@ -32,7 +32,7 @@ FMassExecutionContext::FMassExecutionContext(const FMassExecutionContext& Other)
 FMassExecutionContext::FMassExecutionContext(const FMassExecutionContext& Other, FMassEntityQuery& Query, const TSharedPtr<FMassCommandBuffer>& InCommandBuffer)
 	: FMassExecutionContext(Other)
 {
-	ensure(Other.QueriesStack.Last().Query == &Query);
+	ensureMsgf(Other.QueriesStack.Last().Query == &Query, TEXT("Creating a single-query execution context but the Query doesn't match the source context."));
 	QueriesStack.Add(Other.QueriesStack.Last());
 	SetDeferredCommandBuffer(InCommandBuffer);
 }
@@ -130,6 +130,7 @@ void FMassExecutionContext::PushQuery(FMassEntityQuery& InQuery)
 #if WITH_MASSENTITY_DEBUG
 	// check if this could possibly trigger a break before iterating to avoid extraneous breakpoint checks
 	FMassEntityManager& EntityManagerRef = GetEntityManagerChecked();
+
 	RuntimeData.bCheckProcessorBreaks = FMassDebugger::HasAnyProcessorBreakpoints(EntityManagerRef, DebugGetProcessor());
 	
 	if (UNLIKELY(FMassDebugger::HasAnyFragmentWriteBreakpoints(EntityManagerRef)))
@@ -141,9 +142,10 @@ void FMassExecutionContext::PushQuery(FMassEntityQuery& InQuery)
 					if (Req.AccessMode == EMassFragmentAccess::ReadWrite &&
 						FMassDebugger::HasAnyFragmentWriteBreakpoints(EntityManagerRef, Req.StructType))
 					{
-						if (RuntimeData.BreakFragmentsCount < FQueryTransientRuntime::MaxBreakFragmentCount)
+						if (ensureMsgf(RuntimeData.BreakFragmentsCount < FQueryTransientRuntime::MaxFragmentBreakpointCount, 
+							TEXT("Fragment write breakpoint count limit exceeded for this query.")))
 						{
-							RuntimeData.BreakFragments[RuntimeData.BreakFragmentsCount++] = Req.StructType;
+							RuntimeData.FragmentTypesToBreakOn[RuntimeData.BreakFragmentsCount++] = Req.StructType;
 						}
 					}
 				}
@@ -171,7 +173,7 @@ void FMassExecutionContext::PopQuery(const FMassEntityQuery& InQuery)
 
 FMassExecutionContext::FEntityIterator FMassExecutionContext::CreateEntityIterator()
 {
-	if (!ensureMsgf(QueriesStack.Num(), TEXT("")))
+	if (!ensureMsgf(QueriesStack.Num(), TEXT("Attempting to create an Entity Iterator when no entity query is being executed.")))
 	{
 		return FEntityIterator(*this);
 	}
@@ -202,7 +204,7 @@ FMassExecutionContext::FEntityIterator::FEntityIterator(FMassExecutionContext& I
 FMassExecutionContext::FEntityIterator::FEntityIterator(FMassExecutionContext& InExecutionContext, FQueryTransientRuntime& InQueryRuntime)
 	: ExecutionContext(InExecutionContext)
 	, QueryRuntime(InQueryRuntime)
-	, EntitiesNum(InExecutionContext.GetNumEntities())
+	, NumEntities(InExecutionContext.GetNumEntities())
 	, SerialNumber(InQueryRuntime.IteratorSerialNumber)
 {
 	this->operator++();
@@ -211,7 +213,7 @@ FMassExecutionContext::FEntityIterator::FEntityIterator(FMassExecutionContext& I
 #if WITH_MASSENTITY_DEBUG
 
 UE_DISABLE_OPTIMIZATION_SHIP
-void FMassExecutionContext::FEntityIterator::TestBreakPoints()
+void FMassExecutionContext::FEntityIterator::TestBreakpoints()
 {
 	FMassEntityManager& EntityManagerRef = ExecutionContext.GetEntityManagerChecked();
 	FMassEntityHandle Entity = ExecutionContext.GetEntity(EntityIndex);
@@ -229,16 +231,16 @@ void FMassExecutionContext::FEntityIterator::TestBreakPoints()
 
 		if (bDisableThisBreakpoint)
 		{
-			FMassDebugger::ClearProcessorBreakPoint(EntityManagerRef, ExecutionContext.DebugGetProcessor(), Entity);
+			FMassDebugger::ClearProcessorBreakpoint(EntityManagerRef, ExecutionContext.DebugGetProcessor(), Entity);
 		}
 
 		// bailing out, no point to hit multiple breakpoints for the given entity/processor pair
 		return;
 	}
 
-	for (const UScriptStruct* Fragment : QueryRuntime.BreakFragments)
+	for (const UScriptStruct* Fragment : QueryRuntime.FragmentTypesToBreakOn)
 	{
-		if (FMassDebugger::ShouldFragmentWriteBreak(EntityManagerRef, Fragment, Entity))
+		if (FMassDebugger::ShouldBreakOnFragmentWrite(EntityManagerRef, Fragment, Entity))
 		{
 			bool bDisableThisBreakpoint = false;
 			//====================================================================
diff --git a/Engine/Source/Runtime/MassEntity/Private/MassProcessor.cpp b/Engine/Source/Runtime/MassEntity/Private/MassProcessor.cpp
--- a/Engine/Source/Runtime/MassEntity/Private/MassProcessor.cpp
+++ b/Engine/Source/Runtime/MassEntity/Private/MassProcessor.cpp
@@ -126,6 +126,24 @@ void UMassProcessor::CallInitialize(const TNotNull<UObject*> Owner, const TShare
 	if (ensure(HasAnyFlags(RF_ClassDefaultObject) == false 
 		&& GetClass()->HasAnyClassFlags(CLASS_Abstract) == false))
 	{
+
+#if WITH_MASSENTITY_DEBUG
+		if (EntityManager.Get().DebugHasAllDebugFeatures(FMassEntityManager::EDebugFeatures::TraceProcessors))
+		{
+			DebugDescription = *GetProcessorName();
+			FString NetMode = EntityManager.Get().GetWorld() ? ToString(EntityManager.Get().GetWorld()->GetNetMode()) : TEXT("None");
+			//                       DebugDescription       " ("  NetMode        ")"
+			DebugDescription.Reserve(DebugDescription.Len() + 2 + NetMode.Len() + 1);
+			DebugDescription.Append(TEXT(" ("));
+			DebugDescription.Append(*NetMode);
+			DebugDescription.Append(TEXT(")"));
+		}
+		else
+		{
+			DebugDescription = FString::Printf(TEXT("%s (%s)"), *GetProcessorName(), EntityManager.Get().GetWorld() ? *ToString(EntityManager.Get().GetWorld()->GetNetMode()) : TEXT("No World"));
+		}
+#endif
+
 		for (FMassEntityQuery* Query : OwnedQueries)
 		{
 			// we should never get nulls here since OwnedQueries is private and the only way to
@@ -235,18 +253,7 @@ void UMassProcessor::CallExecute(FMassEntityManager& EntityManager, FMassExecuti
 		// Not using a more specific scope by default (i.e. LLM_SCOPE_BYNAME(*StatId)) since LLM is more strict regarding the provided string (no spaces or '_')
 
 #if WITH_MASSENTITY_DEBUG
-		Context.DebugSetExecutionDesc(FString::Printf(TEXT("%s (%s)"), *GetProcessorName(), EntityManager.GetWorld() ? *ToString(EntityManager.GetWorld()->GetNetMode()) : TEXT("No World")));
-		if (EntityManager.DebugHasAllDebugFeatures(FMassEntityManager::EDebugFeatures::TraceProcessors))
-		{
-			FString DebugDescription = *GetProcessorName();
-			FString NetMode = EntityManager.GetWorld() ? ToString(EntityManager.GetWorld()->GetNetMode()) : TEXT("None");
-			//                       DebugDescription       " ("  NetMode        ")"
-			DebugDescription.Reserve(DebugDescription.Len() + 2 + NetMode.Len() + 1);
-			DebugDescription.Append(TEXT(" ("));
-			DebugDescription.Append(*NetMode);
-			DebugDescription.Append(TEXT(")"));
-			Context.DebugSetExecutionDesc(MoveTemp(DebugDescription));
-		}
+		Context.DebugSetExecutionDesc(DebugDescription);
 		Context.DebugSetProcessor(this);
 #endif
 		// CacheSubsystemRequirements will return true only if all requirements declared with ProcessorRequirements are met
diff --git a/Engine/Source/Runtime/MassEntity/Public/MassDebugger.h b/Engine/Source/Runtime/MassEntity/Public/MassDebugger.h
--- a/Engine/Source/Runtime/MassEntity/Public/MassDebugger.h
+++ b/Engine/Source/Runtime/MassEntity/Public/MassDebugger.h
@@ -25,9 +25,6 @@ enum class EMassFragmentPresence : uint8;
 #endif // WITH_MASSENTITY_DEBUG
 #include "MassDebugger.generated.h"
 
-DECLARE_MULTICAST_DELEGATE_OneParam(FOnFragmentSelectionChanged, FName);
-DECLARE_MULTICAST_DELEGATE_OneParam(FOnProcessorSelectionChanged, const UMassProcessor*);
-
 namespace UE::Mass::Debug
 {
 	struct FArchetypeStats
@@ -142,20 +139,21 @@ struct MASSENTITY_API FMassDebugger
 		TWeakPtr<const FMassEntityManager> EntityManager;
 		TMap<FName, UE::Mass::Debug::FProcessorProviderFunction> ProcessorProviders;
 		FMassEntityHandle SelectedEntity;
+		// js@todo HighlightedEntity will be used to actually highlight the entity (arrow at transform, etc) in the near future
 		FMassEntityHandle HighlightedEntity;
-		TWeakObjectPtr<const UMassProcessor> SelectedProcessor;
-		FOnProcessorSelectionChanged OnProcessorSelectionChanged;
 
-		bool HasBreakPoint = false;
+		bool bHasBreakpoint = false;
 
-		TMultiMap<const UMassProcessor*, FMassEntityHandle> ProcessorBreakPoints;
-		TMultiMap<const UScriptStruct*, FMassEntityHandle> FragmentWriteBreakPoints;
+		TMultiMap<const UMassProcessor*, FMassEntityHandle> ProcessorBreakpoints;
+		TMultiMap<const UScriptStruct*, FMassEntityHandle> FragmentWriteBreakpoints;
 
 		explicit FEnvironment(const FMassEntityManager& InEntityManager)
 			: EntityManager(InEntityManager.AsWeak())
 		{}
 
 		bool IsValid() const { return EntityManager.IsValid(); }
+
+		void ClearBreakpoints();
 	};
 
 	DECLARE_TS_MULTICAST_DELEGATE_TwoParams(FOnEntitySelected, const FMassEntityManager&, const FMassEntityHandle);
@@ -239,88 +237,85 @@ struct MASSENTITY_API FMassDebugger
 	static bool DoesArchetypeMatchRequirements(const FMassArchetypeHandle& ArchetypeHandle, const FMassFragmentRequirements& Requirements, FOutputDevice& OutputDevice);
 
 	/**
-	* Checks if a processor should break on execute for a given entity.
-	*/
+	 * Checks if a processor should break on execute for a given entity.
+	 */
 	static bool ShouldProcessorBreak(const FMassEntityManager& EntityManager, const UMassProcessor* Processor, FMassEntityHandle Entity);
 
 	/**
-	* Checks if a processor has any breakpoints set for any entity.
-	*/
+	 * Checks if a processor has any breakpoints set for any entity.
+	 */
 	static bool HasAnyProcessorBreakpoints(const FMassEntityManager& EntityManager, const UMassProcessor* Processor);
 
 	/**
-	* Checks if a break shoud be triggered for a processor that's about to write a given fragment on an entity.
-	*/
-	static bool ShouldFragmentWriteBreak(const FMassEntityManager& EntityManager, const UScriptStruct* FragmentType, FMassEntityHandle Entity);
+	 * Checks if a break should be triggered for a processor that's about to write a given fragment on an entity.
+	 */
+	static bool ShouldBreakOnFragmentWrite(const FMassEntityManager& EntityManager, const UScriptStruct* FragmentType, FMassEntityHandle Entity);
 
 	/**
-	* Checks if there are any breakpoints set for writing a fragment for any entity
-	* Use FragmentType = nullptr (default) to check for ANY fragment breaks.
-	*/
+	 * Checks if there are any breakpoints set for writing a fragment for any entity
+	 * Use FragmentType = nullptr (default) to check for ANY fragment types.
+	 */
 	static bool HasAnyFragmentWriteBreakpoints(const FMassEntityManager& EntityManager, const UScriptStruct* FragmentType = nullptr);
 
 	/**
-	* Set a break to be triggered on processor execute for an entity.
-	*/
-	static void SetProcessorBreakPoint(const FMassEntityManager& EntityManager, const UMassProcessor* Processor, FMassEntityHandle Entity);
+	 * Sets a break to be triggered on processor execute for an entity.
+	 */
+	static void SetProcessorBreakpoint(const FMassEntityManager& EntityManager, TNotNull<const UMassProcessor*>, FMassEntityHandle Entity);
 
 	/**
-	* Set a break to be triggered for a processor that's about to write a given fragment on an entity.
-	*/
+	 * Sets a break to be triggered for a processor that's about to write a given fragment on an entity.
+	 */
 	static void SetFragmentWriteBreak(const FMassEntityManager& EntityManager, const UScriptStruct* FragmentType, FMassEntityHandle Entity);
 
 	/**
-	* Clear a break triggered on processor execute for an entity.
-	*/
-	static void ClearProcessorBreakPoint(const FMassEntityManager& EntityManager, const UMassProcessor* Processor, FMassEntityHandle Entity);
+	 * Clears a breakpoint triggered on processor execute for an entity.
+	 */
+	static void ClearProcessorBreakpoint(const FMassEntityManager& EntityManager, const UMassProcessor* Processor, FMassEntityHandle Entity);
+
+	/**
+	 * Clears all breakpoints set for a given processor.
+	 */
+	static void ClearAllProcessorBreakpoints(const FMassEntityManager& EntityManager, const UMassProcessor* Processor);
 
 	/**
-	* Clear a break triggered for a processor that's about to write a given fragment on an entity.
-	*/
+	 * Clears a fragment write breakpoint for a given fragment type and entity.
+	 */
 	static void ClearFragmentWriteBreak(const FMassEntityManager& EntityManager, const UScriptStruct* FragmentType, FMassEntityHandle Entity);
+	
+	/**
+	 * Clears all write breakpoints set for a given fragment type.
+	 */
+	static void ClearAllFragmentWriteBreak(const FMassEntityManager& EntityManager, const UScriptStruct* FragmentType);
 
 	/**
-	 * Clear all breakpoints set for the given processor
+	 * Clears all breakpoints set for a given entity.
 	 */
-	static void ClearProcessorBreakPoints(const FMassEntityManager& EntityManager, TNotNull<const UMassProcessor*> Processor);
+	static void ClearAllEntityBreakpoints(const FMassEntityManager& EntityManager, FMassEntityHandle Entity);
 
 	/**
-	 * Clear all write breakpoints for the given element type
+	 * Sets a write breakpoint for the specified fragment on the selcected entity
+	 * @see SelectEntity
 	 */
-	static void ClearAllFragmentWriteBreaks(const FMassEntityManager& EntityManager, TNotNull<const UScriptStruct*> FragmentType);
+	static void BreakOnFragmentWriteForSelectedEntity(FName FragmentName);
 
-	static void BreakOnSelectedFragmentWrite();
+	/**
+	 * Gets the UScriptStruct type for fragment of the specified name.
+	 */
 	static const UScriptStruct* GetFragmentTypeFromName(FName FragmentName);
-	static void ClearAllBreakpoints();
 
 	/**
-	* Select a fragment name for highlight
-	*/
-	static void SelectFragment(FName FragmentName)
-	{
-		if (SelectedFragment != FragmentName)
-		{
-			SelectedFragment = FragmentName;
-			OnFragmentSelectionChanged.Broadcast(FragmentName);  // Broadcast when the fragment changes
-		}
-	}
-	static FName GetSelectedFragment() { return SelectedFragment; }
-	// Delegate to notify when the selected fragment changes
-	static FOnFragmentSelectionChanged& OnSelectedFragmentChanged() { return OnFragmentSelectionChanged; }
-
-	// Delegate to notify when the selected processor for the environment changes
-	static FOnProcessorSelectionChanged OnProcessorSelectionChanged(const FMassEntityManager& EntityManager);
-
-	static void SelectProcessor(const FMassEntityManager& EntityManager, const UMassProcessor* SelectedProcessor);
+	 * Clears all breakpoints in all environments.
+	 */
+	static void ClearAllBreakpoints();
 	
 private:
 	static TArray<FEnvironment> ActiveEnvironments;
 	static UE::FSpinLock EntityManagerRegistrationLock;
-	static void UpdateHasBreakpoint();
-	static bool HasBreakPoint;
-	static FName SelectedFragment;
-	static FOnFragmentSelectionChanged OnFragmentSelectionChanged;
+	static bool bHasBreakpoint;
 	static TMap<FName, const UScriptStruct*> FragmentsByName;
+
+	static void UpdateHasBreakpoint();
+	static FEnvironment& GetActiveEnvironment(const FMassEntityManager& EntityManager);
 };
 
 #else
diff --git a/Engine/Source/Runtime/MassEntity/Public/MassExecutionContext.h b/Engine/Source/Runtime/MassEntity/Public/MassExecutionContext.h
--- a/Engine/Source/Runtime/MassEntity/Public/MassExecutionContext.h
+++ b/Engine/Source/Runtime/MassEntity/Public/MassExecutionContext.h
@@ -12,6 +12,7 @@
 #include "MassArchetypeTypes.h"
 #include "MassEntityQuery.h"
 #include "MassSubsystemAccess.h"
+#include "MassProcessor.h"
 
 
 #define CHECK_IF_VALID(View, Type) \
@@ -76,20 +77,25 @@ private:
 		FMassExternalSubsystemBitSet ConstSubsystemsBitSet;
 		FMassExternalSubsystemBitSet MutableSubsystemsBitSet;
 #if WITH_MASSENTITY_DEBUG
-		// MaxBreakFragmentCount needs to be bigger than the greatest number of fragments a processor could write to that have a breakpoint set
-		static constexpr uint32 MaxBreakFragmentCount = 8;
-		TStaticArray<const UScriptStruct*, MaxBreakFragmentCount> BreakFragments;
+		/** MaxBreakFragmentCount needs to be bigger than the greatest number of fragments a query has a write requirement for to that can have a breakpoint set */
+		static constexpr uint32 MaxFragmentBreakpointCount = 8;
+		TStaticArray<const UScriptStruct*, MaxFragmentBreakpointCount> FragmentTypesToBreakOn;
 
 		bool bCheckProcessorBreaks = false;
 		int32 BreakFragmentsCount = 0;
 #endif // WITH_MASSENTITY_DEBUG
 
+		/** Serial number to ensure iterator consistency (subsequent calls to CreateEntityIterator should not pass equivelncy test) */
 		uint32 IteratorSerialNumber = 0;
 
+		/** Helper function to create an empty instance with a valid Query ptr */
 		static FQueryTransientRuntime& GetDummyInstance();
 	};
+
 	/** We usually expect the queries to go only a single layer deep, so 2 elements here should suffice most of the time */
 	TArray<FQueryTransientRuntime, TInlineAllocator<2>> QueriesStack;
+
+	/** Track the serial number for FEntityIterator creation */
 	uint32 IteratorSerialNumberGenerator = 0;
 
 #if WITH_MASSENTITY_DEBUG
@@ -97,8 +103,7 @@ private:
 
 	/** Currently executing processor, used for debugger breakpoint checking */
 	// js@todo make this more generic
-	// mz@todo this could be moved to FQueryTransientRuntime now 
-	UMassProcessor* DebugProcessor = nullptr;
+	TWeakObjectPtr<UMassProcessor> DebugProcessor;
 #endif // WITH_MASSENTITY_DEBUG
 	
 	/** Used to control when the context is allowed to flush commands collected in DeferredCommandBuffer. This mechanism 
@@ -137,7 +142,7 @@ public:
 	FMassExecutionContext(const FMassExecutionContext& Other, FMassEntityQuery& Query, const TSharedPtr<FMassCommandBuffer>& InCommandBuffer = {});
 	~FMassExecutionContext();
 
-	FMassEntityManager& GetEntityManagerChecked();
+	FMassEntityManager& GetEntityManagerChecked() const;
 	const TSharedRef<FMassEntityManager>& GetSharedEntityManager();
 
 #if WITH_MASSENTITY_DEBUG
@@ -145,7 +150,7 @@ public:
 	void DebugSetExecutionDesc(const FString& Description) { DebugExecutionDescription = Description; }
 	void DebugSetExecutionDesc(FString&& Description) { DebugExecutionDescription = MoveTemp(Description); }
 
-	UMassProcessor* DebugGetProcessor() const { return DebugProcessor; }
+	UMassProcessor* DebugGetProcessor() const { return DebugProcessor.Get(); }
 	void DebugSetProcessor(UMassProcessor* Processor) { DebugProcessor = Processor; }
 #endif
 
@@ -156,26 +161,30 @@ public:
 	void ApplyFragmentRequirements(const FMassEntityQuery& RequestingQuery);
 	void ClearFragmentViews(const FMassEntityQuery& RequestingQuery);
 
+	/**
+	 * Iterator to easily loop through entities in the current chunk.
+	 * Supports ranged for and can be used directly as an entity index for the current chunk.
+	 */
 	struct FEntityIterator
 	{
-		FORCEINLINE int operator*() const
+		FORCEINLINE int32 operator*() const
 		{
 			return EntityIndex;
 		}
 
-		FORCEINLINE bool operator!=(const int& other) const
+		FORCEINLINE bool operator!=(const int& Other) const
 		{
-			return EntityIndex != other;
+			return EntityIndex != Other;
 		}
 
-		FORCEINLINE operator int() const
+		FORCEINLINE operator int32() const
 		{
 			return EntityIndex;
 		}
 
 		FORCEINLINE operator bool() const
 		{
-			return SerialNumber && EntityIndex < EntitiesNum;
+			return SerialNumber && EntityIndex < NumEntities;
 		}
 
 		FORCEINLINE bool operator<(const int32 Other) const
@@ -188,11 +197,11 @@ public:
 			++EntityIndex;
 #if WITH_MASSENTITY_DEBUG
 			if (UNLIKELY(QueryRuntime.bCheckProcessorBreaks || QueryRuntime.BreakFragmentsCount != 0) 
-				&& EntityIndex < EntitiesNum)
+				&& EntityIndex < NumEntities)
 			{
-				TestBreakPoints();
+				TestBreakpoints();
 			}
-#endif //!WITH_MASSENTITY_DEBUG
+#endif //WITH_MASSENTITY_DEBUG
 			return *this;
 		}
 
@@ -211,7 +220,7 @@ public:
 		FEntityIterator end() const
 		{
 			FEntityIterator End = *this;
-			End.EntityIndex = EntitiesNum;
+			End.EntityIndex = NumEntities;
 			return End;
 		}
 
@@ -226,15 +235,20 @@ public:
 		
 
 #if WITH_MASSENTITY_DEBUG
-		MASSENTITY_API void TestBreakPoints();
+		MASSENTITY_API void TestBreakpoints();
 #endif //!WITH_MASSENTITY_DEBUG
 
-		FMassExecutionContext& ExecutionContext;
-		FQueryTransientRuntime& QueryRuntime;
+		const FMassExecutionContext& ExecutionContext;
+		const FQueryTransientRuntime& QueryRuntime;
 		int32 EntityIndex = INDEX_NONE;
-		const int32 EntitiesNum = INDEX_NONE;
+		const int32 NumEntities = INDEX_NONE;
 		const uint32 SerialNumber = 0;
 	};
+
+	/**
+	 * Creats and Entity Iterator for the current chunk.
+	 * Supports ranged for and can be used directly as an entity index for the current chunk.
+	 */
 	FEntityIterator CreateEntityIterator();
 
 	/** Sets bFlushDeferredCommands. Note that setting to True while the system is being executed doesn't result in
@@ -550,7 +564,7 @@ public:
 //-----------------------------------------------------------------------------
 // Inlines
 //-----------------------------------------------------------------------------
-inline FMassEntityManager& FMassExecutionContext::GetEntityManagerChecked()
+inline FMassEntityManager& FMassExecutionContext::GetEntityManagerChecked() const
 {
 	return EntityManager.Get();
 }
diff --git a/Engine/Source/Runtime/MassEntity/Public/MassProcessor.h b/Engine/Source/Runtime/MassEntity/Public/MassProcessor.h
--- a/Engine/Source/Runtime/MassEntity/Public/MassProcessor.h
+++ b/Engine/Source/Runtime/MassEntity/Public/MassProcessor.h
@@ -259,6 +259,10 @@ private:
 	 *  a member variable of a given processor */
 	TArray<FMassEntityQuery*> OwnedQueries;
 
+#if WITH_MASSENTITY_DEBUG
+	FString DebugDescription;
+#endif
+
 	//-----------------------------------------------------------------------------
 	// DEPRECATED
 	//-----------------------------------------------------------------------------
