Fix CIS Issue <12345>:
[Mass]
Fix and simplify ClearFragmentWriteBreak to just remove a single item.

diff --git a/Engine/Source/Runtime/MassEntity/Private/MassDebugger.cpp b/Engine/Source/Runtime/MassEntity/Private/MassDebugger.cpp
--- a/Engine/Source/Runtime/MassEntity/Private/MassDebugger.cpp
+++ b/Engine/Source/Runtime/MassEntity/Private/MassDebugger.cpp
@@ -967,33 +967,7 @@ void FMassDebugger::ClearAllProcessorBreakpoints(const FMassEntityManager& Entit
 void FMassDebugger::ClearFragmentWriteBreak(const FMassEntityManager& EntityManager, const UScriptStruct* FragmentType, FMassEntityHandle Entity)
 {
 	FEnvironment& ActiveEnvironment = GetActiveEnvironment(EntityManager);
-
-	if (FragmentType)
-	{
-		if (Entity.IsSet())
-		{
-			ActiveEnvironment.FragmentWriteBreakpoints.Remove(FragmentType, Entity);
-		}
-		else
-		{
-			ActiveEnvironment.FragmentWriteBreakpoints.Remove(FragmentType);
-		}
-	}
-	else
-	{
-		if (Entity.IsSet())
-		{
-			for (TPair<const UScriptStruct*, FMassEntityHandle>& Pair : ActiveEnvironment.FragmentWriteBreakpoints)
-			{
-				ActiveEnvironment.FragmentWriteBreakpoints.Remove(Pair.Key, Entity);
-			}
-		}
-		else
-		{
-			ActiveEnvironment.FragmentWriteBreakpoints.Reset();
-		}
-	}
-
+	ActiveEnvironment.FragmentWriteBreakpoints.Remove(FragmentType, Entity);
 	UpdateHasBreakpoint();
 }
 
