[Mass] Removed unnecessary MassEntityQuery.h inclusion from MassExecutionContext.h

#rnx

diff --git a/Engine/Source/Runtime/MassEntity/Public/MassExecutionContext.h b/Engine/Source/Runtime/MassEntity/Public/MassExecutionContext.h
--- a/Engine/Source/Runtime/MassEntity/Public/MassExecutionContext.h
+++ b/Engine/Source/Runtime/MassEntity/Public/MassExecutionContext.h
@@ -10,7 +10,6 @@
 #include "MassEntityQuery.h"
 #endif // UE_ENABLE_INCLUDE_ORDER_DEPRECATED_IN_5_6
 #include "MassArchetypeTypes.h"
-#include "MassEntityQuery.h"
 #include "MassSubsystemAccess.h"
 #include "MassProcessor.h"
 
@@ -280,7 +279,7 @@ public:
 
 	void ForEachEntityInChunk(const FMassEntityExecuteFunction& EntityExecuteFunction)
 	{
-		for (FMassExecutionContext::FEntityIterator EntityIterator = CreateEntityIterator(); EntityIterator; ++EntityIterator)
+		for (FEntityIterator EntityIterator = CreateEntityIterator(); EntityIterator; ++EntityIterator)
 		{
 			EntityExecuteFunction(*this, EntityIterator);
 		}
