[Mass] Added a function to MassEntityManager for checking whether a given handle is `Reserved`
Modified how Mass entity storage understands "state" a bit.

#rnx

diff --git a/Engine/Source/Runtime/MassEntity/Private/MassEntityManager.cpp b/Engine/Source/Runtime/MassEntity/Private/MassEntityManager.cpp
--- a/Engine/Source/Runtime/MassEntity/Private/MassEntityManager.cpp
+++ b/Engine/Source/Runtime/MassEntity/Private/MassEntityManager.cpp
@@ -779,8 +779,6 @@ TSharedRef<FMassEntityManager::FEntityCreationContext> FMassEntityManager::Batch
 	// to update the manager's information on these entities (stored in FMassEntityManager::Entities)
 	for (FMassEntityHandle& Entity : EntityHandles)
 	{
-		check(GetEntityStorageInterface().IsValidIndex(Entity.Index));
-
 		const UE::Mass::IEntityStorageInterface::EEntityState EntityState = GetEntityStorageInterface().GetEntityState(Entity.Index);
 		checkf(EntityState == UE::Mass::IEntityStorageInterface::EEntityState::Reserved, TEXT("Trying to build entities that are not reserved. Check all handles are reserved or consider using BatchCreateEntities"));
 
@@ -849,7 +847,6 @@ TSharedRef<FMassEntityManager::FEntityCreationContext> FMassEntityManager::Inter
 
 	for (FMassEntityHandle Entity : ReservedEntities)
 	{
-		check(IsEntityValid(Entity));
 		const UE::Mass::IEntityStorageInterface::EEntityState EntityState = GetEntityStorageInterface().GetEntityState(Entity.Index);
 		checkf(EntityState == UE::Mass::IEntityStorageInterface::EEntityState::Reserved, TEXT("Trying to build entities that are not reserved. Check all handles are reserved or consider using BatchCreateEntities"));
 		
@@ -1917,24 +1914,26 @@ const FSharedStruct* FMassEntityManager::InternalGetSharedFragmentPtr(FMassEntit
 bool FMassEntityManager::IsEntityValid(FMassEntityHandle Entity) const
 {
 	return (Entity.Index != UE::Mass::Private::InvalidEntityIndex) 
-		&& GetEntityStorageInterface().IsValidIndex(Entity.Index) 
-		&& (GetEntityStorageInterface().GetSerialNumber(Entity.Index) == Entity.SerialNumber);
+		&& GetEntityStorageInterface().GetEntityState(Entity) != UE::Mass::IEntityStorageInterface::EEntityState::Invalid;
 }
 
 bool FMassEntityManager::IsEntityBuilt(FMassEntityHandle Entity) const
 {
-	CheckIfEntityIsValid(Entity);
-	const UE::Mass::IEntityStorageInterface::EEntityState CurrentState = GetEntityStorageInterface().GetEntityState(Entity.Index);
-	return CurrentState == UE::Mass::IEntityStorageInterface::EEntityState::Created;
+	return GetEntityStorageInterface().GetEntityState(Entity.Index) == UE::Mass::IEntityStorageInterface::EEntityState::Created;
+}
+
+bool FMassEntityManager::IsEntityReserved(FMassEntityHandle EntityHandle) const
+{
+	return GetEntityStorageInterface().GetEntityState(EntityHandle) == UE::Mass::IEntityStorageInterface::EEntityState::Reserved;
 }
 
 void FMassEntityManager::CheckIfEntityIsValid(FMassEntityHandle Entity) const
 {
-	checkf(IsEntityValid(Entity), TEXT("Invalid entity (ID: %d, SN:%d, %s)"), Entity.Index, Entity.SerialNumber,
-		   (Entity.Index == 0) ? TEXT("was never initialized") : TEXT("already destroyed"));
+	checkf(IsEntityValid(Entity), TEXT("Invalid entity (ID: %d, SN:%d, %s)"), Entity.Index, Entity.SerialNumber
+		, (Entity.Index == 0) ? TEXT("was never initialized") : TEXT("already destroyed"));
 }
 
-void FMassEntityManager::CheckIfEntityIsActive(FMassEntityHandle Entity) const
+FORCEINLINE void FMassEntityManager::CheckIfEntityIsActive(FMassEntityHandle Entity) const
 {
 	checkf(IsEntityBuilt(Entity), TEXT("Entity not yet created(ID: %d, SN:%d)"), Entity.Index, Entity.SerialNumber);
 }
diff --git a/Engine/Source/Runtime/MassEntity/Private/MassEntityManagerStorage.cpp b/Engine/Source/Runtime/MassEntity/Private/MassEntityManagerStorage.cpp
--- a/Engine/Source/Runtime/MassEntity/Private/MassEntityManagerStorage.cpp
+++ b/Engine/Source/Runtime/MassEntity/Private/MassEntityManagerStorage.cpp
@@ -80,7 +80,14 @@ namespace UE::Mass
 				: EEntityState::Reserved;
 		}
 
-		return EEntityState::Free;	
+		return EEntityState::Invalid;
+	}
+
+	IEntityStorageInterface::EEntityState FSingleThreadedEntityStorage::GetEntityState(FMassEntityHandle EntityHandle) const
+	{
+		return IsValidHandle(EntityHandle)
+			? GetEntityState(EntityHandle.Index)
+			: EEntityState::Invalid;
 	}
 
 	int32 FSingleThreadedEntityStorage::GetSerialNumber(int32 Index) const
@@ -239,6 +246,7 @@ namespace UE::Mass
 		check(FMath::IsPowerOfTwo(InInitializationParams.MaxEntitiesPerPage));
 		check(FMath::IsPowerOfTwo(InInitializationParams.MaxEntityCount));
 		MaxEntitiesPerPage = InInitializationParams.MaxEntitiesPerPage;
+		MaxEntitiesPerPageLog2 = FMath::FloorLog2(MaxEntitiesPerPage);
 		MaximumEntityCountShift = FMath::FloorLog2(InInitializationParams.MaxEntityCount);
 		checkf(MaximumEntityCountShift < 32, TEXT("Invalid maximum entity count, cannot exceed 31 bits"));
 
@@ -293,13 +301,13 @@ namespace UE::Mass
 		LookupEntity(Index).CurrentArchetype = Archetype;
 	}
 
-	IEntityStorageInterface::EEntityState FConcurrentEntityStorage::GetEntityState(int32 Index) const
+	IEntityStorageInterface::EEntityState FConcurrentEntityStorage::GetEntityState(const int32 Index) const
 	{
 		//
 		// || Archetype || bIsAllocated || Result    |
-		//  |  nullptr   |      0       |  Free     |
-		//  |  nullptr   |      1       |  Reserved |
-		//  | !nullptr   |      1       |  Created  |
+		//  |  nullptr   |      0       |  Invalid   |
+		//  |  nullptr   |      1       |  Reserved  |
+		//  | !nullptr   |      1       |  Created   |
 		//
 	
 		const FEntityData& EntityData = LookupEntity(Index);
@@ -309,7 +317,21 @@ namespace UE::Mass
 		}
 		return EntityData.bIsAllocated
 			? EEntityState::Reserved
-			: EEntityState::Free;
+			: EEntityState::Invalid;
+	}
+
+	IEntityStorageInterface::EEntityState FConcurrentEntityStorage::GetEntityState(const FMassEntityHandle EntityHandle) const
+	{
+		if (EntityHandle.Index >= 0)
+		{
+			const uint32 PageIndex = static_cast<uint32>(EntityHandle.Index) >> MaxEntitiesPerPageLog2;
+			if (PageIndex < PageCount)
+			{
+				return GetEntityState(EntityHandle.Index);
+			}
+		}
+
+		return EEntityState::Invalid;
 	}
 
 	int32 FConcurrentEntityStorage::GetSerialNumber(int32 Index) const
@@ -322,7 +344,7 @@ namespace UE::Mass
 		// Page Index is which page in the array of pages we need to access
 		if (Index >= 0)
 		{
-			const uint32 PageIndex = static_cast<uint32>(Index) >> FMath::FloorLog2(MaxEntitiesPerPage);
+			const uint32 PageIndex = static_cast<uint32>(Index) >> MaxEntitiesPerPageLog2;
 			return PageIndex < PageCount;
 		}
 		return false;
@@ -596,7 +618,7 @@ namespace UE::Mass
 	{
 		check(Index >= 0);
 		// PageIndex is which Page in the array of pages we need to access
-		const uint32 PageIndex = static_cast<uint32>(Index) >> FMath::FloorLog2(MaxEntitiesPerPage);
+		const uint32 PageIndex = static_cast<uint32>(Index) >> MaxEntitiesPerPageLog2;
 
 		// Convert the entity index into the index with respect to the page
 		const uint32 EntityOffset = (PageIndex * MaxEntitiesPerPage);
diff --git a/Engine/Source/Runtime/MassEntity/Public/MassEntityManager.h b/Engine/Source/Runtime/MassEntity/Public/MassEntityManager.h
--- a/Engine/Source/Runtime/MassEntity/Public/MassEntityManager.h
+++ b/Engine/Source/Runtime/MassEntity/Public/MassEntityManager.h
@@ -496,22 +496,34 @@ public:
 
 	void BatchSetEntityFragmentValues(TConstArrayView<FMassArchetypeEntityCollection> EntityCollections, TArrayView<const FInstancedStruct> FragmentInstanceList);
 
-	// Return true if it is an valid built entity
+	/**
+	 * @return whether the given handle represents a valid and built entity
+	 *	 (i.e., the handle is valid and the entity represent has been constructed already)
+	 */
 	bool IsEntityActive(FMassEntityHandle EntityHandle) const 
 	{
 		return IsEntityValid(EntityHandle) && IsEntityBuilt(EntityHandle);
 	}
 
-	// Returns true if EntityHandle is valid
+	/**
+	 * @return whether the given entity handle is valid, i.e. it points
+	 *	 to a valid spot in the entity storage and the handle's serial number is up to date
+	 */
 	bool IsEntityValid(FMassEntityHandle EntityHandle) const;
 
-	// Returns true if EntityHandle is has been fully built (expecting a valid EntityHandle)
+	/** whether the entity handle represents an entity that has been fully built (expecting a valid EntityHandle) */
 	bool IsEntityBuilt(FMassEntityHandle EntityHandle) const;
 
-	// Asserts that IsEntityValid
+	/**
+	 * @return whether the given EntityHandle is valid and the entity it represents is in `Reserved` state
+	 *	 (i.e. it will also fail if the entity has already been `Created`)
+	 */
+	bool IsEntityReserved(FMassEntityHandle EntityHandle) const;
+
+	/** Asserts that IsEntityValid */
 	void CheckIfEntityIsValid(FMassEntityHandle EntityHandle) const;
 
-	// Asserts that IsEntityBuilt
+	/** Asserts that IsEntityBuilt */
 	void CheckIfEntityIsActive(FMassEntityHandle EntityHandle) const;
 
 	template<typename FragmentType>
diff --git a/Engine/Source/Runtime/MassEntity/Public/MassEntityManagerStorage.h b/Engine/Source/Runtime/MassEntity/Public/MassEntityManagerStorage.h
--- a/Engine/Source/Runtime/MassEntity/Public/MassEntityManagerStorage.h
+++ b/Engine/Source/Runtime/MassEntity/Public/MassEntityManagerStorage.h
@@ -54,12 +54,12 @@ namespace UE::Mass
 	public:
 		enum class EEntityState
 		{
-			/** Entity index refers to an entity that is free to be reserved or created */
-			Free,
+			/** Entity index is invalid or the serial number is out of date */
+			Invalid,
 			/** Entity index refers to a reserved entity */
 			Reserved,
 			/** Entity index refers to an entity assigned to an archetype */
-			Created
+			Created,
 		};
 		virtual ~IEntityStorageInterface() = default;
 
@@ -72,11 +72,17 @@ namespace UE::Mass
 		virtual void SetArchetypeFromShared(int32 Index, const TSharedPtr<FMassArchetypeData>& Archetype) = 0;
 
 		/**
-		 * Returns true if the given entity at index is currently reserved
-		 * False if free or assigned an archetype
+		 * @return EEntityState value indicating the current state of the entry indicated by the Index.
+		 *	Note that it's callers responsibility to ensure the Index is valid.
 		 */
 		virtual EEntityState GetEntityState(int32 Index) const = 0;
 
+		/**
+		 * @return EEntityState value indicating the validity of the handle, and if
+		 * 	the handle is valid, the state of the entry it indicates.
+		 */
+		virtual EEntityState GetEntityState(FMassEntityHandle EntityHandle) const = 0;
+
 		virtual int32 GetSerialNumber(int32 Index) const = 0;
 
 		/** Checks if index can be used to access entity data */
@@ -143,6 +149,7 @@ namespace UE::Mass
 		virtual void SetArchetypeFromShared(int32 Index, TSharedPtr<FMassArchetypeData>&) override;
 		virtual void SetArchetypeFromShared(int32 Index, const TSharedPtr<FMassArchetypeData>&) override;
 		virtual EEntityState GetEntityState(int32 Index) const override;
+		virtual EEntityState GetEntityState(FMassEntityHandle EntityHandle) const override;
 		virtual int32 GetSerialNumber(int32 Index) const override;
 		virtual bool IsValidIndex(int32 Index) const override;
 		virtual bool IsValidHandle(FMassEntityHandle EntityHandle) const override;
@@ -198,6 +205,7 @@ namespace UE::Mass
 		virtual void SetArchetypeFromShared(int32 Index, TSharedPtr<FMassArchetypeData>& Archetype) override;
 		virtual void SetArchetypeFromShared(int32 Index, const TSharedPtr<FMassArchetypeData>& Archetype) override;
 		virtual EEntityState GetEntityState(int32 Index) const override;
+		virtual EEntityState GetEntityState(FMassEntityHandle EntityHandle) const override;
 		virtual int32 GetSerialNumber(int32 Index) const override;
 		virtual bool IsValidIndex(int32 Index) const override;
 		virtual bool IsValidHandle(FMassEntityHandle EntityHandle) const override;
@@ -250,6 +258,7 @@ namespace UE::Mass
 		uint32 EntityCount = 0;
 		uint32 MaximumEntityCountShift = 0;
 		uint32 MaxEntitiesPerPage = 0;
+		uint32 MaxEntitiesPerPageLog2 = 0;
 		uint32 PageCount = 0;
 		/** ALWAYS acquire FreeListMutex before this one */
 		UE::FTransactionallySafeMutex PageAllocateMutex;
