[Mass] Minor Mass concepts massaging and fixing the consequences in FMassArchetypeCompositionDescriptor

#rnx

diff --git a/Engine/Source/Runtime/MassEntity/Public/MassEntityBuilder.h b/Engine/Source/Runtime/MassEntity/Public/MassEntityBuilder.h
--- a/Engine/Source/Runtime/MassEntity/Public/MassEntityBuilder.h
+++ b/Engine/Source/Runtime/MassEntity/Public/MassEntityBuilder.h
@@ -257,11 +257,13 @@ FEntityBuilder& FEntityBuilder::Add(TArgs&&... InArgs) requires (CElement<T> &&
 template<typename T, typename... TArgs>
 T& FEntityBuilder::Add_GetRef(TArgs&&... InArgs) requires (CElement<T> && !(CTag<T> || CChunkFragment<T>))
 {
+	using FElementType = UE::Mass::TElementType<T>;
+
 	State = EState::ReadyToCommit;
 
-	if (ensureMsgf(GetBitSetBuilder<T>().template Contains<T>() == false, TEXT("Element of type %s has already been added"), *T::StaticStruct()->GetName()))
+	if (ensureMsgf(GetBitSetBuilder<FElementType>().template Contains<T>() == false, TEXT("Element of type %s has already been added"), *T::StaticStruct()->GetName()))
 	{
-		GetBitSetBuilder<T>().template Add<T>();
+		GetBitSetBuilder<FElementType>().template Add<T>();
 		CachedArchetypeHandle = FMassArchetypeHandle();
 
 		return GetInstancedStructContainerInternal<T>().Add_GetRef(FInstancedStruct::Make<T>(InArgs...)).template GetMutable<T>();
@@ -276,10 +278,12 @@ T& FEntityBuilder::Add_GetRef(TArgs&&... InArgs) requires (CElement<T> && !(CTag
 template<typename T, typename... TArgs>
 T& FEntityBuilder::GetOrCreate(TArgs&&... InArgs) requires (CElement<T> && !(CTag<T> || CChunkFragment<T>))
 {
+	using FElementType = UE::Mass::TElementType<T>;
+
 	State = EState::ReadyToCommit;
 
 	FInstancedStruct* ElementFound = nullptr;
-	if (GetBitSetBuilder<T>().template Contains<T>())
+	if (GetBitSetBuilder<FElementType>().template Contains<T>())
 	{
 		// replace
 		ElementFound = GetInstancedStructContainerInternal<T>().FindByPredicate([](const FInstancedStruct& ExistingElement)
@@ -292,7 +296,7 @@ T& FEntityBuilder::GetOrCreate(TArgs&&... InArgs) requires (CElement<T> && !(CTa
 	}
 	else
 	{
-		GetBitSetBuilder<T>().template Add<T>();
+		GetBitSetBuilder<FElementType>().template Add<T>();
 		CachedArchetypeHandle = FMassArchetypeHandle();
 
 		ElementFound = &GetInstancedStructContainerInternal<T>().Add_GetRef(FInstancedStruct::Make<T>(InArgs...));
@@ -304,7 +308,9 @@ T& FEntityBuilder::GetOrCreate(TArgs&&... InArgs) requires (CElement<T> && !(CTa
 template<typename T>
 T* FEntityBuilder::Find() requires (CElement<T> && !(CTag<T> || CChunkFragment<T>))
 {
-	if (GetBitSetBuilder<T>().template Contains<T>())
+	using FElementType = UE::Mass::TElementType<T>;
+
+	if (GetBitSetBuilder<FElementType>().template Contains<T>())
 	{
 		FInstancedStruct* ElementFound = GetInstancedStructContainerInternal<T>().FindByPredicate([](const FInstancedStruct& Element)
 		{
diff --git a/Engine/Source/Runtime/MassEntity/Public/MassEntityConcepts.h b/Engine/Source/Runtime/MassEntity/Public/MassEntityConcepts.h
--- a/Engine/Source/Runtime/MassEntity/Public/MassEntityConcepts.h
+++ b/Engine/Source/Runtime/MassEntity/Public/MassEntityConcepts.h
@@ -7,24 +7,41 @@
 
 namespace UE::Mass
 {
-	template <typename T>
+	template<typename T>
 	concept CFragment = TIsDerivedFrom<typename TRemoveReference<T>::Type, FMassFragment>::Value;
 
-	template <typename T>
+	template<typename T>
 	concept CTag = TIsDerivedFrom<typename TRemoveReference<T>::Type, FMassTag>::Value;
 
-	template <typename T>
+	template<typename T>
 	concept CChunkFragment = TIsDerivedFrom<typename TRemoveReference<T>::Type, FMassChunkFragment>::Value;
 
-	template <typename T>
+	template<typename T>
 	concept CSharedFragment = TIsDerivedFrom<typename TRemoveReference<T>::Type, FMassSharedFragment>::Value;
 
-	template <typename T>
+	template<typename T>
 	concept CConstSharedFragment = TIsDerivedFrom<typename TRemoveReference<T>::Type, FMassConstSharedFragment>::Value;
 
-	template <typename T>
+	template<typename T>
 	concept CNonTag = CFragment<T> || CChunkFragment<T> || CSharedFragment<T> || CConstSharedFragment<T>;
 
-	template <typename T>
+	template<typename T>
 	concept CElement = CNonTag<T> || CTag<T>;
+
+	namespace Private
+	{
+		template<CElement T>
+		struct TElementTypeHelper
+		{
+			using Type = std::conditional_t<CFragment<T>, FMassFragment
+				, std::conditional_t<CTag<T>, FMassTag
+				, std::conditional_t<CChunkFragment<T>, FMassChunkFragment
+				, std::conditional_t<CSharedFragment<T>, FMassSharedFragment
+				, std::conditional_t<CConstSharedFragment<T>, FMassConstSharedFragment
+				, void>>>>>;
+		};
+	}
+
+	template<typename T>
+	using TElementType = typename Private::TElementTypeHelper<T>::Type;
 }
diff --git a/Engine/Source/Runtime/MassEntity/Public/MassEntityTypes.h b/Engine/Source/Runtime/MassEntity/Public/MassEntityTypes.h
--- a/Engine/Source/Runtime/MassEntity/Public/MassEntityTypes.h
+++ b/Engine/Source/Runtime/MassEntity/Public/MassEntityTypes.h
@@ -32,6 +32,12 @@ DECLARE_STRUCTTYPEBITSET_EXPORTED(MASSENTITY_API, FMassSharedFragmentBitSet, FMa
 DECLARE_STRUCTTYPEBITSET_EXPORTED(MASSENTITY_API, FMassConstSharedFragmentBitSet, FMassConstSharedFragment);
 DECLARE_CLASSTYPEBITSET_EXPORTED(MASSENTITY_API, FMassExternalSubsystemBitSet, USubsystem);
 
+namespace UE::Mass
+{
+	template<typename T>
+	static constexpr bool TAlwaysFalse = false;
+}
+
 /** The type summarily describing a composition of an entity or an archetype. It contains information on both the
  *  fragments as well as tags */
 struct FMassArchetypeCompositionDescriptor
@@ -121,6 +127,15 @@ struct FMassArchetypeCompositionDescriptor
 			&& ConstSharedFragments.HasAll(OtherDescriptor.ConstSharedFragments);
 	}
 
+	void Append(const FMassArchetypeCompositionDescriptor& OtherDescriptor)
+	{
+		Fragments += OtherDescriptor.Fragments;
+		Tags += OtherDescriptor.Tags;
+		ChunkFragments += OtherDescriptor.ChunkFragments;
+		SharedFragments += OtherDescriptor.SharedFragments;
+		ConstSharedFragments += OtherDescriptor.ConstSharedFragments;
+	}
+
 	MASSENTITY_API static uint32 CalculateHash(const FMassFragmentBitSet& InFragments, const FMassTagBitSet& InTags
 		, const FMassChunkFragmentBitSet& InChunkFragments, const FMassSharedFragmentBitSet& InSharedFragmentBitSet
 		, const FMassConstSharedFragmentBitSet& InConstSharedFragmentBitSet);
@@ -134,26 +149,17 @@ struct FMassArchetypeCompositionDescriptor
 
 	MASSENTITY_API void DebugOutputDescription(FOutputDevice& Ar) const;
 
-	template<UE::Mass::CFragment T>
-	FMassFragmentBitSet& GetContainer();
-
-	template<UE::Mass::CTag T>
-	FMassTagBitSet& GetContainer();
-
-	template<UE::Mass::CChunkFragment T>
-	FMassChunkFragmentBitSet& GetContainer();
-
-	template<UE::Mass::CSharedFragment T>
-	FMassSharedFragmentBitSet& GetContainer();
+	template<typename T>
+	auto& GetContainer() const;
 
-	template<UE::Mass::CConstSharedFragment T>
-	FMassConstSharedFragmentBitSet& GetContainer();
+	template<typename T>
+	auto& GetContainer();
 
 	template<typename T>
 	bool Contains() const;
 
 	template<typename T>
-	bool Add();
+	void Add();
 
 	FMassFragmentBitSet Fragments;
 	FMassTagBitSet Tags;
@@ -651,46 +657,76 @@ struct FMassArchetypeCreationParams
 //-----------------------------------------------------------------------------
 // INLINES
 //-----------------------------------------------------------------------------
-template<UE::Mass::CFragment T>
-FMassFragmentBitSet& FMassArchetypeCompositionDescriptor::GetContainer()
-{
-	return Fragments;
-}
-
-template<UE::Mass::CTag T>
-FMassTagBitSet& FMassArchetypeCompositionDescriptor::GetContainer()
-{
-	return Tags;
-}
-
-template<UE::Mass::CChunkFragment T>
-FMassChunkFragmentBitSet& FMassArchetypeCompositionDescriptor::GetContainer()
-{
-	return ChunkFragments;
-}
-
-template<UE::Mass::CSharedFragment T>
-FMassSharedFragmentBitSet& FMassArchetypeCompositionDescriptor::GetContainer()
+template<typename T>
+auto& FMassArchetypeCompositionDescriptor::GetContainer() const
 {
-	return SharedFragments;
+	if constexpr (std::is_same_v<FMassFragment, T>)
+	{
+		return Fragments;
+	}
+	else if constexpr (std::is_same_v<FMassTag, T>)
+	{
+		return Tags;
+	}
+	else if constexpr (std::is_same_v<FMassChunkFragment, T>)
+	{
+		return ChunkFragments;
+	}
+	else if constexpr (std::is_same_v<FMassSharedFragment, T>)
+	{
+		return SharedFragments;
+	}
+	else if constexpr (std::is_same_v<FMassConstSharedFragment, T>)
+	{
+		return ConstSharedFragments;
+	}
+	else
+    {
+        static_assert(UE::Mass::TAlwaysFalse<T>, "Unknown element type passed to GetContainer.");
+    }	
 }
 
-template<UE::Mass::CConstSharedFragment T>
-FMassConstSharedFragmentBitSet& FMassArchetypeCompositionDescriptor::GetContainer()
+template<typename T>
+auto& FMassArchetypeCompositionDescriptor::GetContainer()
 {
-	return ConstSharedFragments;
+	if constexpr (std::is_same_v<FMassFragment, T>)
+	{
+		return Fragments;
+	}
+	else if constexpr (std::is_same_v<FMassTag, T>)
+	{
+		return Tags;
+	}
+	else if constexpr (std::is_same_v<FMassChunkFragment, T>)
+	{
+		return ChunkFragments;
+	}
+	else if constexpr (std::is_same_v<FMassSharedFragment, T>)
+	{
+		return SharedFragments;
+	}
+	else if constexpr (std::is_same_v<FMassConstSharedFragment, T>)
+	{
+		return ConstSharedFragments;
+	}
+	else
+    {
+        static_assert(UE::Mass::TAlwaysFalse<T>, "Unknown element type passed to GetContainer.");
+    }	
 }
 
 template<typename T>
 bool FMassArchetypeCompositionDescriptor::Contains() const
 {
-	return GetContainer<T>().template Contains<T>();
+	using FElementType = UE::Mass::TElementType<T>;
+	return GetContainer<FElementType>().template Contains<T>();
 }
 
 template<typename T>
-bool FMassArchetypeCompositionDescriptor::Add()
+void FMassArchetypeCompositionDescriptor::Add()
 {
-	return GetContainer<T>().template Add<T>();
+	using FElementType = UE::Mass::TElementType<T>;
+	GetContainer<FElementType>().template Add<T>();
 }
 
 inline SIZE_T FMassArchetypeSharedFragmentValues::GetAllocatedSize() const
