[Mass] EntityBuilder extension:
* can use pre-reserved handles now
* can seed a builder instance with composition and values of an existing entity
* can append composition and values of an existing entity to the builder

diff --git a/Engine/Source/Developer/MassEntityTestSuite/Private/MassEntityBuilderTest.cpp b/Engine/Source/Developer/MassEntityTestSuite/Private/MassEntityBuilderTest.cpp
--- a/Engine/Source/Developer/MassEntityTestSuite/Private/MassEntityBuilderTest.cpp
+++ b/Engine/Source/Developer/MassEntityTestSuite/Private/MassEntityBuilderTest.cpp
@@ -181,7 +181,7 @@ namespace UE::Mass::Test::EntityBuilder
 	};
 	IMPLEMENT_AI_INSTANT_TEST(FOneliner, "System.Mass.EntityBuilder.OneLiner");
 
-	struct FCopy : FEntityTestBase
+	struct FCopyBuilder : FEntityTestBase
 	{
 		virtual bool InstantTest() override
 		{
@@ -216,7 +216,7 @@ namespace UE::Mass::Test::EntityBuilder
 			return true;
 		}
 	};
-	IMPLEMENT_AI_INSTANT_TEST(FCopy, "System.Mass.EntityBuilder.Copy");
+	IMPLEMENT_AI_INSTANT_TEST(FCopyBuilder, "System.Mass.EntityBuilder.Copy");
 
 	struct FOverride : FEntityTestBase
 	{
@@ -490,6 +490,115 @@ namespace UE::Mass::Test::EntityBuilder
 		}
 	};
 	IMPLEMENT_AI_INSTANT_TEST(FSyncBuildingAsyncSubmission, "System.Mass.EntityBuilder.SyncBuildingAsyncSubmission");
+
+	struct FAllElementsUsed : FEntityTestBase
+	{
+		FMassEntityHandle OriginalEntity;
+		const int32 TestIntValue = 17;
+		const float TestFloatValue = 3.1415f;
+		const float TestSharedFloatValue = 2.71828f;
+		const int32 TestSharedIntValue = 1009;
+
+		virtual bool SetUp() override
+		{
+			if (FEntityTestBase::SetUp())
+			{
+				// quick builder just to create an entity with known properties
+				FEntityBuilder Builder(EntityManager.ToSharedRef());
+				Builder.Add<FTestFragment_Int>(TestIntValue);
+				Builder.Add<FTestFragment_Float>(TestFloatValue);
+				Builder.Add<FTestTag_B>();
+				Builder.Add<FTestSharedFragment_Float>(TestSharedFloatValue);
+				Builder.Add<FTestConstSharedFragment_Int>(TestSharedIntValue);
+
+				OriginalEntity = Builder.Commit();
+
+				return true;
+			}
+			return false;
+		}
+
+		virtual bool InstantTest() override
+		{
+			FMassArchetypeCompositionDescriptor PredictedComposition;
+			PredictedComposition.Add<FTestFragment_Int>();
+			PredictedComposition.Add<FTestFragment_Float>();
+			PredictedComposition.Add<FTestTag_B>();
+			PredictedComposition.Add<FTestSharedFragment_Float>();
+			PredictedComposition.Add<FTestConstSharedFragment_Int>();
+
+			// testing composition
+			const FMassArchetypeHandle ArchetypeHandle = EntityManager->GetArchetypeForEntity(OriginalEntity);
+			const FMassArchetypeCompositionDescriptor& ArchetypeComposition = EntityManager->GetArchetypeComposition(ArchetypeHandle);
+			AITEST_TRUE("Resulting archetype composition matches prediction", ArchetypeComposition.IsEquivalent(PredictedComposition));
+
+			return TestEntity(OriginalEntity);
+		}
+
+		bool TestEntity(const FMassEntityHandle TestedEntity) const
+		{
+			{
+				FTestFragment_Int* IntFragmentInstance = EntityManager->GetFragmentDataPtr<FTestFragment_Int>(TestedEntity);
+				AITEST_NOT_NULL("Created entity has the int fragment", IntFragmentInstance);
+				AITEST_EQUAL("Resulting Int fragment value", IntFragmentInstance->Value, TestIntValue);
+			}
+			{
+				FTestFragment_Float* FloatFragmentInstance = EntityManager->GetFragmentDataPtr<FTestFragment_Float>(TestedEntity);
+				AITEST_NOT_NULL("Created entity has the int fragment", FloatFragmentInstance);
+				AITEST_EQUAL("Resulting Int fragment value", FloatFragmentInstance->Value, TestFloatValue);
+			}
+			{
+				FTestSharedFragment_Float* SharedFragmentInstance = EntityManager->GetSharedFragmentDataPtr<FTestSharedFragment_Float>(TestedEntity);
+				AITEST_NOT_NULL("Created entity has the int fragment", SharedFragmentInstance);
+				AITEST_EQUAL("Resulting Int fragment value", SharedFragmentInstance->Value, TestSharedFloatValue);
+			}
+			{
+				FTestConstSharedFragment_Int* ConstSharedFragmentInstance = EntityManager->GetConstSharedFragmentDataPtr<FTestConstSharedFragment_Int>(TestedEntity);
+				AITEST_NOT_NULL("Created entity has the int fragment", ConstSharedFragmentInstance);
+				AITEST_EQUAL("Resulting Int fragment value", ConstSharedFragmentInstance->Value, TestSharedIntValue);
+			}
+			return true;
+		}
+	};
+	IMPLEMENT_AI_INSTANT_TEST(FAllElementsUsed, "System.Mass.EntityBuilder.AllElements");
+
+	struct FCopyEntity : FAllElementsUsed
+	{
+		virtual bool InstantTest() override
+		{
+			FEntityBuilder Builder(EntityManager.ToSharedRef());
+			Builder.CopyDataFromEntity(OriginalEntity);
+
+			const FMassEntityHandle NewEntityHandle = Builder.Commit();
+
+			AITEST_TRUE("Source and target entities are in the same archetype"
+				, EntityManager->GetArchetypeForEntity(OriginalEntity) == EntityManager->GetArchetypeForEntity(NewEntityHandle));
+
+			return TestEntity(NewEntityHandle);
+		}
+	};
+	IMPLEMENT_AI_INSTANT_TEST(FCopyEntity, "System.Mass.EntityBuilder.CopyEntity");
+
+	struct FAppendFromEntity : FAllElementsUsed
+	{
+		virtual bool InstantTest() override
+		{
+			FEntityBuilder Builder(EntityManager.ToSharedRef());
+			// adding something the appending won't add, just to remove it later and test the result
+			Builder.Add<FTestTag_A>();
+			Builder.AppendDataFromEntity(OriginalEntity);
+
+			const FMassEntityHandle NewEntityHandle = Builder.Commit();
+
+			EntityManager->RemoveTagFromEntity(NewEntityHandle, FTestTag_A::StaticStruct());
+
+			AITEST_TRUE("Source and target entities are in the same archetype"
+				, EntityManager->GetArchetypeForEntity(OriginalEntity) == EntityManager->GetArchetypeForEntity(NewEntityHandle));
+
+			return TestEntity(NewEntityHandle);
+		}
+	};
+	IMPLEMENT_AI_INSTANT_TEST(FAppendFromEntity, "System.Mass.EntityBuilder.Append");
 } // UE::Mass::Test::EntityBuilder
 
 UE_ENABLE_OPTIMIZATION_SHIP
diff --git a/Engine/Source/Runtime/MassEntity/Private/MassEntityBuilder.cpp b/Engine/Source/Runtime/MassEntity/Private/MassEntityBuilder.cpp
--- a/Engine/Source/Runtime/MassEntity/Private/MassEntityBuilder.cpp
+++ b/Engine/Source/Runtime/MassEntity/Private/MassEntityBuilder.cpp
@@ -1,7 +1,6 @@
 // Copyright Epic Games, Inc. All Rights Reserved.
 
 #include "MassEntityBuilder.h"
-
 #include "MassCommandBuffer.h"
 #include "MassCommands.h"
 #include "MassEntityManager.h"
@@ -10,6 +9,56 @@
 
 namespace UE::Mass
 {
+	namespace Private
+	{
+		struct FEntityBuilderHelper
+		{
+			template<typename T>
+			static void AppendFromEntity(FEntityBuilder& Builder, const FMassEntityHandle SourceEntityHandle, const FMassArchetypeCompositionDescriptor& ArchetypeComposition)
+			{
+				const auto& SourceContainer = ArchetypeComposition.GetContainer<T>();
+				TArray<FInstancedStruct>& ElementInstanceContainer = Builder.GetInstancedStructContainerInternal<T>();
+
+				// remove all the existing entries that match SourceContainer, and then just copy.
+
+				for (auto Iterator = SourceContainer.GetIndexIterator(); Iterator; ++Iterator)
+				{
+					// @todo we could use an iterator that can fetch the type by simply calling Iterator.GetType()
+					TNotNull<const UScriptStruct*> ElementType = SourceContainer.GetTypeAtIndex(*Iterator);
+					const int32 FoundIndex = ElementInstanceContainer.IndexOfByPredicate([&ElementType](const FInstancedStruct& ExistingElement)
+					{
+						return ExistingElement.GetScriptStruct() == ElementType;
+					});
+					if (FoundIndex != INDEX_NONE)
+					{
+						ElementInstanceContainer.RemoveAtSwap(FoundIndex, EAllowShrinking::No);
+					}
+				}
+				CopyFromEntity<T>(Builder, SourceEntityHandle, ArchetypeComposition);
+			}
+
+			template<typename T>
+			static void CopyFromEntity(FEntityBuilder& Builder, const FMassEntityHandle SourceEntityHandle, const FMassArchetypeCompositionDescriptor& ArchetypeComposition)
+			{
+				const auto& SourceContainer = ArchetypeComposition.GetContainer<T>();
+				TArray<FInstancedStruct>& ElementInstanceContainer = Builder.GetInstancedStructContainerInternal<T>();
+
+				ElementInstanceContainer.Reserve(ElementInstanceContainer.Num() + SourceContainer.CountStoredTypes());
+
+				for (auto Iterator = SourceContainer.GetIndexIterator(); Iterator; ++Iterator)
+				{
+					// @todo we could use an iterator that can fetch the type by simply calling Iterator.GetType()
+					TNotNull<const UScriptStruct*> Type = SourceContainer.GetTypeAtIndex(*Iterator);
+					FConstStructView SourceElementView = Builder.EntityManager->GetElementDataStruct<T>(SourceEntityHandle, Type);
+
+					// this happening is practically impossible, so testing only in debug
+					checkSlow(SourceElementView.IsValid());
+					ElementInstanceContainer.Emplace(SourceElementView);
+				}
+			}
+		};
+	}
+
 //-----------------------------------------------------------------------------
 // FEntityBuilder
 //-----------------------------------------------------------------------------
@@ -166,16 +215,94 @@ void FEntityBuilder::Reprepare()
 	}
 }
 
-void FEntityBuilder::Reset()
+void FEntityBuilder::Reset(const bool bReleaseEntityHandleIfReserved)
 {
-	ConditionallyReleaseEntityHandle();
+	if (bReleaseEntityHandleIfReserved)
+	{
+		ConditionallyReleaseEntityHandle();
+	}
+
+	if (State != EState::Empty)
+	{
+		InvalidateCachedData();
 
-	State = EState::Empty;
+		State = EState::Empty;
 
-	Composition.Reset();
-	Fragments.Reset();
-	SharedFragments.Reset();
-	ConstSharedFragments.Reset();
+		Composition.Reset();
+		Fragments.Reset();
+		SharedFragments.Reset();
+		ConstSharedFragments.Reset();
+	}
+}
+
+bool FEntityBuilder::SetReservedEntityHandle(const FMassEntityHandle ReservedEntityHandle)
+{
+	if (!ensureMsgf(ReservedEntityHandle.IsValid() && EntityManager->IsEntityReserved(ReservedEntityHandle), TEXT("Input ReservedEntityHandle is expected to be valid and represent a reserved entity")))
+	{
+		return false;
+	}
+
+	if (EntityHandle.IsValid() && EntityManager->IsEntityReserved(EntityHandle))
+	{
+		checkf(IsCommitted() == false, TEXT("We only expect to be here when the entity builder has not been `Committed` yet"));
+		EntityManager->ReleaseReservedEntity(EntityHandle);
+	}
+
+	EntityHandle = ReservedEntityHandle;
+	return true;
+}
+
+bool FEntityBuilder::AppendDataFromEntity(const FMassEntityHandle SourceEntityHandle)
+{
+	if (!testableEnsureMsgf(EntityManager->IsEntityActive(SourceEntityHandle)
+		, TEXT("%hs expecting a valid, built entity as input"), __FUNCTION__))
+	{
+		return false;
+	}
+	if (State == EState::Empty)
+	{
+		// copying is significantly more efficient (no lookups for existing data) 
+		return CopyDataFromEntity(SourceEntityHandle);
+	}
+
+	InvalidateCachedData();
+
+	const FMassArchetypeHandle ArchetypeHandle = EntityManager->GetArchetypeForEntity(SourceEntityHandle);
+	const FMassArchetypeCompositionDescriptor& ArchetypeComposition = EntityManager->GetArchetypeComposition(ArchetypeHandle);
+
+	Private::FEntityBuilderHelper::AppendFromEntity<FMassFragment>(*this, SourceEntityHandle, ArchetypeComposition);
+	Private::FEntityBuilderHelper::AppendFromEntity<FMassSharedFragment>(*this, SourceEntityHandle, ArchetypeComposition);
+	Private::FEntityBuilderHelper::AppendFromEntity<FMassConstSharedFragment>(*this, SourceEntityHandle, ArchetypeComposition);
+
+	Composition.Append(ArchetypeComposition);
+
+	State = Composition.IsEmpty() ? EState::Empty : EState::ReadyToCommit;
+
+	return true;
+}
+
+bool FEntityBuilder::CopyDataFromEntity(const FMassEntityHandle SourceEntityHandle)
+{
+	if (!testableEnsureMsgf(EntityManager->IsEntityActive(SourceEntityHandle)
+		, TEXT("%hs expecting a valid, built entity as input"), __FUNCTION__))
+	{
+		return false;
+	}
+
+	Reset(/*bReleaseEntityHandleIfReserved=*/false);
+	
+	const FMassArchetypeHandle ArchetypeHandle = EntityManager->GetArchetypeForEntity(SourceEntityHandle);
+	const FMassArchetypeCompositionDescriptor& ArchetypeComposition = EntityManager->GetArchetypeComposition(ArchetypeHandle);
+
+	Private::FEntityBuilderHelper::CopyFromEntity<FMassFragment>(*this, SourceEntityHandle, ArchetypeComposition);
+	Private::FEntityBuilderHelper::CopyFromEntity<FMassSharedFragment>(*this, SourceEntityHandle, ArchetypeComposition);
+	Private::FEntityBuilderHelper::CopyFromEntity<FMassConstSharedFragment>(*this, SourceEntityHandle, ArchetypeComposition);
+
+	Composition = ArchetypeComposition;
+
+	State = Composition.IsEmpty() ? EState::Empty : EState::ReadyToCommit;
+
+	return true;
 }
 
 FMassEntityHandle FEntityBuilder::GetEntityHandle() const
@@ -211,6 +338,12 @@ void FEntityBuilder::CacheArchetypeHandle()
 	}
 }
 
+void FEntityBuilder::InvalidateCachedData()
+{
+	CachedArchetypeHandle = {};
+	CachedSharedFragmentValues.Reset();
+}
+
 FMassArchetypeHandle FEntityBuilder::GetArchetypeHandle()
 {
 	CacheArchetypeHandle();
@@ -230,7 +363,7 @@ void FEntityBuilder::CacheSharedFragmentValue()
 		for (FInstancedStruct& ConstSharedFragmentInstance : ConstSharedFragments)
 		{
 			check(ConstSharedFragmentInstance.IsValid());
-			const FSharedStruct& ConstSharedStruct = EntityManager->GetOrCreateSharedFragment(*ConstSharedFragmentInstance.GetScriptStruct(), ConstSharedFragmentInstance.GetMemory());
+			const FConstSharedStruct& ConstSharedStruct = EntityManager->GetOrCreateConstSharedFragment(*ConstSharedFragmentInstance.GetScriptStruct(), ConstSharedFragmentInstance.GetMemory());
 			CachedSharedFragmentValues.AddConstSharedFragment(ConstSharedStruct);
 		}
 
diff --git a/Engine/Source/Runtime/MassEntity/Public/MassEntityBuilder.h b/Engine/Source/Runtime/MassEntity/Public/MassEntityBuilder.h
--- a/Engine/Source/Runtime/MassEntity/Public/MassEntityBuilder.h
+++ b/Engine/Source/Runtime/MassEntity/Public/MassEntityBuilder.h
@@ -12,6 +12,10 @@ struct FMassEntityManager;
 
 namespace UE::Mass
 {
+	namespace Private
+	{
+		struct FEntityBuilderHelper;
+	}
 
 /**
  * FEntityBuilder is a utility struct that provides a convenient way to create and configure entities
@@ -85,8 +89,33 @@ struct MASSENTITY_API FEntityBuilder
 	/** if the builder is in "Committed" state it will roll back to ReadyToSubmit and reset the stored entity handle */
 	void Reprepare();
 
-	/** Resets the builder to its initial state, discarding all previous entity configurations. */
-	void Reset();
+	/**
+	 * Resets the builder to its initial state, discarding all previous entity configurations.
+	 * @param bReleaseEntityHandleIfReserved configures what to do with the reserved entity handle, if it's valid.
+	 */
+	void Reset(const bool bReleaseEntityHandleIfReserved = true);
+
+	/**
+	 * Stores ReservedEntityHandle as the cached EntityHandle. The ReservedEntityHandle is expected to be valid
+	 * and represent a reserved entity. These expectations will be checked via ensures.
+	 * If the existing EntityHandle also represents a valid, reserved entity, that handle will be released.
+	 * @return whether the ReservedEntityHandle has been stored.
+	 */
+	bool SetReservedEntityHandle(const FMassEntityHandle ReservedEntityHandle);
+
+	/**
+	 * Appends all element types and values stored by the entity indicated by SourceEntityHandle.
+	 * @param SourceEntityHandle valid handle for a fully constructed, built entity.
+	 * @return whether the operation was successful
+	 */
+	bool AppendDataFromEntity(const FMassEntityHandle SourceEntityHandle);
+
+	/**
+	 * Copies all element types and values stored by the entity indicated by SourceEntityHandle. Any existing builder data will be overridden
+	 * @param SourceEntityHandle valid handle for a fully constructed, built entity.
+	 * @return whether the operation was successful
+	 */
+	bool CopyDataFromEntity(const FMassEntityHandle SourceEntityHandle);
 
 	/**
 	 * Adds a fragment of type T to the entity and returns a reference to it, constructing it with the provided arguments.
@@ -150,8 +179,11 @@ struct MASSENTITY_API FEntityBuilder
 	bool IsCommitted() const;
 
 protected:
+	friend Private::FEntityBuilderHelper;
+
 	void CacheSharedFragmentValue();
 	void CacheArchetypeHandle();
+	void InvalidateCachedData();
 
 private:
 	TSharedRef<FMassEntityManager> EntityManager;
diff --git a/Engine/Source/Runtime/MassEntity/Public/MassEntityManager.h b/Engine/Source/Runtime/MassEntity/Public/MassEntityManager.h
--- a/Engine/Source/Runtime/MassEntity/Public/MassEntityManager.h
+++ b/Engine/Source/Runtime/MassEntity/Public/MassEntityManager.h
@@ -612,6 +612,23 @@ public:
 			: FConstStructView();
 	}
 
+	template<typename T>
+	FConstStructView GetElementDataStruct(FMassEntityHandle EntityHandle, TNotNull<const UScriptStruct*> FragmentType) const
+	{
+		if constexpr (std::is_same_v<T, FMassFragment>)
+		{
+			return GetFragmentDataStruct(EntityHandle, FragmentType);
+		}
+		if constexpr (std::is_same_v<T, FMassSharedFragment>)
+		{
+			return GetSharedFragmentDataStruct(EntityHandle, FragmentType);
+		}
+		if constexpr (std::is_same_v<T, FMassConstSharedFragment>)
+		{
+			return GetConstSharedFragmentDataStruct(EntityHandle, FragmentType);
+		}
+	}
+
 	uint32 GetArchetypeDataVersion() const { return ArchetypeDataVersion; }
 
 	/**
