[Mass] added a debug-time mass entity storage pointer to the EntityManager, for easier storage viewing during debugging with an IDE

#rnx

diff --git a/Engine/Source/Runtime/MassEntity/Private/MassEntityManager.cpp b/Engine/Source/Runtime/MassEntity/Private/MassEntityManager.cpp
--- a/Engine/Source/Runtime/MassEntity/Private/MassEntityManager.cpp
+++ b/Engine/Source/Runtime/MassEntity/Private/MassEntityManager.cpp
@@ -247,6 +247,9 @@ void FMassEntityManager::Initialize(const FMassEntityManagerStorageInitParams& I
 	}
 
 	Visit(UE::Mass::Private::FEntityStorageInitializer{&EntityStorage}, InitializationParams);
+#if WITH_MASSENTITY_DEBUG
+	DebugEntityStoragePtr = &DebugGetEntityStorageInterface();
+#endif // WITH_MASSENTITY_DEBUG
 
 	for (TSharedPtr<FMassCommandBuffer>& CommandBuffer : DeferredCommandBuffers)
 	{
diff --git a/Engine/Source/Runtime/MassEntity/Public/MassEntityManager.h b/Engine/Source/Runtime/MassEntity/Public/MassEntityManager.h
--- a/Engine/Source/Runtime/MassEntity/Public/MassEntityManager.h
+++ b/Engine/Source/Runtime/MassEntity/Public/MassEntityManager.h
@@ -1168,6 +1168,9 @@ private:
 		UE::Mass::FSingleThreadedEntityStorage,
 		UE::Mass::FConcurrentEntityStorage>;
 	mutable FEntityStorageContainerType EntityStorage;
+#if WITH_MASSENTITY_DEBUG
+	UE::Mass::FStorageType* DebugEntityStoragePtr = nullptr;
+#endif // WITH_MASSENTITY_DEBUG
 
 	/** Never access directly, use GetOrMakeCreationContext instead. */
 	TWeakPtr<FEntityCreationContext> ActiveCreationContext;
