[Mass] CIS Issue 724725 - PVS Studio warnings in MassEntityManager.h

#rnx

diff --git a/Engine/Source/Runtime/MassEntity/Public/MassEntityManager.h b/Engine/Source/Runtime/MassEntity/Public/MassEntityManager.h
--- a/Engine/Source/Runtime/MassEntity/Public/MassEntityManager.h
+++ b/Engine/Source/Runtime/MassEntity/Public/MassEntityManager.h
@@ -619,14 +619,18 @@ public:
 		{
 			return GetFragmentDataStruct(EntityHandle, FragmentType);
 		}
-		if constexpr (std::is_same_v<T, FMassSharedFragment>)
+		else if constexpr (std::is_same_v<T, FMassSharedFragment>)
 		{
 			return GetSharedFragmentDataStruct(EntityHandle, FragmentType);
 		}
-		if constexpr (std::is_same_v<T, FMassConstSharedFragment>)
+		else if constexpr (std::is_same_v<T, FMassConstSharedFragment>)
 		{
 			return GetConstSharedFragmentDataStruct(EntityHandle, FragmentType);
 		}
+		else
+		{
+			static_assert(UE::Mass::TAlwaysFalse<T>, "Unsupported element type passed to GetElementDataStruct");
+		}
 	}
 
 	uint32 GetArchetypeDataVersion() const { return ArchetypeDataVersion; }
