[Mass] Fixed missing return type in MassEntityManager.GetElementDataStruct

#rnx

diff --git a/Engine/Source/Runtime/MassEntity/Public/MassEntityManager.h b/Engine/Source/Runtime/MassEntity/Public/MassEntityManager.h
--- a/Engine/Source/Runtime/MassEntity/Public/MassEntityManager.h
+++ b/Engine/Source/Runtime/MassEntity/Public/MassEntityManager.h
@@ -630,6 +630,7 @@ public:
 		else
 		{
 			static_assert(UE::Mass::TAlwaysFalse<T>, "Unsupported element type passed to GetElementDataStruct");
+			return {};
 		}
 	}
 
