[Mass] Replaced all occurences of TIsDerivedFrom<T, FMassFragment>::IsDerived and std::is_base_of_v<FMassFragment, T>  with UE::Mass::CFragment<T>; and Type->IsChildOf(FMassFragment::StaticStruct()) with UE::Mass::IsA<FMassFragment>; and similarly for other Mass element types

diff --git a/Engine/Plugins/AI/MassAI/Source/MassAIBehavior/Private/MassStateTreeExecutionContext.cpp b/Engine/Plugins/AI/MassAI/Source/MassAIBehavior/Private/MassStateTreeExecutionContext.cpp
--- a/Engine/Plugins/AI/MassAI/Source/MassAIBehavior/Private/MassStateTreeExecutionContext.cpp
+++ b/Engine/Plugins/AI/MassAI/Source/MassAIBehavior/Private/MassStateTreeExecutionContext.cpp
@@ -31,8 +31,8 @@ bool CollectExternalData(const FStateTreeExecutionContext& Context, const UState
 		{
 			continue;
 		}
-		
-		if (DataDesc.Struct->IsChildOf(FMassFragment::StaticStruct()))
+
+		if (UE::Mass::IsA<FMassFragment>(DataDesc.Struct))
 		{
 			const UScriptStruct* ScriptStruct = Cast<const UScriptStruct>(DataDesc.Struct);
 			FStructView Fragment = EntityView.GetFragmentDataStruct(ScriptStruct);
@@ -51,7 +51,7 @@ bool CollectExternalData(const FStateTreeExecutionContext& Context, const UState
 				}
 			}
 		}
-		else if (DataDesc.Struct->IsChildOf(FMassSharedFragment::StaticStruct()))
+		else if (UE::Mass::IsA<FMassSharedFragment>(DataDesc.Struct))
 		{
 			const UScriptStruct* ScriptStruct = Cast<const UScriptStruct>(DataDesc.Struct);
 			FStructView Fragment = EntityView.GetSharedFragmentDataStruct(ScriptStruct);
@@ -70,7 +70,7 @@ bool CollectExternalData(const FStateTreeExecutionContext& Context, const UState
 				}
 			}
 		}
-		else if (DataDesc.Struct->IsChildOf(FMassConstSharedFragment::StaticStruct()))
+		else if (UE::Mass::IsA<FMassConstSharedFragment>(DataDesc.Struct))
 		{
 			const UScriptStruct* ScriptStruct = Cast<const UScriptStruct>(DataDesc.Struct);
 			FConstStructView Fragment = EntityView.GetConstSharedFragmentDataStruct(ScriptStruct);
diff --git a/Engine/Plugins/AI/MassAI/Source/MassAIBehavior/Private/MassStateTreeSchema.cpp b/Engine/Plugins/AI/MassAI/Source/MassAIBehavior/Private/MassStateTreeSchema.cpp
--- a/Engine/Plugins/AI/MassAI/Source/MassAIBehavior/Private/MassStateTreeSchema.cpp
+++ b/Engine/Plugins/AI/MassAI/Source/MassAIBehavior/Private/MassStateTreeSchema.cpp
@@ -29,9 +29,9 @@ bool UMassStateTreeSchema::IsExternalItemAllowed(const UStruct& InStruct) const
 {
 	// Allow only WorldSubsystems and fragments as external data.
 	return InStruct.IsChildOf(UWorldSubsystem::StaticClass())
-			|| InStruct.IsChildOf(FMassFragment::StaticStruct())
-			|| InStruct.IsChildOf(FMassSharedFragment::StaticStruct())
-			|| InStruct.IsChildOf(FMassConstSharedFragment::StaticStruct());
+			|| UE::Mass::IsA<FMassFragment>(&InStruct)
+			|| UE::Mass::IsA<FMassSharedFragment>(&InStruct)
+			|| UE::Mass::IsA<FMassConstSharedFragment>(&InStruct);
 }
 
 bool UMassStateTreeSchema::Link(FStateTreeLinker& Linker)
diff --git a/Engine/Plugins/AI/MassAI/Source/MassAIBehavior/Private/MassStateTreeTrait.cpp b/Engine/Plugins/AI/MassAI/Source/MassAIBehavior/Private/MassStateTreeTrait.cpp
--- a/Engine/Plugins/AI/MassAI/Source/MassAIBehavior/Private/MassStateTreeTrait.cpp
+++ b/Engine/Plugins/AI/MassAI/Source/MassAIBehavior/Private/MassStateTreeTrait.cpp
@@ -73,7 +73,7 @@ bool UMassStateTreeTrait::ValidateTemplate(const FMassEntityTemplateBuildContext
 					bIssuesFound = bIssuesFound || !Subsystem;
 				}
 			}
-			else if (ItemDesc.Struct->IsChildOf(FMassFragment::StaticStruct()))
+			else if (UE::Mass::IsA<FMassFragment>(ItemDesc.Struct))
 			{
 				const UScriptStruct& FragmentType = *CastChecked<UScriptStruct>(ItemDesc.Struct);
 				if (BuildContext.HasFragment(FragmentType) == false)
@@ -83,7 +83,7 @@ bool UMassStateTreeTrait::ValidateTemplate(const FMassEntityTemplateBuildContext
 					bIssuesFound = true;
 				}
 			}
-			else if (ItemDesc.Struct->IsChildOf(FMassSharedFragment::StaticStruct()))
+			else if (UE::Mass::IsA<FMassSharedFragment>(ItemDesc.Struct))
 			{
 				const UScriptStruct& FragmentType = *CastChecked<UScriptStruct>(ItemDesc.Struct);
 				if (BuildContext.HasSharedFragment(FragmentType) == false)
@@ -93,7 +93,7 @@ bool UMassStateTreeTrait::ValidateTemplate(const FMassEntityTemplateBuildContext
 					bIssuesFound = true;
 				}
 			}
-			else if (ItemDesc.Struct->IsChildOf(FMassConstSharedFragment::StaticStruct()))
+			else if (UE::Mass::IsA<FMassConstSharedFragment>(ItemDesc.Struct))
 			{
 				const UScriptStruct& FragmentType = *CastChecked<UScriptStruct>(ItemDesc.Struct);
 				if (BuildContext.HasConstSharedFragment(FragmentType) == false)
diff --git a/Engine/Plugins/Experimental/EditorDataStorage/Source/TedsCore/Private/Processors/TypedElementProcessorAdaptors.cpp b/Engine/Plugins/Experimental/EditorDataStorage/Source/TedsCore/Private/Processors/TypedElementProcessorAdaptors.cpp
--- a/Engine/Plugins/Experimental/EditorDataStorage/Source/TedsCore/Private/Processors/TypedElementProcessorAdaptors.cpp
+++ b/Engine/Plugins/Experimental/EditorDataStorage/Source/TedsCore/Private/Processors/TypedElementProcessorAdaptors.cpp
@@ -82,11 +82,11 @@ namespace UE::Editor::DataStorage
 
 			bool HasColumn(const UScriptStruct* ColumnType) const
 			{
-				if (ColumnType->IsChildOf(FMassTag::StaticStruct()))
+				if (UE::Mass::IsA<FMassTag>(ColumnType))
 				{
 					return Context.DoesArchetypeHaveTag(*ColumnType);
 				}
-				if (ColumnType->IsChildOf(FMassFragment::StaticStruct()))
+				if (UE::Mass::IsA<FMassFragment>(ColumnType))
 				{
 					return Context.DoesArchetypeHaveFragment(*ColumnType);
 				}
@@ -102,11 +102,11 @@ namespace UE::Editor::DataStorage
 				FMassArchetypeHandle Archetype = Manager.GetArchetypeForEntity(Entity);
 				const FMassArchetypeCompositionDescriptor& Composition = Manager.GetArchetypeComposition(Archetype);
 
-				if (ColumnType->IsChildOf(FMassTag::StaticStruct()))
+				if (UE::Mass::IsA<FMassTag>(ColumnType))
 				{
 					return Composition.Tags.Contains(*ColumnType);
 				}
-				if (ColumnType->IsChildOf(FMassFragment::StaticStruct()))
+				if (UE::Mass::IsA<FMassFragment>(ColumnType))
 				{
 					return Composition.Fragments.Contains(*ColumnType);
 				}
@@ -132,7 +132,7 @@ namespace UE::Editor::DataStorage
 			{
 				for (const UScriptStruct* ColumnType : ColumnTypes)
 				{
-					if (ColumnType->IsChildOf(FMassTag::StaticStruct()))
+					if (UE::Mass::IsA<FMassTag>(ColumnType))
 					{
 						if (this->Context.DoesArchetypeHaveTag(*ColumnType))
 						{
@@ -141,7 +141,7 @@ namespace UE::Editor::DataStorage
 					}
 					else
 					{
-						checkf(ColumnType->IsChildOf(FMassFragment::StaticStruct()),
+						checkf(UE::Mass::IsA<FMassFragment>(ColumnType),
 							TEXT("Given struct type is not a valid fragment or tag type."));
 						if (this->Context.DoesArchetypeHaveFragment(*ColumnType))
 						{
@@ -157,13 +157,13 @@ namespace UE::Editor::DataStorage
 			{
 				for (const UScriptStruct* ColumnType : ColumnTypes)
 				{
-					if (ColumnType->IsChildOf(FMassTag::StaticStruct()))
+					if (UE::Mass::IsA<FMassTag>(ColumnType))
 					{
 						Descriptor.Tags.Add(*ColumnType);
 					}
 					else
 					{
-						checkf(ColumnType->IsChildOf(FMassFragment::StaticStruct()),
+						checkf(UE::Mass::IsA<FMassFragment>(ColumnType),
 							TEXT("Given struct type is not a valid fragment or tag type."));
 						Descriptor.Fragments.Add(*ColumnType);
 
@@ -279,7 +279,7 @@ namespace UE::Editor::DataStorage
 
 			void* AddColumnUninitialized(RowHandle Row, const UScriptStruct* ObjectType, ObjectCopyOrMove Relocator)
 			{
-				checkf(ObjectType->IsChildOf(FMassFragment::StaticStruct()), TEXT("Column [%s] can not be a tag"), *ObjectType->GetName());
+				checkf(UE::Mass::IsA<FMassFragment>(ObjectType), TEXT("Column [%s] can not be a tag"), *ObjectType->GetName());
 		
 				struct FAddValueColumn
 				{
diff --git a/Engine/Plugins/Experimental/EditorDataStorage/Source/TedsCore/Private/Queries/TypedElementExtendedQueryStore.cpp b/Engine/Plugins/Experimental/EditorDataStorage/Source/TedsCore/Private/Queries/TypedElementExtendedQueryStore.cpp
--- a/Engine/Plugins/Experimental/EditorDataStorage/Source/TedsCore/Private/Queries/TypedElementExtendedQueryStore.cpp
+++ b/Engine/Plugins/Experimental/EditorDataStorage/Source/TedsCore/Private/Queries/TypedElementExtendedQueryStore.cpp
@@ -694,7 +694,7 @@ namespace UE::Editor::DataStorage
 					if (ensureMsgf(Type.IsValid(), TEXT("Provided query selection type can not be null.")) &&
 						ensureMsgf(
 							Type->IsChildOf(FColumn::StaticStruct()) ||
-							Type->IsChildOf(FMassFragment::StaticStruct()),
+							UE::Mass::IsA<FMassFragment>(Type.Get()),
 							TEXT("Provided query selection type '%s' is not based on FColumn or another supported base type."),
 							*Type->GetStructPathName().ToString()))
 					{
@@ -798,11 +798,11 @@ namespace UE::Editor::DataStorage
 						continue;
 					}
 
-					if (Operand->Type->IsChildOf(FMassTag::StaticStruct()))
+					if (UE::Mass::IsA<FMassTag>(Operand->Type.Get()))
 					{
 						NativeQuery.AddTagRequirement(*(Operand->Type), Presence);
 					}
-					else if (Operand->Type->IsChildOf(FMassFragment::StaticStruct()))
+					else if (UE::Mass::IsA<FMassFragment>(Operand->Type.Get()))
 					{
 						NativeQuery.AddRequirement(Operand->Type.Get(), EMassFragmentAccess::None, Presence);
 					}
diff --git a/Engine/Plugins/Experimental/EditorDataStorage/Source/TedsCore/Private/TypedElementDatabase.cpp b/Engine/Plugins/Experimental/EditorDataStorage/Source/TedsCore/Private/TypedElementDatabase.cpp
--- a/Engine/Plugins/Experimental/EditorDataStorage/Source/TedsCore/Private/TypedElementDatabase.cpp
+++ b/Engine/Plugins/Experimental/EditorDataStorage/Source/TedsCore/Private/TypedElementDatabase.cpp
@@ -46,12 +46,10 @@ FAutoConsoleCommandWithOutputDevice PrintSupportedColumnsConsoleCommand(
 	FConsoleCommandWithOutputDeviceDelegate::CreateLambda([](FOutputDevice& Output)
 		{
 			Output.Log(TEXT("The Editor Data Storage supports the following columns:"));
-			
-			UScriptStruct* FragmentTypeInfo = FMassFragment::StaticStruct();
-			UScriptStruct* TagTypeInfo = FMassTag::StaticStruct();
+
 			for (TObjectIterator<UScriptStruct> It; It; ++It)
 			{
-				if (It->IsChildOf(FragmentTypeInfo) || It->IsChildOf(TagTypeInfo))
+				if (UE::Mass::IsA<FMassFragment>(*It) || UE::Mass::IsA<FMassTag>(*It))
 				{
 					int32 StructureSize = It->GetStructureSize();
 					if (StructureSize >= 1024 * 1024)
@@ -86,12 +84,12 @@ namespace UE::Editor::DataStorage::Private
 
 		for (const UScriptStruct* ColumnType : Columns)
 		{
-			if (ColumnType->IsChildOf(FMassFragment::StaticStruct()))
+			if (UE::Mass::IsA<FMassFragment>(ColumnType))
 			{
 				Fragments.Add(*ColumnType);
 				Result.bMustUpdateFragments = true;
 			}
-			else if (ColumnType->IsChildOf(FMassTag::StaticStruct()))
+			else if (UE::Mass::IsA<FMassTag>(ColumnType))
 			{
 				Tags.Add(*ColumnType);
 				Result.bMustUpdateTags = true;
@@ -512,7 +510,7 @@ void UEditorDataStorage::AddColumnData(RowHandle Row, const UScriptStruct* Colum
 	const ColumnCreationCallbackRef& Initializer,
 	ColumnCopyOrMoveCallback Relocator)
 {
-	if (ActiveEditorEntityManager && ColumnType && ColumnType->IsChildOf(FMassFragment::StaticStruct()))
+	if (ActiveEditorEntityManager && UE::Mass::IsA<FMassFragment>(ColumnType))
 	{
 		if (IsRowAssigned(Row))
 		{
@@ -553,8 +551,7 @@ void UEditorDataStorage::RemoveColumn(RowHandle Row, const UScriptStruct* Column
 const void* UEditorDataStorage::GetColumnData(RowHandle Row, const UScriptStruct* ColumnType) const
 {
 	FMassEntityHandle Entity = FMassEntityHandle::FromNumber(Row);
-	if (ActiveEditorEntityManager &&
-		ColumnType && ColumnType->IsChildOf(FMassFragment::StaticStruct()))
+	if (ActiveEditorEntityManager && UE::Mass::IsA<FMassFragment>(ColumnType))
 	{
 		if (ActiveEditorEntityManager->IsEntityActive(Entity))
 		{
@@ -779,11 +776,11 @@ bool UEditorDataStorage::HasColumns(RowHandle Row, TConstArrayView<const UScript
 			const UScriptStruct* const* ColumnTypesEnd = ColumnTypes.end();
 			for (const UScriptStruct* const* ColumnType = ColumnTypes.begin(); ColumnType != ColumnTypesEnd && bHasAllColumns; ++ColumnType)
 			{
-				if ((*ColumnType)->IsChildOf(FMassFragment::StaticStruct()))
+				if (UE::Mass::IsA<FMassFragment>(*ColumnType))
 				{
 					bHasAllColumns = Composition.Fragments.Contains(**ColumnType);
 				}
-				else if ((*ColumnType)->IsChildOf(FMassTag::StaticStruct()))
+				else if (UE::Mass::IsA<FMassTag>(*ColumnType))
 				{
 					bHasAllColumns = Composition.Tags.Contains(**ColumnType);
 				}
@@ -825,12 +822,12 @@ bool UEditorDataStorage::HasColumns(RowHandle Row, TConstArrayView<TWeakObjectPt
 			{
 				if (ColumnType->IsValid())
 				{
-					if ((*ColumnType)->IsChildOf(FMassFragment::StaticStruct()))
+					if (UE::Mass::IsA<FMassFragment>(ColumnType->Get()))
 					{
 						bHasAllColumns = Composition.Fragments.Contains(**ColumnType);
 						continue;
 					}
-					else if ((*ColumnType)->IsChildOf(FMassTag::StaticStruct()))
+					else if (UE::Mass::IsA<FMassTag>(ColumnType->Get()))
 					{
 						bHasAllColumns = Composition.Tags.Contains(**ColumnType);
 						continue;
@@ -935,11 +932,11 @@ bool UEditorDataStorage::MatchesColumns(RowHandle Row, const UE::Editor::DataSto
 			{
 				if (ColumnType.IsValid())
 				{
-					if (ColumnType->IsChildOf(FMassFragment::StaticStruct()))
+					if (UE::Mass::IsA<FMassFragment>(ColumnType.Get()))
 					{
 						return Composition.Fragments.Contains(*ColumnType);
 					}
-					else if (ColumnType->IsChildOf(FMassTag::StaticStruct()))
+					else if (UE::Mass::IsA<FMassTag>(ColumnType.Get()))
 					{
 						return Composition.Tags.Contains(*ColumnType);
 					}
diff --git a/Engine/Plugins/Experimental/EditorDataStorage/Source/TedsCore/Private/TypedElementDatabaseCommandBuffer.cpp b/Engine/Plugins/Experimental/EditorDataStorage/Source/TedsCore/Private/TypedElementDatabaseCommandBuffer.cpp
--- a/Engine/Plugins/Experimental/EditorDataStorage/Source/TedsCore/Private/TypedElementDatabaseCommandBuffer.cpp
+++ b/Engine/Plugins/Experimental/EditorDataStorage/Source/TedsCore/Private/TypedElementDatabaseCommandBuffer.cpp
@@ -21,7 +21,7 @@ namespace UE::Editor::DataStorage::Legacy
 	void* FCommandBuffer::GetQueuedDataColumn(
 		RowHandle Row, const UScriptStruct* ColumnType)
 	{
-		checkf(ColumnType->IsChildOf(FMassFragment::StaticStruct()),
+		checkf(UE::Mass::IsA<FMassFragment>(ColumnType),
 			TEXT("Trying to get the column '%s' which isn't a data column."), *ColumnType->GetName());
 
 		void* const* StoredData = PendingColumns.Find(PendingColumnMappingKey(Row, ColumnType));
@@ -126,7 +126,7 @@ namespace UE::Editor::DataStorage::Legacy
 	void* FCommandBuffer::Queue_AddDataColumnCommandUnitialized(
 		RowHandle Row, const UScriptStruct* ColumnType, ColumnCopyOrMoveCallback Relocator)
 	{
-		checkf(ColumnType->IsChildOf(FMassFragment::StaticStruct()),
+		checkf(UE::Mass::IsA<FMassFragment>(ColumnType),
 			TEXT("Trying to queue a data column creation for '%s' which isn't a data column."), *ColumnType->GetName());
 
 		PendingColumnMappingKey Key(Row, ColumnType);
@@ -225,11 +225,11 @@ namespace UE::Editor::DataStorage::Legacy
 		if (ColumnType)
 		{
 			FMassEntityHandle Entity = FMassEntityHandle::FromNumber(Row);
-			if (ColumnType->IsChildOf(FMassTag::StaticStruct()))
+			if (UE::Mass::IsA<FMassTag>(ColumnType))
 			{
 				MassEntityManager.AddTagToEntity(Entity, ColumnType);
 			}
-			else if (ColumnType->IsChildOf(FMassFragment::StaticStruct()))
+			else if (UE::Mass::IsA<FMassFragment>(ColumnType))
 			{
 				FStructView Column = MassEntityManager.GetFragmentDataStruct(Entity, ColumnType);
 				// Only add if not already added to avoid asserts from Mass.
@@ -247,7 +247,7 @@ namespace UE::Editor::DataStorage::Legacy
 	{
 		if (ColumnType)
 		{
-			checkf(ColumnType->IsChildOf(FMassFragment::StaticStruct()),
+			checkf(UE::Mass::IsA<FMassFragment>(ColumnType),
 				TEXT("Trying to create a data column for '%s' from a deferred command that isn't a data column."), *ColumnType->GetName());
 
 			FMassEntityHandle Entity = FMassEntityHandle::FromNumber(Row);
@@ -305,11 +305,11 @@ namespace UE::Editor::DataStorage::Legacy
 		if (ColumnType)
 		{
 			FMassEntityHandle Entity = FMassEntityHandle::FromNumber(Row);
-			if (ColumnType->IsChildOf(FMassTag::StaticStruct()))
+			if (UE::Mass::IsA<FMassTag>(ColumnType))
 			{
 				MassEntityManager.RemoveTagFromEntity(Entity, ColumnType);
 			}
-			else if (ColumnType->IsChildOf(FMassFragment::StaticStruct()))
+			else if (UE::Mass::IsA<FMassFragment>(ColumnType))
 			{
 				MassEntityManager.RemoveFragmentFromEntity(Entity, ColumnType);
 			}
diff --git a/Engine/Plugins/Runtime/MassGameplay/Source/MassActors/Public/MassActorHelper.h b/Engine/Plugins/Runtime/MassGameplay/Source/MassActors/Public/MassActorHelper.h
--- a/Engine/Plugins/Runtime/MassGameplay/Source/MassActors/Public/MassActorHelper.h
+++ b/Engine/Plugins/Runtime/MassGameplay/Source/MassActors/Public/MassActorHelper.h
@@ -2,7 +2,7 @@
 
 #pragma once
 
-#include "MassEntityTypes.h"
+#include "MassEntityConcepts.h"
 
 class UScriptStruct;
 class AActor;
@@ -22,7 +22,7 @@ namespace UE::MassActor
 	template<typename TagType>
 	bool AddEntityTagToActor(const AActor& Actor)
 	{
-		static_assert(TIsDerivedFrom<TagType, FMassTag>::IsDerived, "Given struct doesn't represent a valid tag type.");
+		static_assert(UE::Mass::CTag<TagType>, "Given struct doesn't represent a valid tag type.");
 		return AddEntityTagToActor(Actor, *TagType::StaticStruct());
 	}
 
@@ -39,7 +39,7 @@ namespace UE::MassActor
 	template<typename TagType>
 	bool RemoveEntityTagFromActor(const AActor& Actor)
 	{
-		static_assert(TIsDerivedFrom<TagType, FMassTag>::IsDerived, "Given struct doesn't represent a valid tag type.");
+		static_assert(UE::Mass::CTag<TagType>, "Given struct doesn't represent a valid tag type.");
 		return RemoveEntityTagFromActor(Actor, *TagType::StaticStruct());
 	}
 } // namespace UE::MassActor
diff --git a/Engine/Plugins/Runtime/MassGameplay/Source/MassEQS/Private/Processors/MassEnvQueryTestProcessor_MassEntityTags.cpp b/Engine/Plugins/Runtime/MassGameplay/Source/MassEQS/Private/Processors/MassEnvQueryTestProcessor_MassEntityTags.cpp
--- a/Engine/Plugins/Runtime/MassGameplay/Source/MassEQS/Private/Processors/MassEnvQueryTestProcessor_MassEntityTags.cpp
+++ b/Engine/Plugins/Runtime/MassGameplay/Source/MassEQS/Private/Processors/MassEnvQueryTestProcessor_MassEntityTags.cpp
@@ -108,7 +108,7 @@ bool Mass::EQS::Utils::TestChunkForAnyTags(const TArray<FInstancedStruct>& Tags,
 	{
 		const UScriptStruct* TagScriptStruct = Tag.GetScriptStruct();
 #if !UE_BUILD_SHIPPING
-		if (!ensureMsgf(TagScriptStruct && TagScriptStruct->IsChildOf(FMassTag::StaticStruct()), TEXT("Non Tag Element Detected in Enviroment Query Mass-Entity-Tags-Test. Likely left empty.")))
+		if (!ensureMsgf(UE::Mass::IsA<FMassTag>(TagScriptStruct), TEXT("Non Tag Element Detected in Environment Query Mass-Entity-Tags-Test. Likely left empty.")))
 		{
 			continue;
 		}
@@ -128,7 +128,7 @@ bool Mass::EQS::Utils::TestChunkForAllTags(const TArray<FInstancedStruct>& Tags,
 	{
 		const UScriptStruct* TagScriptStruct = Tag.GetScriptStruct();
 #if !UE_BUILD_SHIPPING
-		if (!ensureMsgf(TagScriptStruct && TagScriptStruct->IsChildOf(FMassTag::StaticStruct()), TEXT("Non Tag Element Detected in Enviroment Query Mass-Entity-Tags-Test. Likely left empty.")))
+		if (!ensureMsgf(UE::Mass::IsA<FMassTag>(TagScriptStruct), TEXT("Non Tag Element Detected in Environment Query Mass-Entity-Tags-Test. Likely left empty.")))
 		{
 			continue;
 		}
@@ -148,7 +148,7 @@ bool Mass::EQS::Utils::TestChunkForNoTags(const TArray<FInstancedStruct>& Tags,
 	{
 		const UScriptStruct* TagScriptStruct = Tag.GetScriptStruct();
 #if !UE_BUILD_SHIPPING
-		if (!ensureMsgf(TagScriptStruct && TagScriptStruct->IsChildOf(FMassTag::StaticStruct()), TEXT("Non Tag Element Detected in Enviroment Query Mass-Entity-Tags-Test. Likely left empty.")))
+		if (!ensureMsgf(UE::Mass::IsA<FMassTag>(TagScriptStruct), TEXT("Non Tag Element Detected in Environment Query Mass-Entity-Tags-Test. Likely left empty.")))
 		{
 			continue;
 		}
diff --git a/Engine/Plugins/Runtime/MassGameplay/Source/MassSpawner/Private/MassAssortedFragmentsTrait.cpp b/Engine/Plugins/Runtime/MassGameplay/Source/MassSpawner/Private/MassAssortedFragmentsTrait.cpp
--- a/Engine/Plugins/Runtime/MassGameplay/Source/MassSpawner/Private/MassAssortedFragmentsTrait.cpp
+++ b/Engine/Plugins/Runtime/MassGameplay/Source/MassSpawner/Private/MassAssortedFragmentsTrait.cpp
@@ -12,7 +12,7 @@ void UMassAssortedFragmentsTrait::BuildTemplate(FMassEntityTemplateBuildContext&
 		{
 			const UScriptStruct* Type = Fragment.GetScriptStruct();
 			CA_ASSUME(Type);
-			if (Type->IsChildOf(FMassFragment::StaticStruct()))
+			if (UE::Mass::IsA<FMassFragment>(Type))
 			{
 				BuildContext.AddFragment(Fragment);
 			}
@@ -25,18 +25,14 @@ void UMassAssortedFragmentsTrait::BuildTemplate(FMassEntityTemplateBuildContext&
 	
 	for (const FInstancedStruct& Tag : Tags)
 	{
-		if (Tag.IsValid())
+		const UScriptStruct* Type = Tag.GetScriptStruct();
+		if (UE::Mass::IsA<FMassTag>(Type))
 		{
-			const UScriptStruct* Type = Tag.GetScriptStruct();
-			CA_ASSUME(Type);
-			if (Type->IsChildOf(FMassTag::StaticStruct()))
-			{
-				BuildContext.AddTag(*Type);
-			}
-			else
-			{
-				UE_LOG(LogMass, Error, TEXT("Struct type %s is not a child of FMassTag"), *GetPathNameSafe(Type));
-			}
+			BuildContext.AddTag(*Type);
+		}
+		else
+		{
+			UE_LOG(LogMass, Error, TEXT("Struct type %s is not a child of FMassTag"), *GetPathNameSafe(Type));
 		}
 	}
 }
diff --git a/Engine/Plugins/Runtime/MassGameplay/Source/MassSpawner/Public/MassEntityTemplate.h b/Engine/Plugins/Runtime/MassGameplay/Source/MassSpawner/Public/MassEntityTemplate.h
--- a/Engine/Plugins/Runtime/MassGameplay/Source/MassSpawner/Public/MassEntityTemplate.h
+++ b/Engine/Plugins/Runtime/MassGameplay/Source/MassSpawner/Public/MassEntityTemplate.h
@@ -123,13 +123,13 @@ struct MASSSPAWNER_API FMassEntityTemplateData
 	template<typename T>
 	void AddFragment()
 	{
-		static_assert(TIsDerivedFrom<T, FMassFragment>::IsDerived, "Given struct doesn't represent a valid fragment type. Make sure to inherit from FMassFragment or one of its child-types.");
+		static_assert(UE::Mass::CFragment<T>, "Given struct doesn't represent a valid fragment type. Make sure to inherit from FMassFragment or one of its child-types.");
 		Composition.Fragments.Add<T>();
 	}
 
 	void AddFragment(const UScriptStruct& FragmentType)
 	{
-		checkf(FragmentType.IsChildOf(FMassFragment::StaticStruct()), TEXT("Given struct doesn't represent a valid fragment type. Make sure to inherit from FMassFragment or one of its child-types."));
+		checkf(UE::Mass::IsA<FMassFragment>(&FragmentType), TEXT("Given struct doesn't represent a valid fragment type. Make sure to inherit from FMassFragment or one of its child-types."));
 		Composition.Fragments.Add(FragmentType);
 	}
 
@@ -137,7 +137,7 @@ struct MASSSPAWNER_API FMassEntityTemplateData
 	void AddFragment(FConstStructView Fragment)
 	{
 		const UScriptStruct* FragmentType = Fragment.GetScriptStruct();
-		checkf(FragmentType && FragmentType->IsChildOf(FMassFragment::StaticStruct()), TEXT("Given struct doesn't represent a valid fragment type. Make sure to inherit from FMassFragment or one of its child-types."));
+		checkf(UE::Mass::IsA<FMassFragment>(FragmentType), TEXT("Given struct doesn't represent a valid fragment type. Make sure to inherit from FMassFragment or one of its child-types."));
 		if (!Composition.Fragments.Contains(*FragmentType))
 		{
 			Composition.Fragments.Add(*FragmentType);
@@ -152,7 +152,7 @@ struct MASSSPAWNER_API FMassEntityTemplateData
 	template<typename T>
 	T& AddFragment_GetRef()
 	{
-		static_assert(TIsDerivedFrom<T, FMassFragment>::IsDerived, "Given struct doesn't represent a valid fragment type. Make sure to inherit from FMassFragment or one of its child-types.");
+		static_assert(UE::Mass::CFragment<T>, "Given struct doesn't represent a valid fragment type. Make sure to inherit from FMassFragment or one of its child-types.");
 		if (!Composition.Fragments.Contains<T>())
 		{
 			Composition.Fragments.Add<T>();
@@ -169,7 +169,7 @@ struct MASSSPAWNER_API FMassEntityTemplateData
 	template<typename T>
 	T* GetMutableFragment()
 	{
-		static_assert(TIsDerivedFrom<T, FMassFragment>::IsDerived, "Given struct doesn't represent a valid fragment type. Make sure to inherit from FMassFragment or one of its child-types.");
+		static_assert(UE::Mass::CFragment<T>, "Given struct doesn't represent a valid fragment type. Make sure to inherit from FMassFragment or one of its child-types.");
 		FInstancedStruct* Fragment = InitialFragmentValues.FindByPredicate(FStructTypeEqualOperator(T::StaticStruct()));
 		return Fragment ? &Fragment->template GetMutable<T>() : (T*)nullptr;
 	}
@@ -177,26 +177,26 @@ struct MASSSPAWNER_API FMassEntityTemplateData
 	template<typename T>
 	void AddTag()
 	{
-		static_assert(TIsDerivedFrom<T, FMassTag>::IsDerived, "Given struct doesn't represent a valid mass tag type. Make sure to inherit from FMassTag or one of its child-types.");
+		static_assert(UE::Mass::CTag<T>, "Given struct doesn't represent a valid mass tag type. Make sure to inherit from FMassTag or one of its child-types.");
 		Composition.Tags.Add<T>();
 	}
 	
 	void AddTag(const UScriptStruct& TagType)
 	{
-		checkf(TagType.IsChildOf(FMassTag::StaticStruct()), TEXT("Given struct doesn't represent a valid mass tag type. Make sure to inherit from FMassTag or one of its child-types."));
+		checkf(UE::Mass::IsA<FMassTag>(&TagType), TEXT("Given struct doesn't represent a valid mass tag type. Make sure to inherit from FMassTag or one of its child-types."));
 		Composition.Tags.Add(TagType);
 	}
 
 	template<typename T>
 	void RemoveTag()
 	{
-		static_assert(TIsDerivedFrom<T, FMassTag>::IsDerived, "Given struct doesn't represent a valid mass tag type. Make sure to inherit from FMassTag or one of its child-types.");
+		static_assert(UE::Mass::CTag<T>, "Given struct doesn't represent a valid mass tag type. Make sure to inherit from FMassTag or one of its child-types.");
 		Composition.Tags.Remove<T>();
 	}
 
 	void RemoveTag(const UScriptStruct& TagType)
 	{
-		checkf(TagType.IsChildOf(FMassTag::StaticStruct()), TEXT("Given struct doesn't represent a valid mass tag type. Make sure to inherit from FMassTag or one of its child-types."));
+		checkf(UE::Mass::IsA<FMassTag>(&TagType), TEXT("Given struct doesn't represent a valid mass tag type. Make sure to inherit from FMassTag or one of its child-types."));
 		Composition.Tags.Remove(TagType);
 	}
 
@@ -206,14 +206,14 @@ struct MASSSPAWNER_API FMassEntityTemplateData
 	template<typename T>
 	void AddChunkFragment()
 	{
-		static_assert(TIsDerivedFrom<T, FMassChunkFragment>::IsDerived, "Given struct doesn't represent a valid chunk fragment type. Make sure to inherit from FMassChunkFragment or one of its child-types.");
+		static_assert(UE::Mass::CChunkFragment<T>, "Given struct doesn't represent a valid chunk fragment type. Make sure to inherit from FMassChunkFragment or one of its child-types.");
 		Composition.ChunkFragments.Add<T>();
 	}
 
 	void AddConstSharedFragment(const FConstSharedStruct& SharedFragment)
 	{
 		const UScriptStruct* FragmentType = SharedFragment.GetScriptStruct();
-		if(ensureMsgf(FragmentType && FragmentType->IsChildOf(FMassConstSharedFragment::StaticStruct()), TEXT("Given struct doesn't represent a valid const shared fragment type. Make sure to inherit from FMassConstSharedFragment or one of its child-types.")))
+		if(ensureMsgf(UE::Mass::IsA<FMassConstSharedFragment>(FragmentType), TEXT("Given struct doesn't represent a valid const shared fragment type. Make sure to inherit from FMassConstSharedFragment or one of its child-types.")))
 		{
 			if (!Composition.ConstSharedFragments.Contains(*FragmentType))
 			{
@@ -234,7 +234,7 @@ struct MASSSPAWNER_API FMassEntityTemplateData
 	void AddSharedFragment(const FSharedStruct& SharedFragment)
 	{
 		const UScriptStruct* FragmentType = SharedFragment.GetScriptStruct();
-		if(ensureMsgf(FragmentType && FragmentType->IsChildOf(FMassSharedFragment::StaticStruct()), TEXT("Given struct doesn't represent a valid shared fragment type. Make sure to inherit from FMassSharedFragment or one of its child-types.")))
+		if(ensureMsgf(UE::Mass::IsA<FMassSharedFragment>(FragmentType), TEXT("Given struct doesn't represent a valid shared fragment type. Make sure to inherit from FMassSharedFragment or one of its child-types.")))
 		{
 			if (!Composition.SharedFragments.Contains(*FragmentType))
 			{
diff --git a/Engine/Plugins/Runtime/MassGameplay/Source/MassSpawner/Public/MassEntityTemplateRegistry.h b/Engine/Plugins/Runtime/MassGameplay/Source/MassSpawner/Public/MassEntityTemplateRegistry.h
--- a/Engine/Plugins/Runtime/MassGameplay/Source/MassSpawner/Public/MassEntityTemplateRegistry.h
+++ b/Engine/Plugins/Runtime/MassGameplay/Source/MassSpawner/Public/MassEntityTemplateRegistry.h
@@ -141,7 +141,7 @@ struct FMassEntityTemplateBuildContext
 	 */
 	void RemoveTag(const UScriptStruct& TagType)
 	{
-		checkf(TagType.IsChildOf(FMassTag::StaticStruct()), TEXT("Given struct doesn't represent a valid mass tag type. Make sure to inherit from FMassTag or one of its child-types."));
+		checkf(UE::Mass::IsA<FMassTag>(&TagType), TEXT("Given struct doesn't represent a valid mass tag type. Make sure to inherit from FMassTag or one of its child-types."));
 		RemovedTypes.Add({&TagType
 #if WITH_EDITORONLY_DATA
 			, TraitsData.Last().Trait
@@ -229,14 +229,14 @@ struct FMassEntityTemplateBuildContext
 	template<typename T>
 	void RequireFragment()
 	{
-		static_assert(TIsDerivedFrom<T, FMassTag>::IsDerived == false, "Given struct type is a valid fragment type.");
+		static_assert(UE::Mass::CTag<T> == false, "Given struct type is a valid fragment type.");
 		AddDependency(T::StaticStruct());
 	}
 
 	template<typename T>
 	void RequireTag()
 	{
-		static_assert(TIsDerivedFrom<T, FMassTag>::IsDerived, "Given struct type is not a valid tag type.");
+		static_assert(UE::Mass::CTag<T>, "Given struct type is not a valid tag type.");
 		AddDependency(T::StaticStruct());
 	}
 
diff --git a/Engine/Source/Runtime/MassEntity/Private/MassBitSetRegistry.cpp b/Engine/Source/Runtime/MassEntity/Private/MassBitSetRegistry.cpp
--- a/Engine/Source/Runtime/MassEntity/Private/MassBitSetRegistry.cpp
+++ b/Engine/Source/Runtime/MassEntity/Private/MassBitSetRegistry.cpp
@@ -9,7 +9,7 @@ namespace UE::Mass
 		: StructTracker(FMassFragment::StaticStruct()
 			, [](const UStruct* Struct)
 				{
-					return Struct !=  nullptr;
+					return UE::Mass::IsA<FMassFragment>(Struct);
 				})
 	{	
 	}
@@ -19,7 +19,7 @@ namespace UE::Mass
 		: StructTracker(FMassTag::StaticStruct()
 			, [](const UStruct* Struct)
 				{
-					return Struct && (Struct->GetStructureSize() == 1);
+					return UE::Mass::IsA<FMassTag>(Struct);
 				})
 	{	
 	}
diff --git a/Engine/Source/Runtime/MassEntity/Private/MassEntityManager.cpp b/Engine/Source/Runtime/MassEntity/Private/MassEntityManager.cpp
--- a/Engine/Source/Runtime/MassEntity/Private/MassEntityManager.cpp
+++ b/Engine/Source/Runtime/MassEntity/Private/MassEntityManager.cpp
@@ -267,37 +267,45 @@ void FMassEntityManager::Initialize(const FMassEntityManagerStorageInitParams& I
 	FMassTagBitSet Tags;
 	FMassChunkFragmentBitSet ChunkFragments;
 	FMassSharedFragmentBitSet LocalSharedFragments;
+	FMassConstSharedFragmentBitSet LocalConstSharedFragments;
 
 	for (TObjectIterator<UScriptStruct> StructIt; StructIt; ++StructIt)
 	{
-		if (StructIt->IsChildOf(FMassFragment::StaticStruct()))
+		if (UE::Mass::IsA<FMassFragment>(*StructIt))
 		{
 			if (*StructIt != FMassFragment::StaticStruct())
 			{
 				Fragments.Add(**StructIt);
 			}
 		}
-		else if (StructIt->IsChildOf(FMassTag::StaticStruct()))
+		else if (UE::Mass::IsA<FMassTag>(*StructIt))
 		{
 			if (*StructIt != FMassTag::StaticStruct())
 			{
 				Tags.Add(**StructIt);
 			}
 		}
-		else if (StructIt->IsChildOf(FMassChunkFragment::StaticStruct()))
+		else if (UE::Mass::IsA<FMassChunkFragment>(*StructIt))
 		{
 			if (*StructIt != FMassChunkFragment::StaticStruct())
 			{
 				ChunkFragments.Add(**StructIt);
 			}
 		}
-		else if (StructIt->IsChildOf(FMassSharedFragment::StaticStruct()))
+		else if (UE::Mass::IsA<FMassSharedFragment>(*StructIt))
 		{
 			if (*StructIt != FMassSharedFragment::StaticStruct())
 			{
 				LocalSharedFragments.Add(**StructIt);
 			}
 		}
+		else if (UE::Mass::IsA<FMassConstSharedFragment>(*StructIt))
+		{
+			if (*StructIt != FMassConstSharedFragment::StaticStruct())
+			{
+				LocalConstSharedFragments.Add(**StructIt);
+			}
+		}
 	}
 #if WITH_MASSENTITY_DEBUG
 	RequirementAccessDetector.Initialize();
@@ -565,15 +573,15 @@ void FMassEntityManager::InternalAppendFragmentsAndTagsToArchetypeCompositionDes
 {
 	for (const UScriptStruct* Type : FragmentsAndTagsList)
 	{
-		if (Type->IsChildOf(FMassFragment::StaticStruct()))
+		if (UE::Mass::IsA<FMassFragment>(Type))
 		{
 			InOutComposition.Fragments.Add(*Type);
 		}
-		else if (Type->IsChildOf(FMassTag::StaticStruct()))
+		else if (UE::Mass::IsA<FMassTag>(Type))
 		{
 			InOutComposition.Tags.Add(*Type);
 		}
-		else if (Type->IsChildOf(FMassChunkFragment::StaticStruct()))
+		else if (UE::Mass::IsA<FMassChunkFragment>(Type))
 		{
 			InOutComposition.ChunkFragments.Add(*Type);
 		}
@@ -1308,8 +1316,8 @@ void FMassEntityManager::SwapTagsForEntity(FMassEntityHandle Entity, const UScri
 
 	CheckIfEntityIsActive(Entity);
 
-	checkf((OldTagType != nullptr) && OldTagType->IsChildOf(FMassTag::StaticStruct()), TEXT("%hs works only with tags while '%s' is not one."), __FUNCTION__, *GetPathNameSafe(OldTagType));
-	checkf((NewTagType != nullptr) && NewTagType->IsChildOf(FMassTag::StaticStruct()), TEXT("%hs works only with tags while '%s' is not one."), __FUNCTION__, *GetPathNameSafe(NewTagType));
+	checkf(UE::Mass::IsA<FMassTag>(OldTagType), TEXT("%hs works only with tags while '%s' is not one."), __FUNCTION__, *GetPathNameSafe(OldTagType));
+	checkf(UE::Mass::IsA<FMassTag>(NewTagType), TEXT("%hs works only with tags while '%s' is not one."), __FUNCTION__, *GetPathNameSafe(NewTagType));
 	
 	TSharedPtr<FMassArchetypeData>& CurrentArchetype = GetEntityStorageInterface().GetArchetypeAsShared(Entity.Index);
 	check(CurrentArchetype);
@@ -1333,7 +1341,7 @@ void FMassEntityManager::SwapTagsForEntity(FMassEntityHandle Entity, const UScri
 void FMassEntityManager::AddTagToEntity(FMassEntityHandle Entity, const UScriptStruct* TagType)
 {
 	CHECK_SYNC_API();
-	checkf((TagType != nullptr) && TagType->IsChildOf(FMassTag::StaticStruct()), TEXT("%hs works only with tags while '%s' is not one."), __FUNCTION__, *GetPathNameSafe(TagType));
+	checkf(UE::Mass::IsA<FMassTag>(TagType), TEXT("%hs works only with tags while '%s' is not one."), __FUNCTION__, *GetPathNameSafe(TagType));
 
 	CheckIfEntityIsActive(Entity);
 	
@@ -1366,7 +1374,7 @@ void FMassEntityManager::AddTagToEntity(FMassEntityHandle Entity, const UScriptS
 void FMassEntityManager::RemoveTagFromEntity(FMassEntityHandle Entity, const UScriptStruct* TagType)
 {
 	CHECK_SYNC_API();
-	checkf((TagType != nullptr) && TagType->IsChildOf(FMassTag::StaticStruct()), TEXT("%hs works only with tags while '%s' is not one."), __FUNCTION__, *GetPathNameSafe(TagType));
+	checkf(UE::Mass::IsA<FMassTag>(TagType), TEXT("%hs works only with tags while '%s' is not one."), __FUNCTION__, *GetPathNameSafe(TagType));
 
 	CheckIfEntityIsActive(Entity);
 
@@ -1449,7 +1457,7 @@ bool FMassEntityManager::RemoveConstSharedFragmentFromEntity(const FMassEntityHa
 {
 	CHECK_SYNC_API();
 
-	if (!ensureMsgf(ConstSharedFragmentType.IsChildOf(FMassConstSharedFragment::StaticStruct()), TEXT("%hs parameter ConstSharedFragmentType is expected to be a FMassConstSharedFragment"), __FUNCTION__))
+	if (!ensureMsgf(UE::Mass::IsA<FMassConstSharedFragment>(&ConstSharedFragmentType), TEXT("%hs parameter ConstSharedFragmentType is expected to be a FMassConstSharedFragment"), __FUNCTION__))
 	{
 		return false;
 	}
@@ -1542,7 +1550,7 @@ bool FMassEntityManager::RemoveSharedFragmentFromEntity(const FMassEntityHandle
 {
 	CHECK_SYNC_API();
 
-	if (!ensureMsgf(SharedFragmentType.IsChildOf(FMassSharedFragment::StaticStruct()), TEXT("%hs parameter SharedFragmentType is expected to be a FMassSharedFragment"), __FUNCTION__))
+	if (!ensureMsgf(UE::Mass::IsA<FMassSharedFragment>(&SharedFragmentType), TEXT("%hs parameter SharedFragmentType is expected to be a FMassSharedFragment"), __FUNCTION__))
 	{
 		return false;
 	}
diff --git a/Engine/Source/Runtime/MassEntity/Private/MassObserverManager.cpp b/Engine/Source/Runtime/MassEntity/Private/MassObserverManager.cpp
--- a/Engine/Source/Runtime/MassEntity/Private/MassObserverManager.cpp
+++ b/Engine/Source/Runtime/MassEntity/Private/MassObserverManager.cpp
@@ -280,9 +280,10 @@ bool FMassObserverManager::OnCompositionChanged(const FMassEntityHandle Entity,
 
 void FMassObserverManager::OnFragmentOrTagOperation(const UScriptStruct& FragmentOrTagType, const FMassArchetypeEntityCollection& EntityCollection, const EMassObservedOperation Operation)
 {
-	check(FragmentOrTagType.IsChildOf(FMassFragment::StaticStruct()) || FragmentOrTagType.IsChildOf(FMassTag::StaticStruct()));
+	const bool bIsFragment = UE::Mass::IsA<FMassFragment>(&FragmentOrTagType);
+	check(bIsFragment || UE::Mass::IsA<FMassTag>(&FragmentOrTagType));
 
-	if (FragmentOrTagType.IsChildOf(FMassFragment::StaticStruct()))
+	if (bIsFragment)
 	{
 		if (ObservedFragments[(uint8)Operation].Contains(FragmentOrTagType))
 		{
@@ -326,11 +327,12 @@ void FMassObserverManager::HandleSingleEntityImpl(const UScriptStruct& FragmentT
 
 void FMassObserverManager::AddObserverInstance(const UScriptStruct& FragmentOrTagType, const EMassObservedOperation Operation, UMassProcessor& ObserverProcessor)
 {
-	checkSlow(FragmentOrTagType.IsChildOf(FMassFragment::StaticStruct()) || FragmentOrTagType.IsChildOf(FMassTag::StaticStruct()));
+	const bool bIsFragment = UE::Mass::IsA<FMassFragment>(&FragmentOrTagType);
+	checkSlow(bIsFragment || UE::Mass::IsA<FMassTag>(&FragmentOrTagType));
 
 	FMassRuntimePipeline* Pipeline = nullptr;
 
-	if (FragmentOrTagType.IsChildOf(FMassFragment::StaticStruct()))
+	if (bIsFragment)
 	{
 		Pipeline = &(*FragmentObservers[(uint8)Operation]).FindOrAdd(&FragmentOrTagType);
 		ObservedFragments[(uint8)Operation].Add(FragmentOrTagType);
@@ -351,13 +353,13 @@ void FMassObserverManager::AddObserverInstance(const UScriptStruct& FragmentOrTa
 
 void FMassObserverManager::RemoveObserverInstance(const UScriptStruct& FragmentOrTagType, const EMassObservedOperation Operation, UMassProcessor& ObserverProcessor)
 {
-	if (!ensure(FragmentOrTagType.IsChildOf(FMassFragment::StaticStruct()) || FragmentOrTagType.IsChildOf(FMassTag::StaticStruct())))
+	const bool bIsFragmentObserver = UE::Mass::IsA<FMassFragment>(&FragmentOrTagType);
+
+	if (!ensure(bIsFragmentObserver || UE::Mass::IsA<FMassTag>(&FragmentOrTagType)))
 	{
 		return;
 	}
 
-	bool bIsFragmentObserver = FragmentOrTagType.IsChildOf(FMassFragment::StaticStruct());
-
 	TMap<TObjectPtr<const UScriptStruct>, FMassRuntimePipeline>& ObserversMap =
 		bIsFragmentObserver ? *FragmentObservers[(uint8)Operation] : *TagObservers[(uint8)Operation];
 
diff --git a/Engine/Source/Runtime/MassEntity/Private/MassObserverRegistry.cpp b/Engine/Source/Runtime/MassEntity/Private/MassObserverRegistry.cpp
--- a/Engine/Source/Runtime/MassEntity/Private/MassObserverRegistry.cpp
+++ b/Engine/Source/Runtime/MassEntity/Private/MassObserverRegistry.cpp
@@ -16,9 +16,10 @@ UMassObserverRegistry::UMassObserverRegistry()
 void UMassObserverRegistry::RegisterObserver(const UScriptStruct& ObservedType, const EMassObservedOperation Operation, TSubclassOf<UMassProcessor> ObserverClass)
 {
 	check(ObserverClass);
-	checkSlow(ObservedType.IsChildOf(FMassFragment::StaticStruct()) || ObservedType.IsChildOf(FMassTag::StaticStruct()));
+	const bool bIsFragment = UE::Mass::IsA<FMassFragment>(&ObservedType);
+	checkSlow(bIsFragment || UE::Mass::IsA<FMassTag>(&ObservedType));
 
-	if (ObservedType.IsChildOf(FMassFragment::StaticStruct()))
+	if (bIsFragment)
 	{
 		(*FragmentObservers[(uint8)Operation]).FindOrAdd(&ObservedType).ClassCollection.AddUnique(ObserverClass);
 	}
diff --git a/Engine/Source/Runtime/MassEntity/Public/MassCommandBuffer.h b/Engine/Source/Runtime/MassEntity/Public/MassCommandBuffer.h
--- a/Engine/Source/Runtime/MassEntity/Public/MassCommandBuffer.h
+++ b/Engine/Source/Runtime/MassEntity/Public/MassCommandBuffer.h
@@ -73,28 +73,28 @@ public:
 	template<typename T>
 	void AddFragment(FMassEntityHandle Entity)
 	{
-		static_assert(TIsDerivedFrom<T, FMassFragment>::IsDerived, "Given struct type is not a valid fragment type.");
+		static_assert(UE::Mass::CFragment<T>, "Given struct type is not a valid fragment type.");
 		PushCommand<FMassCommandAddFragmentsInternal<EMassCommandCheckTime::CompileTimeCheck, T>>(Entity);
 	}
 
 	template<typename T>
 	void AddFragment_RuntimeCheck(FMassEntityHandle Entity)
 	{
-		checkf(T::StaticStruct()->IsChildOf(FMassFragment::StaticStruct()), TEXT("Given struct type is not a valid fragment type."));
+		checkf(UE::Mass::IsA<FMassFragment>(T::StaticStruct()), TEXT("Given struct type is not a valid fragment type."));
 		PushCommand<FMassCommandAddFragmentsInternal<EMassCommandCheckTime::RuntimeCheck, T>>(Entity);
 	}
 
 	template<typename T>
 	void RemoveFragment(FMassEntityHandle Entity)
 	{
-		static_assert(TIsDerivedFrom<T, FMassFragment>::IsDerived, "Given struct type is not a valid fragment type.");
+		static_assert(UE::Mass::CFragment<T>, "Given struct type is not a valid fragment type.");
 		PushCommand<FMassCommandRemoveFragmentsInternal<EMassCommandCheckTime::CompileTimeCheck, T>>(Entity);
 	}
 
 	template<typename T>
 	void RemoveFragment_RuntimeCheck(FMassEntityHandle Entity)
 	{
-		checkf(T::StaticStruct()->IsChildOf(FMassFragment::StaticStruct()), TEXT("Given struct type is not a valid fragment type."));
+		checkf(UE::Mass::IsA<FMassFragment>(T::StaticStruct()), TEXT("Given struct type is not a valid fragment type."));
 		PushCommand<FMassCommandRemoveFragmentsInternal<EMassCommandCheckTime::RuntimeCheck, T>>(Entity);
 	}
 
@@ -102,14 +102,14 @@ public:
 	template<typename T>
 	void AddTag(FMassEntityHandle Entity)
 	{
-		static_assert(TIsDerivedFrom<T, FMassTag>::IsDerived, "Given struct type is not a valid tag type.");
+		static_assert(UE::Mass::CTag<T>, "Given struct type is not a valid tag type.");
 		PushCommand<FMassCommandAddTagsInternal<EMassCommandCheckTime::CompileTimeCheck, T>>(Entity);
 	}
 
 	template<typename T>
 	void AddTag_RuntimeCheck(FMassEntityHandle Entity)
 	{
-		checkf(T::StaticStruct()->IsChildOf(FMassTag::StaticStruct()), TEXT("Given struct type is not a valid tag type."));
+		checkf(UE::Mass::IsA<FMassTag>(T::StaticStruct()), TEXT("Given struct type is not a valid tag type."));
 		PushCommand<FMassCommandAddTagsInternal<EMassCommandCheckTime::RuntimeCheck, T>>(Entity);
 	}
 
@@ -117,14 +117,14 @@ public:
 	template<typename T>
 	void RemoveTag(FMassEntityHandle Entity)
 	{
-		static_assert(TIsDerivedFrom<T, FMassTag>::IsDerived, "Given struct type is not a valid tag type.");
+		static_assert(UE::Mass::CTag<T>, "Given struct type is not a valid tag type.");
 		PushCommand<FMassCommandRemoveTagsInternal<EMassCommandCheckTime::CompileTimeCheck, T>>(Entity);
 	}
 
 	template<typename T>
 	void RemoveTag_RuntimeCheck(FMassEntityHandle Entity)
 	{
-		checkf(T::StaticStruct()->IsChildOf(FMassTag::StaticStruct()), TEXT("Given struct type is not a valid tag type."));
+		checkf(UE::Mass::IsA<FMassTag>(T::StaticStruct()), TEXT("Given struct type is not a valid tag type."));
 		PushCommand<FMassCommandRemoveTagsInternal<EMassCommandCheckTime::RuntimeCheck, T>>(Entity);
 	}
 
@@ -132,16 +132,16 @@ public:
 	template<typename TOld, typename TNew>
 	void SwapTags(FMassEntityHandle Entity)
 	{
-		static_assert(TIsDerivedFrom<TOld, FMassTag>::IsDerived, "Given struct type is not a valid tag type.");
-		static_assert(TIsDerivedFrom<TNew, FMassTag>::IsDerived, "Given struct type is not a valid tag type.");
+		static_assert(UE::Mass::CTag<TOld>, "Given struct type is not a valid tag type.");
+		static_assert(UE::Mass::CTag<TNew>, "Given struct type is not a valid tag type.");
 		PushCommand<FMassCommandSwapTagsInternal<EMassCommandCheckTime::CompileTimeCheck, TOld, TNew>>(Entity);
 	}
 
 	template<typename TOld, typename TNew>
 	void SwapTags_RuntimeCheck(FMassEntityHandle Entity)
 	{
-		checkf(TOld::StaticStruct()->IsChildOf(FMassTag::StaticStruct()), TEXT("Given struct type is not a valid tag type."));
-		checkf(TNew::StaticStruct()->IsChildOf(FMassTag::StaticStruct()), TEXT("Given struct type is not a valid tag type."));
+		checkf(UE::Mass::IsA<FMassTag>(TOld::StaticStruct()), TEXT("Given struct type is not a valid tag type."));
+		checkf(UE::Mass::IsA<FMassTag>(TNew::StaticStruct()), TEXT("Given struct type is not a valid tag type."));
 		PushCommand<FMassCommandSwapTagsInternal<EMassCommandCheckTime::RuntimeCheck, TOld, TNew>>(Entity);
 	}
 
diff --git a/Engine/Source/Runtime/MassEntity/Public/MassEntityElementTypes.h b/Engine/Source/Runtime/MassEntity/Public/MassEntityElementTypes.h
--- a/Engine/Source/Runtime/MassEntity/Public/MassEntityElementTypes.h
+++ b/Engine/Source/Runtime/MassEntity/Public/MassEntityElementTypes.h
@@ -37,8 +37,41 @@ struct FMassConstSharedFragment
 	GENERATED_BODY()
 };
 
-USTRUCT()
-struct FMassRelationship
+namespace UE::Mass
 {
-	GENERATED_BODY()
-};
+	template<typename T>
+	bool IsA(const UStruct* /*Struct*/)
+	{
+		return false;
+	}
+
+	template<>
+	inline bool IsA<FMassFragment>(const UStruct* Struct)
+	{
+		return Struct && Struct->IsChildOf(FMassFragment::StaticStruct());
+	}
+
+	template<>
+	inline bool IsA<FMassTag>(const UStruct* Struct)
+	{
+		return Struct && Struct->IsChildOf(FMassTag::StaticStruct());
+	}
+
+	template<>
+	inline bool IsA<FMassChunkFragment>(const UStruct* Struct)
+	{
+		return Struct && Struct->IsChildOf(FMassChunkFragment::StaticStruct());
+	}
+
+	template<>
+	inline bool IsA<FMassSharedFragment>(const UStruct* Struct)
+	{
+		return Struct && Struct->IsChildOf(FMassSharedFragment::StaticStruct());
+	}
+
+	template<>
+	inline bool IsA<FMassConstSharedFragment>(const UStruct* Struct)
+	{
+		return Struct && Struct->IsChildOf(FMassConstSharedFragment::StaticStruct());
+	}
+}
diff --git a/Engine/Source/Runtime/MassEntity/Public/MassEntityManager.h b/Engine/Source/Runtime/MassEntity/Public/MassEntityManager.h
--- a/Engine/Source/Runtime/MassEntity/Public/MassEntityManager.h
+++ b/Engine/Source/Runtime/MassEntity/Public/MassEntityManager.h
@@ -529,7 +529,7 @@ public:
 	template<typename FragmentType>
 	FragmentType& GetFragmentDataChecked(FMassEntityHandle EntityHandle) const
 	{
-		static_assert(TIsDerivedFrom<FragmentType, FMassFragment>::IsDerived
+		static_assert(UE::Mass::CFragment<FragmentType>
 			, "Given struct doesn't represent a valid fragment type. Make sure to inherit from FMassFragment or one of its child-types.");
 		return *((FragmentType*)InternalGetFragmentDataChecked(EntityHandle, FragmentType::StaticStruct()));
 	}
@@ -537,14 +537,14 @@ public:
 	template<typename FragmentType>
 	FragmentType* GetFragmentDataPtr(FMassEntityHandle EntityHandle) const
 	{
-		static_assert(TIsDerivedFrom<FragmentType, FMassFragment>::IsDerived
+		static_assert(UE::Mass::CFragment<FragmentType>
 			, "Given struct doesn't represent a valid fragment type. Make sure to inherit from FMassFragment or one of its child-types.");
 		return (FragmentType*)InternalGetFragmentDataPtr(EntityHandle, FragmentType::StaticStruct());
 	}
 
 	FStructView GetFragmentDataStruct(FMassEntityHandle EntityHandle, const UScriptStruct* FragmentType) const
 	{
-		checkf((FragmentType != nullptr) && FragmentType->IsChildOf(FMassFragment::StaticStruct())
+		checkf(UE::Mass::IsA<FMassFragment>(FragmentType)
 			, TEXT("GetFragmentDataStruct called with an invalid fragment type '%s'"), *GetPathNameSafe(FragmentType));
 		return FStructView(FragmentType, static_cast<uint8*>(InternalGetFragmentDataPtr(EntityHandle, FragmentType)));
 	}
@@ -552,7 +552,7 @@ public:
 	template<typename ConstSharedFragmentType>
 	ConstSharedFragmentType* GetConstSharedFragmentDataPtr(FMassEntityHandle EntityHandle) const
 	{
-		static_assert(TIsDerivedFrom<ConstSharedFragmentType, FMassConstSharedFragment>::IsDerived, "Given struct doesn't represent a valid const shared fragment type. Make sure to inherit from FMassConstSharedFragment or one of its child-types.");
+		static_assert(UE::Mass::CConstSharedFragment<ConstSharedFragmentType>, "Given struct doesn't represent a valid const shared fragment type. Make sure to inherit from FMassConstSharedFragment or one of its child-types.");
 		const FConstSharedStruct* ConstSharedStruct = InternalGetConstSharedFragmentPtr(EntityHandle, ConstSharedFragmentType::StaticStruct());
 		return (ConstSharedFragmentType*)(ConstSharedStruct ? ConstSharedStruct->GetMemory() : nullptr);
 	}
@@ -567,7 +567,7 @@ public:
 
 	FConstStructView GetConstSharedFragmentDataStruct(FMassEntityHandle EntityHandle, const UScriptStruct* ConstSharedFragmentType) const
 	{
-		checkf((ConstSharedFragmentType != nullptr) && ConstSharedFragmentType->IsChildOf(FMassConstSharedFragment::StaticStruct())
+		checkf(UE::Mass::IsA<FMassConstSharedFragment>(ConstSharedFragmentType)
 			, TEXT("GetConstSharedFragmentDataStruct called with an invalid fragment type '%s'"), *GetPathNameSafe(ConstSharedFragmentType));
 		const FConstSharedStruct* ConstSharedStruct = InternalGetConstSharedFragmentPtr(EntityHandle, ConstSharedFragmentType);
 		return ConstSharedStruct
@@ -578,7 +578,7 @@ public:
 	template<typename SharedFragmentType>
 	TConstArrayView<FSharedStruct> GetSharedFragmentsOfType()
 	{
-		static_assert(TIsDerivedFrom<SharedFragmentType, FMassSharedFragment>::IsDerived
+		static_assert(UE::Mass::CSharedFragment<SharedFragmentType>
 			, "Given struct doesn't represent a valid shared fragment type. Make sure to inherit from FMassSharedFragment or one of its child-types.");
 
 		TArray<FSharedStruct>* InstancesOfType = SharedFragmentsContainer.Find(SharedFragmentType::StaticStruct());
@@ -588,7 +588,7 @@ public:
 	template<typename SharedFragmentType>
 	SharedFragmentType* GetSharedFragmentDataPtr(FMassEntityHandle EntityHandle) const
 	{
-		static_assert(TIsDerivedFrom<SharedFragmentType, FMassSharedFragment>::IsDerived
+		static_assert(UE::Mass::CSharedFragment<SharedFragmentType>
 			, "Given struct doesn't represent a valid shared fragment type. Make sure to inherit from FMassSharedFragment or one of its child-types.");
 		const FSharedStruct* FragmentPtr = InternalGetSharedFragmentPtr(EntityHandle, SharedFragmentType::StaticStruct());
 		return (SharedFragmentType*)(FragmentPtr ? FragmentPtr->GetMemory() : nullptr);
@@ -604,7 +604,7 @@ public:
 
 	FConstStructView GetSharedFragmentDataStruct(FMassEntityHandle EntityHandle, const UScriptStruct* SharedFragmentType) const
 	{
-		checkf((SharedFragmentType != nullptr) && SharedFragmentType->IsChildOf(FMassSharedFragment::StaticStruct())
+		checkf(UE::Mass::IsA<FMassSharedFragment>(SharedFragmentType)
 			, TEXT("GetSharedFragmentDataStruct called with an invalid fragment type '%s'"), *GetPathNameSafe(SharedFragmentType));
 		const FSharedStruct* FragmentPtr = InternalGetSharedFragmentPtr(EntityHandle, SharedFragmentType);
 		return FragmentPtr
@@ -671,7 +671,7 @@ public:
 	UE_DEPRECATED(5.5, "This method will no longer be exposed. Use GetOrCreateConstSharedFragment instead.")
 	const FConstSharedStruct& GetOrCreateConstSharedFragmentByHash(const uint32 Hash, const T& Fragment)
 	{
-		static_assert(TIsDerivedFrom<T, FMassConstSharedFragment>::IsDerived, "Given struct doesn't represent a valid const shared fragment type. Make sure to inherit from FMassConstSharedFragment or one of its child-types.");
+		static_assert(UE::Mass::CConstSharedFragment<T>, "Given struct doesn't represent a valid const shared fragment type. Make sure to inherit from FMassConstSharedFragment or one of its child-types.");
 		return ConstSharedFragmentsContainer.FindOrAdd(Hash, T::StaticStruct(), reinterpret_cast<const uint8*>(&Fragment));
 	}
 
@@ -679,7 +679,7 @@ private:
 	template<typename T>
 	const FSharedStruct& GetOrCreateSharedFragmentByHash(const uint32 Hash, const T& Fragment)
 	{
-		static_assert(TIsDerivedFrom<T, FMassSharedFragment>::IsDerived, "Given struct doesn't represent a valid shared fragment type. Make sure to inherit from FMassSharedFragment or one of its child-types.");
+		static_assert(UE::Mass::CSharedFragment<T>, "Given struct doesn't represent a valid shared fragment type. Make sure to inherit from FMassSharedFragment or one of its child-types.");
 		return SharedFragmentsContainer.FindOrAdd(Hash, T::StaticStruct(), reinterpret_cast<const uint8*>(&Fragment));
 	}
 
@@ -696,7 +696,7 @@ private:
 	template<typename T, typename... TArgs>
 	const FConstSharedStruct& GetOrCreateConstSharedFragmentByHash(const uint32 Hash, TArgs&&... InArgs)
 	{
-		static_assert(TIsDerivedFrom<T, FMassConstSharedFragment>::IsDerived, "Given struct doesn't represent a valid const shared fragment type. Make sure to inherit from FMassConstSharedFragment or one of its child-types.");
+		static_assert(UE::Mass::CConstSharedFragment<T>, "Given struct doesn't represent a valid const shared fragment type. Make sure to inherit from FMassConstSharedFragment or one of its child-types.");
 		return ConstSharedFragmentsContainer.FindOrAdd<T>(Hash, Forward<TArgs>(InArgs)...);
 	}
 
@@ -722,7 +722,7 @@ public:
 	UE_DEPRECATED(5.5, "This method will no longer be exposed. Use GetOrCreateSharedFragment instead.")
 	const FSharedStruct& GetOrCreateSharedFragmentByHash(const uint32 Hash, TArgs&&... InArgs)
 	{
-		static_assert(TIsDerivedFrom<T, FMassSharedFragment>::IsDerived, "Given struct doesn't represent a valid shared fragment type. Make sure to inherit from FMassSharedFragment or one of its child-types.");
+		static_assert(UE::Mass::CSharedFragment<T>, "Given struct doesn't represent a valid shared fragment type. Make sure to inherit from FMassSharedFragment or one of its child-types.");
 		return SharedFragmentsContainer.FindOrAdd<T>(Hash, Forward<TArgs>(InArgs)...);
 	}
 
@@ -751,7 +751,7 @@ public:
 	template<typename T>
 	const FConstSharedStruct& GetOrCreateConstSharedFragment(const T& Fragment)
 	{
-		static_assert(TIsDerivedFrom<T, FMassConstSharedFragment>::IsDerived,
+		static_assert(UE::Mass::CConstSharedFragment<T>,
 			"Given struct doesn't represent a valid const shared fragment type. Make sure to inherit from FMassConstSharedFragment or one of its child-types.");
 		const uint32 Hash = UE::StructUtils::GetStructCrc32(FConstStructView::Make(Fragment));
 PRAGMA_DISABLE_DEPRECATION_WARNINGS
@@ -784,7 +784,7 @@ PRAGMA_ENABLE_DEPRECATION_WARNINGS
 	template<typename T>
 	const FSharedStruct& GetOrCreateSharedFragment(const T& Fragment)
 	{
-		static_assert(TIsDerivedFrom<T, FMassSharedFragment>::IsDerived,
+		static_assert(UE::Mass::CSharedFragment<T>,
 			"Given struct doesn't represent a valid shared fragment type. Make sure to inherit from FMassSharedFragment or one of its child-types.");
 		const uint32 Hash = UE::StructUtils::GetStructCrc32(FConstStructView::Make(Fragment));
 		return GetOrCreateSharedFragmentByHash(Hash, Fragment);
@@ -815,7 +815,7 @@ PRAGMA_ENABLE_DEPRECATION_WARNINGS
 	template<typename T, typename... TArgs>
 	const FConstSharedStruct& GetOrCreateConstSharedFragment(TArgs&&... InArgs)
 	{
-		static_assert(TIsDerivedFrom<T, FMassConstSharedFragment>::IsDerived,
+		static_assert(UE::Mass::CConstSharedFragment<T>,
 			"Given struct doesn't represent a valid const shared fragment type. Make sure to inherit from FMassConstSharedFragment or one of its child-types.");
 		T Struct(Forward<TArgs>(InArgs)...);
 		const uint32 Hash = UE::StructUtils::GetStructCrc32(FConstStructView::Make(Struct));
@@ -849,7 +849,7 @@ PRAGMA_ENABLE_DEPRECATION_WARNINGS
 	template<typename T, typename... TArgs>
 	const FSharedStruct& GetOrCreateSharedFragment(TArgs&&... InArgs)
 	{
-		static_assert(TIsDerivedFrom<T, FMassSharedFragment>::IsDerived,
+		static_assert(UE::Mass::CSharedFragment<T>,
 			"Given struct doesn't represent a valid shared fragment type. Make sure to inherit from FMassSharedFragment or one of its child-types.");
 		T Struct(Forward<TArgs>(InArgs)...);
 		const uint32 Hash = UE::StructUtils::GetStructCrc32(FConstStructView::Make(Struct));
@@ -934,7 +934,7 @@ PRAGMA_ENABLE_DEPRECATION_WARNINGS
 	template<typename T, typename... TArgs>
 	const FConstSharedStruct& GetOrCreateConstSharedFragment(const FConstStructView HashingHelperStruct, TArgs&&... InArgs)
 	{
-		static_assert(TIsDerivedFrom<T, FMassConstSharedFragment>::IsDerived,
+		static_assert(UE::Mass::CConstSharedFragment<T>,
 			"Given struct doesn't represent a valid const shared fragment type. Make sure to inherit from FMassConstSharedFragment or one of its child-types.");
 		T Fragment(Forward<TArgs>(InArgs)...);
 		const uint32 Hash = UE::StructUtils::GetStructCrc32(HashingHelperStruct);
@@ -981,7 +981,7 @@ PRAGMA_ENABLE_DEPRECATION_WARNINGS
 	template<typename T, typename... TArgs>
 	const FSharedStruct& GetOrCreateSharedFragment(const FConstStructView HashingHelperStruct, TArgs&&... InArgs)
 	{
-		static_assert(TIsDerivedFrom<T, FMassSharedFragment>::IsDerived,
+		static_assert(UE::Mass::CSharedFragment<T>,
 			"Given struct doesn't represent a valid shared fragment type. Make sure to inherit from FMassSharedFragment or one of its child-types.");
 		T Fragment(Forward<TArgs>(InArgs)...);
 		const uint32 Hash = UE::StructUtils::GetStructCrc32(HashingHelperStruct);
diff --git a/Engine/Source/Runtime/MassEntity/Public/MassEntityTypes.h b/Engine/Source/Runtime/MassEntity/Public/MassEntityTypes.h
--- a/Engine/Source/Runtime/MassEntity/Public/MassEntityTypes.h
+++ b/Engine/Source/Runtime/MassEntity/Public/MassEntityTypes.h
@@ -243,17 +243,14 @@ struct MASSENTITY_API FMassArchetypeSharedFragmentValues
 
 	FORCEINLINE bool ContainsType(const UScriptStruct* FragmentType) const
 	{
-		if (FragmentType)
+		if (UE::Mass::IsA<FMassSharedFragment>(FragmentType))
 		{
-			if (FragmentType->IsChildOf(FMassSharedFragment::StaticStruct()))
-			{
-				return SharedFragmentBitSet.Contains(*FragmentType);
-			}
-
-			if (FragmentType->IsChildOf(FMassConstSharedFragment::StaticStruct()))
-			{
-				return ConstSharedFragmentBitSet.Contains(*FragmentType);
-			}
+			return SharedFragmentBitSet.Contains(*FragmentType);
+		}
+
+		if (UE::Mass::IsA<FMassConstSharedFragment>(FragmentType))
+		{
+			return ConstSharedFragmentBitSet.Contains(*FragmentType);
 		}
 
 		return false;
@@ -262,11 +259,11 @@ struct MASSENTITY_API FMassArchetypeSharedFragmentValues
 	template<typename T>
 	FORCEINLINE bool ContainsType() const
 	{
-		if constexpr (TIsDerivedFrom<T, FMassConstSharedFragment>::IsDerived)
+		if constexpr (UE::Mass::CConstSharedFragment<T>)
 		{
 			return ConstSharedFragmentBitSet.Contains(*T::StaticStruct());
 		}
-		else if constexpr (TIsDerivedFrom<T, FMassSharedFragment>::IsDerived)
+		else if constexpr (UE::Mass::CSharedFragment<T>)
 		{
 			return SharedFragmentBitSet.Contains(*T::StaticStruct());
 		}
diff --git a/Engine/Source/Runtime/MassEntity/Public/MassEntityView.h b/Engine/Source/Runtime/MassEntity/Public/MassEntityView.h
--- a/Engine/Source/Runtime/MassEntity/Public/MassEntityView.h
+++ b/Engine/Source/Runtime/MassEntity/Public/MassEntityView.h
@@ -51,9 +51,9 @@ struct MASSENTITY_API FMassEntityView
 	template<typename T>
 	T& GetFragmentData() const
 	{
-		static_assert(!std::is_base_of_v<FMassTag, T>,
+		static_assert(!UE::Mass::CTag<T>,
 			"Given struct doesn't represent a valid fragment type but a tag. Use HasTag instead.");
-		static_assert(std::is_base_of_v<FMassTag, T> || std::is_base_of_v<FMassFragment, T>,
+		static_assert(UE::Mass::CFragment<T>,
 			"Given struct doesn't represent a valid fragment type. Make sure to inherit from FMassFragment or one of its child-types.");
 
 		return *((T*)GetFragmentPtrChecked(*T::StaticStruct()));
@@ -63,9 +63,9 @@ struct MASSENTITY_API FMassEntityView
 	template<typename T>
 	T* GetFragmentDataPtr() const
 	{
-		static_assert(!std::is_base_of_v<FMassTag, T>,
+		static_assert(!UE::Mass::CTag<T>,
 			"Given struct doesn't represent a valid fragment type but a tag. Use HasTag instead.");
-		static_assert(std::is_base_of_v<FMassTag, T> || std::is_base_of_v<FMassFragment, T>,
+		static_assert(UE::Mass::CFragment<T>,
 			"Given struct doesn't represent a valid fragment type. Make sure to inherit from FMassFragment or one of its child-types.");
 
 		return (T*)GetFragmentPtr(*T::StaticStruct());
@@ -81,8 +81,8 @@ struct MASSENTITY_API FMassEntityView
 	template<typename T>
 	const T* GetConstSharedFragmentDataPtr() const
 	{
-		static_assert(std::is_base_of_v<FMassConstSharedFragment, T>,
-			"Given struct doesn't represent a valid shared fragment type. Make sure to inherit from FMassConstSharedFragment or one of its child-types.");
+		static_assert(UE::Mass::CConstSharedFragment<T>,
+			"Given struct doesn't represent a valid const shared fragment type. Make sure to inherit from FMassConstSharedFragment or one of its child-types.");
 
 		return (const T*)GetConstSharedFragmentPtr(*T::StaticStruct());
 	}
@@ -91,7 +91,7 @@ struct MASSENTITY_API FMassEntityView
 	template<typename T>
 	const T& GetConstSharedFragmentData() const
 	{
-		static_assert(std::is_base_of_v<FMassConstSharedFragment, T>,
+		static_assert(UE::Mass::CConstSharedFragment<T>,
 			"Given struct doesn't represent a valid const shared fragment type. Make sure to inherit from FMassConstSharedFragment or one of its child-types.");
 
 		return *((const T*)GetConstSharedFragmentPtrChecked(*T::StaticStruct()));
@@ -99,32 +99,32 @@ struct MASSENTITY_API FMassEntityView
 
 	FConstStructView GetConstSharedFragmentDataStruct(const UScriptStruct* FragmentType) const
 	{
-		check(FragmentType && FragmentType->IsChildOf(FMassConstSharedFragment::StaticStruct()));
+		check(UE::Mass::IsA<FMassConstSharedFragment>(FragmentType));
 		return FConstStructView(FragmentType, static_cast<const uint8*>(GetConstSharedFragmentPtr(*FragmentType)));
 	}
 
 	/** will fail a check if the viewed entity doesn't have the given shared fragment */
-	template<typename T UE_REQUIRES(std::is_base_of_v<FMassSharedFragment, T>)>
+	template<UE::Mass::CSharedFragment T>
 	T& GetSharedFragmentData() const
 	{
 		return *((T*)GetSharedFragmentPtrChecked(*T::StaticStruct()));
 	}
 
 	/** if the viewed entity doesn't have the given shared fragment the function will return null */
-	template<typename T UE_REQUIRES(std::is_base_of_v<FMassSharedFragment, T>)>
+	template<UE::Mass::CSharedFragment T>
 	T* GetSharedFragmentDataPtr() const
 	{
 		return (T*)GetSharedFragmentPtr(*T::StaticStruct());
 	}
 
-	template<typename T UE_REQUIRES(std::is_base_of_v<FMassConstSharedFragment, T>)>
+	template<UE::Mass::CConstSharedFragment T>
 	UE_DEPRECATED(5.5, "Using GetSharedFragmentDataPtr with const shared fragments is deprecated. Use GetConstSharedFragmentDataPtr instead")
 	T* GetSharedFragmentDataPtr() const
 	{
 		return const_cast<T*>(GetConstSharedFragmentDataPtr<T>());
 	}
 
-	template<typename T UE_REQUIRES(std::is_base_of_v<FMassConstSharedFragment, T>)>
+	template<UE::Mass::CConstSharedFragment T>
 	UE_DEPRECATED(5.5, "Using GetSharedFragmentDataPtr with const shared fragments is deprecated. Use GetConstSharedFragmentData instead")
 	T& GetSharedFragmentData() const
 	{
@@ -135,14 +135,14 @@ struct MASSENTITY_API FMassEntityView
 
 	FStructView GetSharedFragmentDataStruct(const UScriptStruct* FragmentType) const
 	{
-		check(FragmentType && FragmentType->IsChildOf(FMassSharedFragment::StaticStruct()));
+		check(UE::Mass::IsA<FMassSharedFragment>(FragmentType));
 		return FStructView(FragmentType, static_cast<uint8*>(GetSharedFragmentPtr(*FragmentType)));
 	}
 
 	template<typename T>
 	bool HasTag() const
 	{
-		static_assert(std::is_base_of_v<FMassTag, T>, "Given struct doesn't represent a valid tag type. Make sure to inherit from FMassTag or one of its child-types.");
+		static_assert(UE::Mass::CTag<T>, "Given struct doesn't represent a valid tag type. Make sure to inherit from FMassTag or one of its child-types.");
 		return HasTag(*T::StaticStruct());
 	}
 
diff --git a/Engine/Source/Runtime/MassEntity/Public/MassExecutionContext.h b/Engine/Source/Runtime/MassEntity/Public/MassExecutionContext.h
--- a/Engine/Source/Runtime/MassEntity/Public/MassExecutionContext.h
+++ b/Engine/Source/Runtime/MassEntity/Public/MassExecutionContext.h
@@ -297,7 +297,7 @@ public:
 	template<typename T>
 	bool DoesArchetypeHaveFragment() const
 	{
-		static_assert(TIsDerivedFrom<T, FMassFragment>::IsDerived, "Given struct is not of a valid fragment type.");
+		static_assert(UE::Mass::CFragment<T>, "Given struct is not of a valid fragment type.");
 		return CurrentArchetypeCompositionDescriptor.Fragments.Contains<T>();
 	}
 
@@ -309,7 +309,7 @@ public:
 	template<typename T>
 	bool DoesArchetypeHaveTag() const
 	{
-		static_assert(TIsDerivedFrom<T, FMassTag>::IsDerived, "Given struct is not of a valid tag type.");
+		static_assert(UE::Mass::CTag<T>, "Given struct is not of a valid tag type.");
 		return CurrentArchetypeCompositionDescriptor.Tags.Contains<T>();
 	}
 
@@ -320,7 +320,7 @@ public:
 	template<typename T>
 	T* GetMutableChunkFragmentPtr()
 	{
-		static_assert(TIsDerivedFrom<T, FMassChunkFragment>::IsDerived, "Given struct doesn't represent a valid chunk fragment type. Make sure to inherit from FMassChunkFragment or one of its child-types.");
+		static_assert(UE::Mass::CChunkFragment<T>, "Given struct doesn't represent a valid chunk fragment type. Make sure to inherit from FMassChunkFragment or one of its child-types.");
 
 		const UScriptStruct* Type = T::StaticStruct();
 		FChunkFragmentView* FoundChunkFragmentData = ChunkFragmentViews.FindByPredicate([Type](const FChunkFragmentView& Element) { return Element.Requirement.StructType == Type; } );
@@ -339,7 +339,7 @@ public:
 	template<typename T>
 	const T* GetChunkFragmentPtr() const
 	{
-		static_assert(TIsDerivedFrom<T, FMassChunkFragment>::IsDerived, "Given struct doesn't represent a valid chunk fragment type. Make sure to inherit from FMassChunkFragment or one of its child-types.");
+		static_assert(UE::Mass::CChunkFragment<T>, "Given struct doesn't represent a valid chunk fragment type. Make sure to inherit from FMassChunkFragment or one of its child-types.");
 
 		const UScriptStruct* Type = T::StaticStruct();
 		const FChunkFragmentView* FoundChunkFragmentData = ChunkFragmentViews.FindByPredicate([Type](const FChunkFragmentView& Element) { return Element.Requirement.StructType == Type; } );
@@ -364,7 +364,7 @@ public:
 	template<typename T>
 	const T* GetConstSharedFragmentPtr() const
 	{
-		static_assert(TIsDerivedFrom<T, FMassConstSharedFragment>::IsDerived, "Given struct doesn't represent a valid const shared fragment type. Make sure to inherit from FMassConstSharedFragment or one of its child-types.");
+		static_assert(UE::Mass::CConstSharedFragment<T>, "Given struct doesn't represent a valid const shared fragment type. Make sure to inherit from FMassConstSharedFragment or one of its child-types.");
 
 		const FConstSharedFragmentView* FoundSharedFragmentData = ConstSharedFragmentViews.FindByPredicate([](const FConstSharedFragmentView& Element) { return Element.Requirement.StructType == T::StaticStruct(); });
 		return FoundSharedFragmentData ? FoundSharedFragmentData->FragmentView.GetPtr<const T>() : static_cast<const T*>(nullptr);
@@ -381,7 +381,7 @@ public:
 	template<typename T>
 	T* GetMutableSharedFragmentPtr()
 	{
-		static_assert(TIsDerivedFrom<T, FMassSharedFragment>::IsDerived, "Given struct doesn't represent a valid shared fragment type. Make sure to inherit from FMassSharedFragment or one of its child-types.");
+		static_assert(UE::Mass::CSharedFragment<T>, "Given struct doesn't represent a valid shared fragment type. Make sure to inherit from FMassSharedFragment or one of its child-types.");
 
 		FSharedFragmentView* FoundSharedFragmentData = SharedFragmentViews.FindByPredicate([](const FSharedFragmentView& Element) { return Element.Requirement.StructType == T::StaticStruct(); });
 		CHECK_IF_READWRITE(FoundSharedFragmentData);
@@ -391,7 +391,7 @@ public:
 	template<typename T>
 	const T* GetSharedFragmentPtr() const
 	{
-		static_assert(TIsDerivedFrom<T, FMassSharedFragment>::IsDerived, "Given struct doesn't represent a valid shared fragment type. Make sure to inherit from FMassSharedFragment or one of its child-types.");
+		static_assert(UE::Mass::CSharedFragment<T>, "Given struct doesn't represent a valid shared fragment type. Make sure to inherit from FMassSharedFragment or one of its child-types.");
 
 		const FSharedFragmentView* FoundSharedFragmentData = SharedFragmentViews.FindByPredicate([](const FSharedFragmentView& Element) { return Element.Requirement.StructType == T::StaticStruct(); });
 		return FoundSharedFragmentData ? FoundSharedFragmentData->FragmentView.GetPtr<T>() : static_cast<const T*>(nullptr);
diff --git a/Engine/Source/Runtime/MassEntity/Public/MassRequirements.h b/Engine/Source/Runtime/MassEntity/Public/MassRequirements.h
--- a/Engine/Source/Runtime/MassEntity/Public/MassRequirements.h
+++ b/Engine/Source/Runtime/MassEntity/Public/MassRequirements.h
@@ -199,7 +199,7 @@ public:
 		checkf(FragmentRequirements.FindByPredicate([](const FMassFragmentRequirementDescription& Item) { return Item.StructType == T::StaticStruct(); }) == nullptr
 			, TEXT("Duplicated requirements are not supported. %s already present"), *T::StaticStruct()->GetName());
 
-		static_assert(TIsDerivedFrom<T, FMassFragment>::IsDerived, "Given struct doesn't represent a valid fragment type. Make sure to inherit from FMassFragment or one of its child-types.");
+		static_assert(UE::Mass::CFragment<T>, "Given struct doesn't represent a valid fragment type. Make sure to inherit from FMassFragment or one of its child-types.");
 		
 		if (Presence != EMassFragmentPresence::None)
 		{
@@ -253,7 +253,7 @@ public:
 	{
 		checkf(bInitialized, TEXT("Modifying requirements before initialization is not supported."));
 		checkf(int(Presence) != int(EMassFragmentPresence::MAX), TEXT("MAX presence is not a valid value for AddTagRequirement"));
-		static_assert(TIsDerivedFrom<T, FMassTag>::IsDerived, "Given struct doesn't represent a valid tag type. Make sure to inherit from FMassFragment or one of its child-types.");
+		static_assert(UE::Mass::CTag<T>, "Given struct doesn't represent a valid tag type. Make sure to inherit from FMassFragment or one of its child-types.");
 		switch (Presence)
 		{
 			case EMassFragmentPresence::All:
@@ -288,7 +288,7 @@ public:
 	template<typename T>
 	FMassFragmentRequirements& AddChunkRequirement(const EMassFragmentAccess AccessMode, const EMassFragmentPresence Presence = EMassFragmentPresence::All)
 	{
-		static_assert(TIsDerivedFrom<T, FMassChunkFragment>::IsDerived, "Given struct doesn't represent a valid chunk fragment type. Make sure to inherit from FMassChunkFragment or one of its child-types.");
+		static_assert(UE::Mass::CChunkFragment<T>, "Given struct doesn't represent a valid chunk fragment type. Make sure to inherit from FMassChunkFragment or one of its child-types.");
 		checkf(bInitialized, TEXT("Modifying requirements before initialization is not supported."));
 		checkf(ChunkFragmentRequirements.FindByPredicate([](const FMassFragmentRequirementDescription& Item) { return Item.StructType == T::StaticStruct(); }) == nullptr
 			, TEXT("Duplicated requirements are not supported. %s already present"), *T::StaticStruct()->GetName());
@@ -340,7 +340,7 @@ public:
 	template<typename T>
 	FMassFragmentRequirements& AddConstSharedRequirement(const EMassFragmentPresence Presence = EMassFragmentPresence::All)
 	{
-		static_assert(TIsDerivedFrom<T, FMassConstSharedFragment>::IsDerived, "Given struct doesn't represent a valid const shared fragment type. Make sure to inherit from FMassConstSharedFragment or one of its child-types.");
+		static_assert(UE::Mass::CConstSharedFragment<T>, "Given struct doesn't represent a valid const shared fragment type. Make sure to inherit from FMassConstSharedFragment or one of its child-types.");
 		checkf(bInitialized, TEXT("Modifying requirements before initialization is not supported."));
 		checkf(ConstSharedFragmentRequirements.FindByPredicate([](const FMassFragmentRequirementDescription& Item) { return Item.StructType == T::StaticStruct(); }) == nullptr
 			, TEXT("Duplicated requirements are not supported. %s already present"), *T::StaticStruct()->GetName());
@@ -366,7 +366,7 @@ public:
 
 	FMassFragmentRequirements& AddConstSharedRequirement(const UScriptStruct* FragmentType, const EMassFragmentPresence Presence = EMassFragmentPresence::All)
 	{
-		if (!ensureMsgf(FragmentType->IsChildOf(FMassConstSharedFragment::StaticStruct())
+		if (!ensureMsgf(UE::Mass::IsA<FMassConstSharedFragment>(FragmentType)
 			, TEXT("Given struct doesn't represent a valid const shared fragment type. Make sure to inherit from FMassConstSharedFragment or one of its child-types.")))
 		{
 			return *this;
@@ -398,7 +398,7 @@ public:
 	template<typename T>
 	FMassFragmentRequirements& AddSharedRequirement(const EMassFragmentAccess AccessMode, const EMassFragmentPresence Presence = EMassFragmentPresence::All)
 	{
-		static_assert(TIsDerivedFrom<T, FMassSharedFragment>::IsDerived, "Given struct doesn't represent a valid shared fragment type. Make sure to inherit from FMassSharedFragment or one of its child-types.");
+		static_assert(UE::Mass::CSharedFragment<T>, "Given struct doesn't represent a valid shared fragment type. Make sure to inherit from FMassSharedFragment or one of its child-types.");
 		checkf(bInitialized, TEXT("Modifying requirements before initialization is not supported."));
 		checkf(SharedFragmentRequirements.FindByPredicate([](const FMassFragmentRequirementDescription& Item) { return Item.StructType == T::StaticStruct(); }) == nullptr
 			, TEXT("Duplicated requirements are not supported. %s already present"), *T::StaticStruct()->GetName());
