[Backout] - CL38817958

#rnx

Original CL Desc
-----------------------------------------------------------------
[Mass] Added a function to MassEntityManager for checking whether a given handle is `Reserved`
Modified how Mass entity storage understands "state" a bit.

diff --git a/Engine/Source/Runtime/MassEntity/Private/MassEntityManager.cpp b/Engine/Source/Runtime/MassEntity/Private/MassEntityManager.cpp
--- a/Engine/Source/Runtime/MassEntity/Private/MassEntityManager.cpp
+++ b/Engine/Source/Runtime/MassEntity/Private/MassEntityManager.cpp
@@ -790,6 +790,8 @@ TSharedRef<FMassEntityManager::FEntityCreationContext> FMassEntityManager::Batch
 	// to update the manager's information on these entities (stored in FMassEntityManager::Entities)
 	for (FMassEntityHandle& Entity : EntityHandles)
 	{
+		check(GetEntityStorageInterface().IsValidIndex(Entity.Index));
+
 		const UE::Mass::IEntityStorageInterface::EEntityState EntityState = GetEntityStorageInterface().GetEntityState(Entity.Index);
 		checkf(EntityState == UE::Mass::IEntityStorageInterface::EEntityState::Reserved, TEXT("Trying to build entities that are not reserved. Check all handles are reserved or consider using BatchCreateEntities"));
 
@@ -858,6 +860,7 @@ TSharedRef<FMassEntityManager::FEntityCreationContext> FMassEntityManager::Inter
 
 	for (FMassEntityHandle Entity : ReservedEntities)
 	{
+		check(IsEntityValid(Entity));
 		const UE::Mass::IEntityStorageInterface::EEntityState EntityState = GetEntityStorageInterface().GetEntityState(Entity.Index);
 		checkf(EntityState == UE::Mass::IEntityStorageInterface::EEntityState::Reserved, TEXT("Trying to build entities that are not reserved. Check all handles are reserved or consider using BatchCreateEntities"));
 		
@@ -1925,26 +1928,30 @@ const FSharedStruct* FMassEntityManager::InternalGetSharedFragmentPtr(FMassEntit
 bool FMassEntityManager::IsEntityValid(FMassEntityHandle Entity) const
 {
 	return (Entity.Index != UE::Mass::Private::InvalidEntityIndex) 
-		&& GetEntityStorageInterface().GetEntityState(Entity) != UE::Mass::IEntityStorageInterface::EEntityState::Invalid;
+		&& GetEntityStorageInterface().IsValidIndex(Entity.Index) 
+		&& (GetEntityStorageInterface().GetSerialNumber(Entity.Index) == Entity.SerialNumber);
 }
 
 bool FMassEntityManager::IsEntityBuilt(FMassEntityHandle Entity) const
 {
-	return GetEntityStorageInterface().GetEntityState(Entity.Index) == UE::Mass::IEntityStorageInterface::EEntityState::Created;
+	CheckIfEntityIsValid(Entity);
+	const UE::Mass::IEntityStorageInterface::EEntityState CurrentState = GetEntityStorageInterface().GetEntityState(Entity.Index);
+	return CurrentState == UE::Mass::IEntityStorageInterface::EEntityState::Created;
 }
 
 bool FMassEntityManager::IsEntityReserved(FMassEntityHandle EntityHandle) const
 {
-	return GetEntityStorageInterface().GetEntityState(EntityHandle) == UE::Mass::IEntityStorageInterface::EEntityState::Reserved;
+	CheckIfEntityIsValid(EntityHandle);
+	return GetEntityStorageInterface().GetEntityState(EntityHandle.Index) == UE::Mass::IEntityStorageInterface::EEntityState::Reserved;
 }
 
 void FMassEntityManager::CheckIfEntityIsValid(FMassEntityHandle Entity) const
 {
-	checkf(IsEntityValid(Entity), TEXT("Invalid entity (ID: %d, SN:%d, %s)"), Entity.Index, Entity.SerialNumber
-		, (Entity.Index == 0) ? TEXT("was never initialized") : TEXT("already destroyed"));
+	checkf(IsEntityValid(Entity), TEXT("Invalid entity (ID: %d, SN:%d, %s)"), Entity.Index, Entity.SerialNumber,
+		   (Entity.Index == 0) ? TEXT("was never initialized") : TEXT("already destroyed"));
 }
 
-FORCEINLINE void FMassEntityManager::CheckIfEntityIsActive(FMassEntityHandle Entity) const
+void FMassEntityManager::CheckIfEntityIsActive(FMassEntityHandle Entity) const
 {
 	checkf(IsEntityBuilt(Entity), TEXT("Entity not yet created(ID: %d, SN:%d)"), Entity.Index, Entity.SerialNumber);
 }
diff --git a/Engine/Source/Runtime/MassEntity/Private/MassEntityManagerStorage.cpp b/Engine/Source/Runtime/MassEntity/Private/MassEntityManagerStorage.cpp
--- a/Engine/Source/Runtime/MassEntity/Private/MassEntityManagerStorage.cpp
+++ b/Engine/Source/Runtime/MassEntity/Private/MassEntityManagerStorage.cpp
@@ -80,14 +80,7 @@ namespace UE::Mass
 				: EEntityState::Reserved;
 		}
 
-		return EEntityState::Invalid;
-	}
-
-	IEntityStorageInterface::EEntityState FSingleThreadedEntityStorage::GetEntityState(FMassEntityHandle EntityHandle) const
-	{
-		return IsValidHandle(EntityHandle)
-			? GetEntityState(EntityHandle.Index)
-			: EEntityState::Invalid;
+		return EEntityState::Free;	
 	}
 
 	int32 FSingleThreadedEntityStorage::GetSerialNumber(int32 Index) const
@@ -246,7 +239,6 @@ namespace UE::Mass
 		check(FMath::IsPowerOfTwo(InInitializationParams.MaxEntitiesPerPage));
 		check(FMath::IsPowerOfTwo(InInitializationParams.MaxEntityCount));
 		MaxEntitiesPerPage = InInitializationParams.MaxEntitiesPerPage;
-		MaxEntitiesPerPageLog2 = FMath::FloorLog2(MaxEntitiesPerPage);
 		MaximumEntityCountShift = FMath::FloorLog2(InInitializationParams.MaxEntityCount);
 		checkf(MaximumEntityCountShift < 32, TEXT("Invalid maximum entity count, cannot exceed 31 bits"));
 
@@ -301,13 +293,13 @@ namespace UE::Mass
 		LookupEntity(Index).CurrentArchetype = Archetype;
 	}
 
-	IEntityStorageInterface::EEntityState FConcurrentEntityStorage::GetEntityState(const int32 Index) const
+	IEntityStorageInterface::EEntityState FConcurrentEntityStorage::GetEntityState(int32 Index) const
 	{
 		//
 		// || Archetype || bIsAllocated || Result    |
-		//  |  nullptr   |      0       |  Invalid   |
-		//  |  nullptr   |      1       |  Reserved  |
-		//  | !nullptr   |      1       |  Created   |
+		//  |  nullptr   |      0       |  Free     |
+		//  |  nullptr   |      1       |  Reserved |
+		//  | !nullptr   |      1       |  Created  |
 		//
 	
 		const FEntityData& EntityData = LookupEntity(Index);
@@ -317,21 +309,7 @@ namespace UE::Mass
 		}
 		return EntityData.bIsAllocated
 			? EEntityState::Reserved
-			: EEntityState::Invalid;
-	}
-
-	IEntityStorageInterface::EEntityState FConcurrentEntityStorage::GetEntityState(const FMassEntityHandle EntityHandle) const
-	{
-		if (EntityHandle.Index >= 0)
-		{
-			const uint32 PageIndex = static_cast<uint32>(EntityHandle.Index) >> MaxEntitiesPerPageLog2;
-			if (PageIndex < PageCount)
-			{
-				return GetEntityState(EntityHandle.Index);
-			}
-		}
-
-		return EEntityState::Invalid;
+			: EEntityState::Free;
 	}
 
 	int32 FConcurrentEntityStorage::GetSerialNumber(int32 Index) const
@@ -344,7 +322,7 @@ namespace UE::Mass
 		// Page Index is which page in the array of pages we need to access
 		if (Index >= 0)
 		{
-			const uint32 PageIndex = static_cast<uint32>(Index) >> MaxEntitiesPerPageLog2;
+			const uint32 PageIndex = static_cast<uint32>(Index) >> FMath::FloorLog2(MaxEntitiesPerPage);
 			return PageIndex < PageCount;
 		}
 		return false;
@@ -618,7 +596,7 @@ namespace UE::Mass
 	{
 		check(Index >= 0);
 		// PageIndex is which Page in the array of pages we need to access
-		const uint32 PageIndex = static_cast<uint32>(Index) >> MaxEntitiesPerPageLog2;
+		const uint32 PageIndex = static_cast<uint32>(Index) >> FMath::FloorLog2(MaxEntitiesPerPage);
 
 		// Convert the entity index into the index with respect to the page
 		const uint32 EntityOffset = (PageIndex * MaxEntitiesPerPage);
diff --git a/Engine/Source/Runtime/MassEntity/Public/MassEntityManagerStorage.h b/Engine/Source/Runtime/MassEntity/Public/MassEntityManagerStorage.h
--- a/Engine/Source/Runtime/MassEntity/Public/MassEntityManagerStorage.h
+++ b/Engine/Source/Runtime/MassEntity/Public/MassEntityManagerStorage.h
@@ -54,12 +54,12 @@ namespace UE::Mass
 	public:
 		enum class EEntityState
 		{
-			/** Entity index is invalid or the serial number is out of date */
-			Invalid,
+			/** Entity index refers to an entity that is free to be reserved or created */
+			Free,
 			/** Entity index refers to a reserved entity */
 			Reserved,
 			/** Entity index refers to an entity assigned to an archetype */
-			Created,
+			Created
 		};
 		virtual ~IEntityStorageInterface() = default;
 
@@ -72,17 +72,11 @@ namespace UE::Mass
 		virtual void SetArchetypeFromShared(int32 Index, const TSharedPtr<FMassArchetypeData>& Archetype) = 0;
 
 		/**
-		 * @return EEntityState value indicating the current state of the entry indicated by the Index.
-		 *	Note that it's callers responsibility to ensure the Index is valid.
+		 * Returns true if the given entity at index is currently reserved
+		 * False if free or assigned an archetype
 		 */
 		virtual EEntityState GetEntityState(int32 Index) const = 0;
 
-		/**
-		 * @return EEntityState value indicating the validity of the handle, and if
-		 * 	the handle is valid, the state of the entry it indicates.
-		 */
-		virtual EEntityState GetEntityState(FMassEntityHandle EntityHandle) const = 0;
-
 		virtual int32 GetSerialNumber(int32 Index) const = 0;
 
 		/** Checks if index can be used to access entity data */
@@ -149,7 +143,6 @@ namespace UE::Mass
 		virtual void SetArchetypeFromShared(int32 Index, TSharedPtr<FMassArchetypeData>&) override;
 		virtual void SetArchetypeFromShared(int32 Index, const TSharedPtr<FMassArchetypeData>&) override;
 		virtual EEntityState GetEntityState(int32 Index) const override;
-		virtual EEntityState GetEntityState(FMassEntityHandle EntityHandle) const override;
 		virtual int32 GetSerialNumber(int32 Index) const override;
 		virtual bool IsValidIndex(int32 Index) const override;
 		virtual bool IsValidHandle(FMassEntityHandle EntityHandle) const override;
@@ -205,7 +198,6 @@ namespace UE::Mass
 		virtual void SetArchetypeFromShared(int32 Index, TSharedPtr<FMassArchetypeData>& Archetype) override;
 		virtual void SetArchetypeFromShared(int32 Index, const TSharedPtr<FMassArchetypeData>& Archetype) override;
 		virtual EEntityState GetEntityState(int32 Index) const override;
-		virtual EEntityState GetEntityState(FMassEntityHandle EntityHandle) const override;
 		virtual int32 GetSerialNumber(int32 Index) const override;
 		virtual bool IsValidIndex(int32 Index) const override;
 		virtual bool IsValidHandle(FMassEntityHandle EntityHandle) const override;
@@ -258,7 +250,6 @@ namespace UE::Mass
 		uint32 EntityCount = 0;
 		uint32 MaximumEntityCountShift = 0;
 		uint32 MaxEntitiesPerPage = 0;
-		uint32 MaxEntitiesPerPageLog2 = 0;
 		uint32 PageCount = 0;
 		/** ALWAYS acquire FreeListMutex before this one */
 		UE::FTransactionallySafeMutex PageAllocateMutex;
