Trivial format/style changes in AI-team owned code

#rnx

diff --git a/Engine/Source/Runtime/MassEntity/Private/MassDebugger.cpp b/Engine/Source/Runtime/MassEntity/Private/MassDebugger.cpp
--- a/Engine/Source/Runtime/MassEntity/Private/MassDebugger.cpp
+++ b/Engine/Source/Runtime/MassEntity/Private/MassDebugger.cpp
@@ -564,7 +564,7 @@ TArray<FMassEntityHandle> FMassDebugger::GetEntitiesOfArchetype(const FMassArche
 	return EntitiesOfArchetype;
 }
 
-TConstArrayView<struct UMassCompositeProcessor::FDependencyNode> FMassDebugger::GetProcessingGraph(const UMassCompositeProcessor& GraphOwner)
+TConstArrayView<UMassCompositeProcessor::FDependencyNode> FMassDebugger::GetProcessingGraph(const UMassCompositeProcessor& GraphOwner)
 {
 	return GraphOwner.FlatProcessingGraph;
 }
