[Mass] Made Mass subsystems ini-configurable via DefaultMass.ini

diff --git a/Engine/Source/Runtime/MassEntity/Public/MassSubsystemBase.h b/Engine/Source/Runtime/MassEntity/Public/MassSubsystemBase.h
--- a/Engine/Source/Runtime/MassEntity/Public/MassSubsystemBase.h
+++ b/Engine/Source/Runtime/MassEntity/Public/MassSubsystemBase.h
@@ -20,7 +20,7 @@ namespace UE::Mass::Subsystems
  * The sole responsibility of this world subsystem class is to serve functionality common to all 
  * Mass-related UWorldSubsystem-based subsystems, like whether the subsystems should get created at all. 
  */
-UCLASS(Abstract)
+UCLASS(Abstract, config = Mass, defaultconfig)
 class MASSENTITY_API UMassSubsystemBase : public UWorldSubsystem
 {
 	GENERATED_BODY()
@@ -55,7 +55,7 @@ protected:
  * The sole responsibility of this tickable world subsystem class is to serve functionality common to all
  * Mass-related UTickableWorldSubsystem-based subsystems, like whether the subsystems should get created at all.
  */
-UCLASS(Abstract)
+UCLASS(Abstract, config = Mass, defaultconfig)
 class MASSENTITY_API UMassTickableSubsystemBase : public UTickableWorldSubsystem
 {
 	GENERATED_BODY()
