[Mass] Added a UScriptStuct-base FMassRequirement.AddSharedRequirement, as required by dynamic StateTree processor work (in progress)

diff --git a/Engine/Source/Runtime/MassEntity/Private/MassRequirements.cpp b/Engine/Source/Runtime/MassEntity/Private/MassRequirements.cpp
--- a/Engine/Source/Runtime/MassEntity/Private/MassRequirements.cpp
+++ b/Engine/Source/Runtime/MassEntity/Private/MassRequirements.cpp
@@ -145,6 +145,13 @@ void FMassFragmentRequirements::SortRequirements()
 	SharedFragmentRequirements.Sort(FScriptStructSortOperator());
 }
 
+bool FMassFragmentRequirements::IsGameThreadOnlySharedFragment(TNotNull<const UScriptStruct*> SharedFragmentType) const
+{
+	return true;
+	// @todo always returning true at the moment, in near future it will be more like:
+	// return CachedEntityManager->GetTypeTraits(SharedFragmentType).IsGameThreadOnly();	
+}
+
 FORCEINLINE void FMassFragmentRequirements::CacheProperties() const
 {
 	if (bPropertiesCached == false)
diff --git a/Engine/Source/Runtime/MassEntity/Public/MassRequirements.h b/Engine/Source/Runtime/MassEntity/Public/MassRequirements.h
--- a/Engine/Source/Runtime/MassEntity/Public/MassRequirements.h
+++ b/Engine/Source/Runtime/MassEntity/Public/MassRequirements.h
@@ -430,6 +430,40 @@ public:
 		return *this;
 	}
 
+	FMassFragmentRequirements& AddSharedRequirement(TNotNull<const UScriptStruct*> SharedFragmentType, const EMassFragmentAccess AccessMode, const EMassFragmentPresence Presence = EMassFragmentPresence::All)
+	{
+		checkf(UE::Mass::IsA<FMassSharedFragment>(SharedFragmentType), TEXT("Given struct doesn't represent a valid shared fragment type. Make sure to inherit from FMassSharedFragment or one of its child-types.")); 
+		checkf(bInitialized, TEXT("Modifying requirements before initialization is not supported."));
+		checkf(SharedFragmentRequirements.FindByPredicate([&SharedFragmentType](const FMassFragmentRequirementDescription& Item) { return Item.StructType == SharedFragmentType; }) == nullptr
+			, TEXT("Duplicated requirements are not supported. %s already present"), *SharedFragmentType->GetName());
+		checkf(Presence != EMassFragmentPresence::Any, TEXT("\'Any\' is not a valid Presence value for AddSharedRequirement."));
+
+		switch (Presence)
+		{
+		case EMassFragmentPresence::All:
+			RequiredAllSharedFragments.Add(*SharedFragmentType);
+			SharedFragmentRequirements.Emplace(SharedFragmentType, AccessMode, Presence);
+			if (AccessMode == EMassFragmentAccess::ReadWrite)
+			{
+				bRequiresGameThreadExecution |= IsGameThreadOnlySharedFragment(SharedFragmentType);
+			}
+			break;
+		case EMassFragmentPresence::Optional:
+			RequiredOptionalSharedFragments.Add(*SharedFragmentType);
+			SharedFragmentRequirements.Emplace(SharedFragmentType, AccessMode, Presence);
+			if (AccessMode == EMassFragmentAccess::ReadWrite)
+			{
+				bRequiresGameThreadExecution |= IsGameThreadOnlySharedFragment(SharedFragmentType);
+			}
+			break;
+		case EMassFragmentPresence::None:
+			RequiredNoneSharedFragments.Add(*SharedFragmentType);
+			break;
+		}
+		IncrementChangeCounter();
+		return *this;
+	}
+
 	MASSENTITY_API void Reset();
 
 	/** 
@@ -483,6 +517,13 @@ protected:
 	void IncrementChangeCounter();
 	void ConsumeIncrementalChangesCount();
 	bool HasIncrementalChanges() const;
+	
+	/**
+	 * A helper function that passes the query over to CachedEntityManager.
+	 * Main purpose is to have the implementation in cpp and not include the EntityManager header here
+	 * @todo this function always returns True at the moment, proper implementation waiting for implementation of "type trait information" (WIP)
+	 */
+	MASSENTITY_API bool IsGameThreadOnlySharedFragment(TNotNull<const UScriptStruct*> SharedFragmentType) const;
 
 	friend FMassRequirementAccessDetector;
 
