[Mass] Fixed MassEntityView so that it can detect if the stored data is still correct (i.e. whether the underlaying archetype/chunk has changed)

#oforgithub public

diff --git a/Engine/Source/Developer/MassEntityTestSuite/Private/MassEntityViewTest.cpp b/Engine/Source/Developer/MassEntityTestSuite/Private/MassEntityViewTest.cpp
new file mode 100644
--- /dev/null
+++ b/Engine/Source/Developer/MassEntityTestSuite/Private/MassEntityViewTest.cpp
@@ -0,0 +1,37 @@
+// Copyright Epic Games, Inc. All Rights Reserved.
+
+#include "AITestsCommon.h"
+#include "MassEntityManager.h"
+#include "MassEntityTestTypes.h"
+#include "MassEntityView.h"
+
+
+#define LOCTEXT_NAMESPACE "MassTest"
+
+UE_DISABLE_OPTIMIZATION_SHIP
+
+namespace UE::Mass::Test
+{
+struct FEntityView_Invalidated : FEntityTestBase
+{
+	virtual bool InstantTest() override
+	{
+		const FMassEntityHandle EntityHandle = EntityManager->CreateEntity(IntsArchetype);
+		FMassEntityView EntityView(*EntityManager.Get(), EntityHandle);
+
+		AITEST_TRUE("The entity view is valid", EntityView.IsValid());
+
+		EntityManager->AddTagToEntity(EntityHandle, FTestTag_A::StaticStruct());
+
+		AITEST_FALSE("(NOT) The entity view is valid", EntityView.IsValid());
+
+		return true;
+	}
+};
+IMPLEMENT_AI_INSTANT_TEST(FEntityView_Invalidated, "System.Mass.EntityView.Invalidate");
+
+} // UE::Mass::Test
+
+UE_ENABLE_OPTIMIZATION_SHIP
+
+#undef LOCTEXT_NAMESPACE
diff --git a/Engine/Source/Runtime/MassEntity/Internal/MassArchetypeData.h b/Engine/Source/Runtime/MassEntity/Internal/MassArchetypeData.h
--- a/Engine/Source/Runtime/MassEntity/Internal/MassArchetypeData.h
+++ b/Engine/Source/Runtime/MassEntity/Internal/MassArchetypeData.h
@@ -351,20 +351,45 @@ public:
 	FORCEINLINE const int32* GetFragmentIndex(const UScriptStruct* FragmentType) const { return FragmentIndexMap.Find(FragmentType); }
 	FORCEINLINE int32 GetFragmentIndexChecked(const UScriptStruct* FragmentType) const { return FragmentIndexMap.FindChecked(FragmentType); }
 
-	FORCEINLINE void* GetFragmentData(const int32 FragmentIndex, const FMassRawEntityInChunkData EntityIndex) const
+	FORCEINLINE void* GetFragmentData(const int32 FragmentIndex, const FMassRawEntityInChunkData RawEntityInChunkHandle) const
 	{
-		return FragmentConfigs[FragmentIndex].GetFragmentData(EntityIndex.ChunkRawMemory, EntityIndex.IndexWithinChunk);
+		return FragmentConfigs[FragmentIndex].GetFragmentData(RawEntityInChunkHandle.ChunkRawMemory, RawEntityInChunkHandle.IndexWithinChunk);
 	}
 
-	FORCEINLINE FMassRawEntityInChunkData MakeEntityHandle(int32 EntityIndex) const
+	FORCEINLINE bool IsValidHandle(const FMassEntityInChunkDataHandle Handle) const
+	{
+		return Handle.IsSet() && Chunks.IsValidIndex(Handle.ChunkIndex) && Chunks[Handle.ChunkIndex].GetSerialModificationNumber() == Handle.ChunkSerialNumber;
+	}
+
+	FORCEINLINE void* GetFragmentData(const int32 FragmentIndex, const FMassEntityInChunkDataHandle EntityInChunkHandle) const
+	{
+		checkf(IsValidHandle(EntityInChunkHandle), TEXT("Input FMassRawEntityInChunkData is out of date."));
+		return FragmentConfigs[FragmentIndex].GetFragmentData(EntityInChunkHandle.ChunkRawMemory, EntityInChunkHandle.IndexWithinChunk);
+	}
+
+	FORCEINLINE FMassRawEntityInChunkData MakeRawEntityHandle(int32 EntityIndex) const
+	{
+		const int32 AbsoluteIndex = EntityMap.FindChecked(EntityIndex);
+		const int32 ChunkIndex = AbsoluteIndex / NumEntitiesPerChunk;
+	
+		return FMassRawEntityInChunkData(Chunks[ChunkIndex].GetRawMemory(), AbsoluteIndex % NumEntitiesPerChunk);
+	}
+
+	FORCEINLINE FMassRawEntityInChunkData MakeRawEntityHandle(const FMassEntityHandle Entity) const
+	{
+		return MakeRawEntityHandle(Entity.Index); 
+	}
+
+	FORCEINLINE FMassEntityInChunkDataHandle MakeEntityHandle(int32 EntityIndex) const
 	{
 		const int32 AbsoluteIndex = EntityMap.FindChecked(EntityIndex);
 		const int32 ChunkIndex = AbsoluteIndex / NumEntitiesPerChunk;
 	
-		return FMassRawEntityInChunkData(Chunks[ChunkIndex].GetRawMemory(), AbsoluteIndex % NumEntitiesPerChunk); 
+		return FMassEntityInChunkDataHandle(Chunks[ChunkIndex].GetRawMemory(), AbsoluteIndex % NumEntitiesPerChunk
+			, ChunkIndex, Chunks[ChunkIndex].GetSerialModificationNumber());
 	}
 
-	FORCEINLINE FMassRawEntityInChunkData MakeEntityHandle(FMassEntityHandle Entity) const
+	FORCEINLINE FMassEntityInChunkDataHandle MakeEntityHandle(const FMassEntityHandle Entity) const
 	{
 		return MakeEntityHandle(Entity.Index); 
 	}
diff --git a/Engine/Source/Runtime/MassEntity/Private/MassArchetypeTypes.cpp b/Engine/Source/Runtime/MassEntity/Private/MassArchetypeTypes.cpp
--- a/Engine/Source/Runtime/MassEntity/Private/MassArchetypeTypes.cpp
+++ b/Engine/Source/Runtime/MassEntity/Private/MassArchetypeTypes.cpp
@@ -362,3 +362,17 @@ FMassArchetypeEntityCollectionWithPayload::FMassArchetypeEntityCollectionWithPay
 {
 	Entities.BuildEntityRanges(TrueIndices);
 }
+
+//-----------------------------------------------------------------------------
+// FMassEntityInChunkDataHandle
+//-----------------------------------------------------------------------------
+bool FMassEntityInChunkDataHandle::IsValid(const FMassArchetypeData* ArchetypeData) const
+{
+	return ArchetypeData != nullptr
+		&& ArchetypeData->IsValidHandle(*this);
+}
+
+bool FMassEntityInChunkDataHandle::IsValid(const FMassArchetypeHandle& ArchetypeHandle) const
+{
+	return IsValid(FMassArchetypeHelper::ArchetypeDataFromHandle(ArchetypeHandle));
+}
diff --git a/Engine/Source/Runtime/MassEntity/Private/MassEntityView.cpp b/Engine/Source/Runtime/MassEntity/Private/MassEntityView.cpp
--- a/Engine/Source/Runtime/MassEntity/Private/MassEntityView.cpp
+++ b/Engine/Source/Runtime/MassEntity/Private/MassEntityView.cpp
@@ -3,6 +3,7 @@
 #include "MassEntityView.h"
 #include "MassEntityManager.h"
 #include "MassArchetypeData.h"
+#include "MassTestableEnsures.h"
 
 #include UE_INLINE_GENERATED_CPP_BY_NAME(MassEntityView)
 
@@ -33,20 +34,27 @@ FMassEntityView FMassEntityView::TryMakeView(const FMassEntityManager& EntityMan
 
 void* FMassEntityView::GetFragmentPtr(const UScriptStruct& FragmentType) const
 {
-	checkSlow(Archetype && EntityDataHandle.IsValid());
-	if (const int32* FragmentIndex = Archetype->GetFragmentIndex(&FragmentType))
+	if (testableEnsureMsgf(Archetype, TEXT("%hs: Trying to access fragments while no archetype set"), __FUNCTION__))
 	{
-		// failing the below Find means given entity's archetype is missing given FragmentType
-		return Archetype->GetFragmentData(*FragmentIndex, EntityDataHandle);
+		CA_ASSUME(Archetype);
+		if (const int32* FragmentIndex = Archetype->GetFragmentIndex(&FragmentType))
+		{
+			// failing the below Find means given entity's archetype is missing given FragmentType
+			return Archetype->GetFragmentData(*FragmentIndex, EntityDataHandle);
+		}
 	}
 	return nullptr;
 }
 
 void* FMassEntityView::GetFragmentPtrChecked(const UScriptStruct& FragmentType) const
 {
-	checkSlow(Archetype && EntityDataHandle.IsValid());
-	const int32 FragmentIndex = Archetype->GetFragmentIndexChecked(&FragmentType);
-	return Archetype->GetFragmentData(FragmentIndex, EntityDataHandle);
+	if (testableEnsureMsgf(Archetype, TEXT("%hs: Trying to access fragments while no archetype set"), __FUNCTION__))
+	{
+		CA_ASSUME(Archetype);
+		const int32 FragmentIndex = Archetype->GetFragmentIndexChecked(&FragmentType);
+		return Archetype->GetFragmentData(FragmentIndex, EntityDataHandle);
+	}
+	return nullptr;
 }
 
 const void* FMassEntityView::GetConstSharedFragmentPtr(const UScriptStruct& FragmentType) const
@@ -65,20 +73,18 @@ const void* FMassEntityView::GetConstSharedFragmentPtrChecked(const UScriptStruc
 void* FMassEntityView::GetSharedFragmentPtr(const UScriptStruct& FragmentType) const
 {
 	const FSharedStruct* SharedFragment = Archetype->GetSharedFragmentValues(Entity).GetSharedFragments().FindByPredicate(FStructTypeEqualOperator(&FragmentType));
-	// @todo: fix constness, should the getter be non-const too?
-	return (SharedFragment != nullptr) ? const_cast<uint8*>(SharedFragment->GetMemory()) : nullptr;
+	return (SharedFragment != nullptr) ? SharedFragment->GetMemory() : nullptr;
 }
 
 void* FMassEntityView::GetSharedFragmentPtrChecked(const UScriptStruct& FragmentType) const
 {
 	const FSharedStruct* SharedFragment = Archetype->GetSharedFragmentValues(Entity).GetSharedFragments().FindByPredicate(FStructTypeEqualOperator(&FragmentType));
 	check(SharedFragment != nullptr);
-	// @todo: fix constness, should the getter be non-const too?
-	return const_cast<uint8*>(SharedFragment->GetMemory());
+	return SharedFragment->GetMemory();
 }
 
 bool FMassEntityView::HasTag(const UScriptStruct& TagType) const
 {
-	checkSlow(Archetype && EntityDataHandle.IsValid());
+	check(EntityDataHandle.IsValid(Archetype));
 	return Archetype->HasTagType(&TagType);
 }
diff --git a/Engine/Source/Runtime/MassEntity/Public/MassArchetypeTypes.h b/Engine/Source/Runtime/MassEntity/Public/MassArchetypeTypes.h
--- a/Engine/Source/Runtime/MassEntity/Public/MassArchetypeTypes.h
+++ b/Engine/Source/Runtime/MassEntity/Public/MassArchetypeTypes.h
@@ -196,14 +196,37 @@ public:
 struct FMassRawEntityInChunkData 
 {
 	FMassRawEntityInChunkData() = default;
-	FMassRawEntityInChunkData(uint8* InChunkRawMemory, const int32 InIndexWithinChunk)
-        : ChunkRawMemory(InChunkRawMemory), IndexWithinChunk(InIndexWithinChunk)
-	{}
-	bool IsValid() const { return ChunkRawMemory != nullptr && IndexWithinChunk != INDEX_NONE; }
-	bool operator==(const FMassRawEntityInChunkData & Other) const { return ChunkRawMemory == Other.ChunkRawMemory && IndexWithinChunk == Other.IndexWithinChunk; }
+	FMassRawEntityInChunkData(uint8* InChunkRawMemory, const int32 InIndexWithinChunk);
+
+	bool IsSet() const;
+	bool operator==(const FMassRawEntityInChunkData & Other) const;
 
-	uint8* ChunkRawMemory = nullptr;
-	int32 IndexWithinChunk = INDEX_NONE;
+	uint8* const ChunkRawMemory = nullptr;
+	const int32 IndexWithinChunk = INDEX_NONE;;
+};
+
+//-----------------------------------------------------------------------------
+// FMassEntityInChunkDataHandle
+//-----------------------------------------------------------------------------
+/**
+ * This is an extension of FMassRawEntityInChunkData that provides additional safety features.
+ * It can be used to detect that the underlying data has changed. 
+ */
+struct FMassEntityInChunkDataHandle : FMassRawEntityInChunkData 
+{
+	FMassEntityInChunkDataHandle() = default;
+	FMassEntityInChunkDataHandle(FMassEntityInChunkDataHandle&&) = default;
+	FMassEntityInChunkDataHandle(const FMassEntityInChunkDataHandle&) = default;
+	FMassEntityInChunkDataHandle& operator=(const FMassEntityInChunkDataHandle&);
+	FMassEntityInChunkDataHandle& operator=(FMassEntityInChunkDataHandle&&);
+	FMassEntityInChunkDataHandle(uint8* InChunkRawMemory, const int32 InIndexWithinChunk, const int32 InChunkIndex, const int32 InChunkSerialNumber);
+
+	MASSENTITY_API bool IsValid(const FMassArchetypeData* ArchetypeData) const;
+	MASSENTITY_API bool IsValid(const FMassArchetypeHandle& ArchetypeHandle) const;
+	bool operator==(const FMassEntityInChunkDataHandle& Other) const;
+
+	const int32 ChunkIndex = INDEX_NONE;
+	const int32 ChunkSerialNumber = INDEX_NONE;
 };
 
 //-----------------------------------------------------------------------------
@@ -224,3 +247,46 @@ struct FMassQueryRequirementIndicesMapping
 		return EntityFragments.Num() == 0 || ChunkFragments.Num() == 0;
 	}
 };
+
+
+//-----------------------------------------------------------------------------
+// INLINES
+//-----------------------------------------------------------------------------
+inline FMassRawEntityInChunkData::FMassRawEntityInChunkData(uint8* InChunkRawMemory, const int32 InIndexWithinChunk)
+	: ChunkRawMemory(InChunkRawMemory), IndexWithinChunk(InIndexWithinChunk)
+{}
+
+inline bool FMassRawEntityInChunkData::IsSet() const
+{
+	return ChunkRawMemory != nullptr && IndexWithinChunk != INDEX_NONE;
+}
+
+inline bool FMassRawEntityInChunkData::operator==(const FMassRawEntityInChunkData & Other) const
+{
+	return ChunkRawMemory == Other.ChunkRawMemory && IndexWithinChunk == Other.IndexWithinChunk;
+}
+
+inline FMassEntityInChunkDataHandle& FMassEntityInChunkDataHandle::operator=(const FMassEntityInChunkDataHandle& Other)
+{
+	new (this) FMassEntityInChunkDataHandle(Other);
+	return *this;
+}
+
+inline FMassEntityInChunkDataHandle& FMassEntityInChunkDataHandle::operator=(FMassEntityInChunkDataHandle&& Other)
+{
+	new (this) FMassEntityInChunkDataHandle(Other);
+	return *this;
+}
+
+inline FMassEntityInChunkDataHandle::FMassEntityInChunkDataHandle(uint8* InChunkRawMemory, const int32 InIndexWithinChunk, const int32 InChunkIndex, const int32 InChunkSerialNumber)
+		: FMassRawEntityInChunkData(InChunkRawMemory, InIndexWithinChunk)
+		, ChunkIndex(InChunkIndex), ChunkSerialNumber(InChunkSerialNumber)
+{
+}
+
+inline bool FMassEntityInChunkDataHandle::operator==(const FMassEntityInChunkDataHandle& Other) const
+{
+	return FMassRawEntityInChunkData::operator==(Other)
+		&& ChunkIndex == Other.ChunkIndex
+		&& ChunkSerialNumber == Other.ChunkSerialNumber;
+}
diff --git a/Engine/Source/Runtime/MassEntity/Public/MassEntityView.h b/Engine/Source/Runtime/MassEntity/Public/MassEntityView.h
--- a/Engine/Source/Runtime/MassEntity/Public/MassEntityView.h
+++ b/Engine/Source/Runtime/MassEntity/Public/MassEntityView.h
@@ -128,8 +128,6 @@ struct MASSENTITY_API FMassEntityView
 	UE_DEPRECATED(5.5, "Using GetSharedFragmentDataPtr with const shared fragments is deprecated. Use GetConstSharedFragmentData instead")
 	T& GetSharedFragmentData() const
 	{
-		/*static T DummyInstance;
-		return DummyInstance;*/
 		return const_cast<T&>(GetConstSharedFragmentData<T>());
 	}
 
@@ -148,9 +146,9 @@ struct MASSENTITY_API FMassEntityView
 
 	bool HasTag(const UScriptStruct& TagType) const;
 
-	bool IsSet() const { return Archetype != nullptr && EntityDataHandle.IsValid(); }
-	bool IsValid() const { return IsSet(); }
-	bool operator==(const FMassEntityView& Other) const { return Archetype == Other.Archetype && EntityDataHandle == Other.EntityDataHandle; }
+	bool IsSet() const;
+	bool IsValid() const;
+	bool operator==(const FMassEntityView& Other) const;
 
 protected:
 	void* GetFragmentPtr(const UScriptStruct& FragmentType) const;
@@ -162,6 +160,24 @@ protected:
 
 private:
 	FMassEntityHandle Entity;
-	FMassRawEntityInChunkData EntityDataHandle;
+	FMassEntityInChunkDataHandle EntityDataHandle;
 	FMassArchetypeData* Archetype = nullptr;
 };
+
+//-----------------------------------------------------------------------------
+// INLINES
+//-----------------------------------------------------------------------------
+inline bool FMassEntityView::IsSet() const
+{
+	return EntityDataHandle.IsValid(Archetype);
+}
+
+inline bool FMassEntityView::IsValid() const
+{
+	return IsSet();
+}
+
+inline bool FMassEntityView::operator==(const FMassEntityView& Other) const
+{
+	return Archetype == Other.Archetype && EntityDataHandle == Other.EntityDataHandle;
+}
