[Mass] Made it thread-safe to register and unregister dynamic processors. It's also allowed while the Mass processing is underway now.

diff --git a/Engine/Plugins/Runtime/MassGameplay/Source/MassSimulation/Private/MassSimulationSubsystem.cpp b/Engine/Plugins/Runtime/MassGameplay/Source/MassSimulation/Private/MassSimulationSubsystem.cpp
--- a/Engine/Plugins/Runtime/MassGameplay/Source/MassSimulation/Private/MassSimulationSubsystem.cpp
+++ b/Engine/Plugins/Runtime/MassGameplay/Source/MassSimulation/Private/MassSimulationSubsystem.cpp
@@ -150,13 +150,11 @@ void UMassSimulationSubsystem::OnWorldBeginPlay(UWorld& InWorld)
 
 void UMassSimulationSubsystem::RegisterDynamicProcessor(UMassProcessor& Processor)
 {
-	checkf(!IsDuringMassProcessing(), TEXT("Unable to add dynamic processors to Mass during processing."));
 	PhaseManager->RegisterDynamicProcessor(Processor);
 }
 
 void UMassSimulationSubsystem::UnregisterDynamicProcessor(UMassProcessor& Processor)
 {
-	checkf(!IsDuringMassProcessing(), TEXT("Unable to remove dynamic processors to Mass during processing."));
 	PhaseManager->UnregisterDynamicProcessor(Processor);
 }
 
diff --git a/Engine/Source/Runtime/MassEntity/Private/MassProcessingPhaseManager.cpp b/Engine/Source/Runtime/MassEntity/Private/MassProcessingPhaseManager.cpp
--- a/Engine/Source/Runtime/MassEntity/Private/MassProcessingPhaseManager.cpp
+++ b/Engine/Source/Runtime/MassEntity/Private/MassProcessingPhaseManager.cpp
@@ -294,6 +294,17 @@ void FMassProcessingPhaseManager::Deinitialize()
 	{
 		GraphBuildState.Reset();
 	}
+
+	// manually deque all the queues, since there's no guarantee that this
+	// FMassProcessingPhaseManager instance is getting destroyed right after this call
+	FDynamicProcessorOperation DummyElement;
+	for (TMpscQueue<FDynamicProcessorOperation>& Queue : PendingDynamicProcessors)
+	{	
+		while (Queue.Dequeue(DummyElement))
+		{
+			// empty on purpose
+		}
+	}
 }
 
 const FGraphEventRef& FMassProcessingPhaseManager::TriggerPhase(const EMassProcessingPhase Phase, const float DeltaTime
@@ -471,6 +482,12 @@ void FMassProcessingPhaseManager::OnPhaseStart(FMassProcessingPhase& Phase)
 	}
 
 	const int32 PhaseAsInt = int32(Phase.Phase);
+
+	if (PendingDynamicProcessors[PhaseAsInt].IsEmpty() == false)
+	{
+		HandlePendingDynamicProcessorOperations(PhaseAsInt);
+	}
+
 	if (Owner.IsValid()
 		&& ensure(Phase.Phase != EMassProcessingPhase::MAX)
 		&& (ProcessingGraphBuildStates[PhaseAsInt].bNewArchetypes || ProcessingGraphBuildStates[PhaseAsInt].bProcessorsNeedRebuild)
@@ -524,31 +541,71 @@ FString FMassProcessingPhaseManager::GetName() const
 
 void FMassProcessingPhaseManager::RegisterDynamicProcessor(UMassProcessor& Processor)
 {
-	if (Processor.IsInitialized() == false)
+	if (ensureMsgf(Processor.GetProcessingPhase() != EMassProcessingPhase::MAX
+		, TEXT("%hs, Misconfigured processor %s, marked as ProcessingPhase == MAX"), __FUNCTION__, *Processor.GetName()))
+	{
+		PendingDynamicProcessors[int32(Processor.GetProcessingPhase())].Enqueue(&Processor, EDynamicProcessorOperationType::Add);
+	}
+}
+
+void FMassProcessingPhaseManager::RegisterDynamicProcessorInternal(TNotNull<UMassProcessor*> Processor)
+{
+	if (Processor->IsInitialized() == false)
 	{
 		check(EntityManager->GetOwner());
-		Processor.CallInitialize(EntityManager->GetOwner(), EntityManager.ToSharedRef());
+		Processor->CallInitialize(EntityManager->GetOwner(), EntityManager.ToSharedRef());
 	}
-	DynamicProcessors.Add(&Processor);
-	Processor.MarkAsDynamic();
-	ProcessingGraphBuildStates[uint32(Processor.GetProcessingPhase())].bProcessorsNeedRebuild = true;
+	DynamicProcessors.Add(Processor);
+	Processor->MarkAsDynamic();
 }
 
 void FMassProcessingPhaseManager::UnregisterDynamicProcessor(UMassProcessor& Processor)
 {
-	int32 Index;
-	if (DynamicProcessors.Find(&Processor, Index))
+	if (ensureMsgf(Processor.GetProcessingPhase() != EMassProcessingPhase::MAX
+		, TEXT("%hs, Misconfigured processor %s, marked as ProcessingPhase == MAX"), __FUNCTION__, *Processor.GetName()))
 	{
-		DynamicProcessors.RemoveAt(Index);
-		ProcessingGraphBuildStates[uint32(Processor.GetProcessingPhase())].bProcessorsNeedRebuild = true;
+		PendingDynamicProcessors[int32(Processor.GetProcessingPhase())].Enqueue(&Processor, EDynamicProcessorOperationType::Remove);
+	}
+}
+
+void FMassProcessingPhaseManager::UnregisterDynamicProcessorInternal(TNotNull<UMassProcessor*> Processor)
+{
+	int32 Index = INDEX_NONE;
+	if (DynamicProcessors.Find(Processor, Index))
+	{
+		DynamicProcessors.RemoveAtSwap(Index, 1, EAllowShrinking::No);
+		ProcessingGraphBuildStates[int32(Processor->GetProcessingPhase())].bProcessorsNeedRebuild = true;
 
 		// it's possible that the given dynamic processor is a part of processing graph at the moment
 		// we need to store the information about its removal and use it when rebuilding the graph next time around.
-		RemovedDynamicProcessors.Add(&Processor);
+		RemovedDynamicProcessors.Add(Processor);
 	}
 	else
 	{
-		checkf(false, TEXT("Unable to remove Processor '%s', as it was never added or already removed."), *Processor.GetName());
+		checkf(false, TEXT("Unable to remove Processor '%s', as it was never added or already removed."), *Processor->GetName());
+	}
+}
+
+void FMassProcessingPhaseManager::HandlePendingDynamicProcessorOperations(const int32 PhaseIndex)
+{
+	bool bWorkDone = false;
+	FDynamicProcessorOperation Operation;
+	while (PendingDynamicProcessors[PhaseIndex].Dequeue(Operation))
+	{
+		if (Operation.Get<1>() == EDynamicProcessorOperationType::Add)
+		{
+			RegisterDynamicProcessorInternal(Operation.Get<0>().Get());
+		}
+		else
+		{
+			UnregisterDynamicProcessorInternal(Operation.Get<0>().Get());
+		}
+		bWorkDone = true;
+	}
+
+	if (bWorkDone)
+	{
+		ProcessingGraphBuildStates[PhaseIndex].bProcessorsNeedRebuild = bWorkDone;
 	}
 }
 
diff --git a/Engine/Source/Runtime/MassEntity/Public/MassProcessingPhaseManager.h b/Engine/Source/Runtime/MassEntity/Public/MassProcessingPhaseManager.h
--- a/Engine/Source/Runtime/MassEntity/Public/MassProcessingPhaseManager.h
+++ b/Engine/Source/Runtime/MassEntity/Public/MassProcessingPhaseManager.h
@@ -8,6 +8,7 @@
 #endif // UE_ENABLE_INCLUDE_ORDER_DEPRECATED_IN_5_6
 #include "UObject/Object.h"
 #include "UObject/GCObject.h"
+#include "Containers/MpscQueue.h"
 #include "Engine/EngineBaseTypes.h"
 #include "MassProcessingTypes.h"
 #include "MassProcessor.h"
@@ -207,6 +208,10 @@ protected:
 	}
 	// End of FGCObject interface
 
+	void RegisterDynamicProcessorInternal(TNotNull<UMassProcessor*> Processor);
+	void UnregisterDynamicProcessorInternal(TNotNull<UMassProcessor*> Processor);
+	void HandlePendingDynamicProcessorOperations(const int32 PhaseIndex);
+
 	/** Override this function if you want to modify how the phase tick functions get executed. */
 	virtual void EnableTickFunctions(const UWorld& World);
 
@@ -235,6 +240,14 @@ protected:
 	TArray<FMassProcessingPhaseConfig> ProcessingPhasesConfig;
 	TArray<TObjectPtr<UMassProcessor>> DynamicProcessors;
 	TArray<TWeakObjectPtr<UMassProcessor>> RemovedDynamicProcessors;
+	enum class EDynamicProcessorOperationType : uint8
+	{
+		Add,
+		Remove
+	};
+	/** using TStrongObjectPtr to not worry about GC while the processor instances are waiting in PendingDynamicProcessors */
+	using FDynamicProcessorOperation = TPair<TStrongObjectPtr<UMassProcessor>, EDynamicProcessorOperationType>;
+	TMpscQueue<FDynamicProcessorOperation> PendingDynamicProcessors[(uint8)EMassProcessingPhase::MAX];
 
 	TSharedPtr<FMassEntityManager> EntityManager;
 
