[Mass] added optional processing graph logging, by default logging every "new" graph (in "Log" category, so not showing by default anyway).

#rnx

diff --git a/Engine/Source/Runtime/MassEntity/Private/MassProcessor.cpp b/Engine/Source/Runtime/MassEntity/Private/MassProcessor.cpp
--- a/Engine/Source/Runtime/MassEntity/Private/MassProcessor.cpp
+++ b/Engine/Source/Runtime/MassEntity/Private/MassProcessor.cpp
@@ -17,9 +17,21 @@ DECLARE_CYCLE_STAT(TEXT("Mass Processor Task"), STAT_Mass_DoTask, STATGROUP_Mass
 #if WITH_MASSENTITY_DEBUG
 namespace UE::Mass::Debug
 {
-	bool bLogProcessingGraph = false;
-	FAutoConsoleVariableRef CVarLogProcessingGraph(TEXT("mass.LogProcessingGraph"), bLogProcessingGraph
-		, TEXT("When enabled will log task graph tasks created while dispatching processors to other threads, along with their dependencies"), ECVF_Cheat);
+	bool bLogProcessingGraphEveryFrame = false;
+	bool bLogNewProcessingGraph = true;
+
+	namespace
+	{
+		FAutoConsoleVariableRef CVars[] = {
+			{ TEXT("mass.LogProcessingGraph"), bLogProcessingGraphEveryFrame
+				, TEXT("When enabled every composite processor, every frame, will log task graph tasks created while dispatching processors to other threads, along with their dependencies.")
+				, ECVF_Cheat }
+			, { TEXT("mass.LogNewProcessingGraph"), bLogNewProcessingGraph
+				, TEXT("When enabled every time a new processing graph is created the composite processor hosting it will log it during first execution.")
+				, ECVF_Cheat }
+		};
+	}
+
 }
 
 // change to 1 to enable more detailed processing tasks logging
@@ -427,10 +439,11 @@ FGraphEventRef UMassCompositeProcessor::DispatchProcessorTasks(const TSharedPtr<
 		}
 	}
 
-
 #if WITH_MASSENTITY_DEBUG
-	if (UE::Mass::Debug::bLogProcessingGraph)
+	if (UE::Mass::Debug::bLogProcessingGraphEveryFrame || bDebugLogNewProcessingGraph)
 	{
+		FScopedCategoryAndVerbosityOverride LogOverride(TEXT("LogMass"), ELogVerbosity::Log);
+
 		for (int i = 0; i < FlatProcessingGraph.Num(); ++i)
 		{
 			FDependencyNode& ProcessingNode = FlatProcessingGraph[i];
@@ -444,15 +457,17 @@ FGraphEventRef UMassCompositeProcessor::DispatchProcessorTasks(const TSharedPtr<
 			if (Events[i].IsValid())
 			{
 				PROCESSOR_TASK_LOG(TEXT("Task %u %s%s%s"), Events[i]->GetTraceId(), *ProcessingNode.Processor->GetProcessorName()
-				   , DependenciesDesc.Len() > 0 ? TEXT(" depends on ") : TEXT(""), *DependenciesDesc);
+					, DependenciesDesc.Len() > 0 ? TEXT(" depends on ") : TEXT(""), *DependenciesDesc);
 			}
 			else
 			{
 				ensureMsgf(ProcessingNode.Processor->IsActive() == false, TEXT("This path is expected to trigger only for inactive processors"))
 				PROCESSOR_TASK_LOG(TEXT("Task [INACTIVE] %s%s%s"), *ProcessingNode.Processor->GetProcessorName()
-				   , DependenciesDesc.Len() > 0 ? TEXT(" depends on ") : TEXT(""), *DependenciesDesc);
+					, DependenciesDesc.Len() > 0 ? TEXT(" depends on ") : TEXT(""), *DependenciesDesc);
 			}
 		}
+
+		bDebugLogNewProcessingGraph = false;
 	}
 #endif // WITH_MASSENTITY_DEBUG
 
@@ -546,6 +561,8 @@ void UMassCompositeProcessor::BuildFlatProcessingGraph(TConstArrayView<FMassProc
 		}
 		++Index;
 	}
+
+	bDebugLogNewProcessingGraph = UE::Mass::Debug::bLogNewProcessingGraph;
 #endif // WITH_MASSENTITY_DEBUG
 }
 
diff --git a/Engine/Source/Runtime/MassEntity/Public/MassProcessor.h b/Engine/Source/Runtime/MassEntity/Public/MassProcessor.h
--- a/Engine/Source/Runtime/MassEntity/Public/MassProcessor.h
+++ b/Engine/Source/Runtime/MassEntity/Public/MassProcessor.h
@@ -353,6 +353,10 @@ protected:
 	UPROPERTY()
 	FName GroupName;
 
+#if WITH_MASSENTITY_DEBUG
+	bool bDebugLogNewProcessingGraph = false;
+#endif // WITH_MASSENTITY_DEBUG
+
 	TArray<FDependencyNode> FlatProcessingGraph;
 
 	struct FProcessorCompletion
