[Mass] Added FMassArchetypeCompositionDescriptor::Remove to complement FMassArchetypeCompositionDescriptor::Append

#rnx

diff --git a/Engine/Source/Runtime/MassEntity/Public/MassEntityTypes.h b/Engine/Source/Runtime/MassEntity/Public/MassEntityTypes.h
--- a/Engine/Source/Runtime/MassEntity/Public/MassEntityTypes.h
+++ b/Engine/Source/Runtime/MassEntity/Public/MassEntityTypes.h
@@ -136,6 +136,15 @@ struct FMassArchetypeCompositionDescriptor
 		ConstSharedFragments += OtherDescriptor.ConstSharedFragments;
 	}
 
+	void Remove(const FMassArchetypeCompositionDescriptor& OtherDescriptor)
+	{
+		Fragments -= OtherDescriptor.Fragments;
+		Tags -= OtherDescriptor.Tags;
+		ChunkFragments -= OtherDescriptor.ChunkFragments;
+		SharedFragments -= OtherDescriptor.SharedFragments;
+		ConstSharedFragments -= OtherDescriptor.ConstSharedFragments;
+	}
+
 	MASSENTITY_API static uint32 CalculateHash(const FMassFragmentBitSet& InFragments, const FMassTagBitSet& InTags
 		, const FMassChunkFragmentBitSet& InChunkFragments, const FMassSharedFragmentBitSet& InSharedFragmentBitSet
 		, const FMassConstSharedFragmentBitSet& InConstSharedFragmentBitSet);
