[Mass] Fixed FMassEntityManager::AddCompositionToEntity_GetDelta and FMassEntityManager::RemoveCompositionFromEntity not handling all the element types properly (just tags and fragments).

#rnx

diff --git a/Engine/Source/Runtime/MassEntity/Private/MassEntityManager.cpp b/Engine/Source/Runtime/MassEntity/Private/MassEntityManager.cpp
--- a/Engine/Source/Runtime/MassEntity/Private/MassEntityManager.cpp
+++ b/Engine/Source/Runtime/MassEntity/Private/MassEntityManager.cpp
@@ -1110,7 +1110,7 @@ void FMassEntityManager::AddFragmentListToEntity(FMassEntityHandle Entity, TCons
 	}
 }
 
-void FMassEntityManager::AddCompositionToEntity_GetDelta(FMassEntityHandle Entity, FMassArchetypeCompositionDescriptor& InDescriptor)
+void FMassEntityManager::AddCompositionToEntity_GetDelta(FMassEntityHandle Entity, FMassArchetypeCompositionDescriptor& InOutDescriptor)
 {
 	CHECK_SYNC_API();
 
@@ -1119,16 +1119,15 @@ void FMassEntityManager::AddCompositionToEntity_GetDelta(FMassEntityHandle Entit
 	FMassArchetypeData* OldArchetype = GetEntityStorageInterface().GetArchetype(Entity.Index);
 	check(OldArchetype);
 
-	InDescriptor.Fragments -= OldArchetype->GetCompositionDescriptor().Fragments;
-	InDescriptor.Tags -= OldArchetype->GetCompositionDescriptor().Tags;
+	InOutDescriptor.Remove(OldArchetype->GetCompositionDescriptor());
 
-	ensureMsgf(InDescriptor.ChunkFragments.IsEmpty(), TEXT("Adding new chunk fragments is not supported"));
+	ensureMsgf(InOutDescriptor.ChunkFragments.IsEmpty(), TEXT("Adding new chunk fragments is not supported"));
+	ensureMsgf(InOutDescriptor.SharedFragments.IsEmpty(), TEXT("Adding new shared fragments is not supported"));
 
-	if (InDescriptor.IsEmpty() == false)
+	if (InOutDescriptor.IsEmpty() == false)
 	{
 		FMassArchetypeCompositionDescriptor NewDescriptor = OldArchetype->GetCompositionDescriptor();
-		NewDescriptor.Fragments += InDescriptor.Fragments;
-		NewDescriptor.Tags += InDescriptor.Tags;
+		NewDescriptor.Append(InOutDescriptor);
 
 		const FMassArchetypeHandle NewArchetypeHandle = CreateArchetype(NewDescriptor, FMassArchetypeCreationParams(*OldArchetype));
 
@@ -1143,7 +1142,7 @@ void FMassEntityManager::AddCompositionToEntity_GetDelta(FMassEntityHandle Entit
 
 			if (IsAllowedToTriggerObservers())
 			{
-				ObserverManager.OnPostCompositionAdded(Entity, InDescriptor);
+				ObserverManager.OnPostCompositionAdded(Entity, InOutDescriptor);
 			}
 		}
 	}
@@ -1161,16 +1160,13 @@ void FMassEntityManager::RemoveCompositionFromEntity(FMassEntityHandle Entity, c
 		check(OldArchetype);
 
 		FMassArchetypeCompositionDescriptor NewDescriptor = OldArchetype->GetCompositionDescriptor();
-		NewDescriptor.Fragments -= InDescriptor.Fragments;
-		NewDescriptor.Tags -= InDescriptor.Tags;
+		NewDescriptor.Remove(InDescriptor);
 
 		ensureMsgf(InDescriptor.ChunkFragments.IsEmpty(), TEXT("Removing chunk fragments is not supported"));
 		ensureMsgf(InDescriptor.SharedFragments.IsEmpty(), TEXT("Removing shared fragments is not supported"));
 
 		if (NewDescriptor.IsEquivalent(OldArchetype->GetCompositionDescriptor()) == false)
 		{
-			ensureMsgf(OldArchetype->GetCompositionDescriptor().HasAll(InDescriptor), TEXT("Some of the elements being removed are already missing from entity\'s composition."));
-			
 			if (IsAllowedToTriggerObservers())
 			{
 				ObserverManager.OnPreCompositionRemoved(Entity, InDescriptor);
