[Mass] Minor dependency solver code clean up, removed commented out code, removed unused header inclusion

#rnx

diff --git a/Engine/Source/Runtime/MassEntity/Private/MassProcessorDependencySolver.cpp b/Engine/Source/Runtime/MassEntity/Private/MassProcessorDependencySolver.cpp
--- a/Engine/Source/Runtime/MassEntity/Private/MassProcessorDependencySolver.cpp
+++ b/Engine/Source/Runtime/MassEntity/Private/MassProcessorDependencySolver.cpp
@@ -5,7 +5,6 @@
 #include "MassProcessor.h"
 #include "Logging/MessageLog.h"
 #include "HAL/FileManager.h"
-#include "Misc/Paths.h"
 #include "Engine/World.h"
 
 #define LOCTEXT_NAMESPACE "Mass"
@@ -579,396 +578,6 @@ void FMassProcessorDependencySolver::LogNode(const FNode& Node, int Indent)
 	}
 }
 
-#if 0 // disabled the graph building for now, leaving the old code here for reference
-struct FDumpGraphDependencyUtils
-{
-	static void DumpGraphNode(FArchive& LogFile, const FProcessorDependencySolver::FNode& Node, int Indent, TSet<const FProcessorDependencySolver::FNode*>& AllNodes, const bool bRoot)
-	{
-		const FString NodeName = Node.Name.ToString();
-		if (!Node.Processor)
-		{
-			const FString ClusterNodeName = NodeName.Replace(TEXT("."), TEXT("_"));
-			const FString GraphNodeName = NodeName.Replace(TEXT("."), TEXT(" "));
-
-			int32 Index = -1;
-			Node.Name.ToString().FindLastChar(TEXT('.'), Index);
-			const FString GroupName = NodeName.Mid(Index+1);
-
-			LogFile.Logf(TEXT("%*ssubgraph cluster_%s"), Indent, TEXT(""), *ClusterNodeName);
-			LogFile.Logf(TEXT("%*s{"), Indent, TEXT(""));
-			LogFile.Logf(TEXT("%*slabel =\"%s\";"), Indent + 4, TEXT(""), *GroupName);
-			LogFile.Logf(TEXT("%*s\"%s Start\"%s;"), Indent + 4, TEXT(""), *GraphNodeName, bRoot ? TEXT("") : TEXT("[shape=point style=invis]"));
-			LogFile.Logf(TEXT("%*s\"%s End\"%s;"), Indent + 4, TEXT(""), *GraphNodeName, bRoot ? TEXT("") : TEXT("[shape=point style=invis]"));
-			for (const FProcessorDependencySolver::FNode& SubNode : Node.SubNodes)
-			{
-				DumpGraphNode(LogFile, SubNode, Indent + 4, AllNodes, false);
-			}
-			LogFile.Logf(TEXT("%*s}"), Indent, TEXT(""));
-		}
-		else
-		{
-			LogFile.Logf(TEXT("%*s\"%s\""), Indent, TEXT(""), *NodeName);
-			AllNodes.Add(&Node);
-		}
-	}
-
-	static const FProcessorDependencySolver::FNode* FindDepNodeInParents(const TArray<const FProcessorDependencySolver::FNode*>& Parents, FName DependencyName)
-	{
-		for (int32 i = Parents.Num() - 1; i >= 0; --i)
-		{
-			const FProcessorDependencySolver::FNode* CurNode = Parents[i];
-			int32 DependencyIndex = CurNode->FindNodeIndex(DependencyName);
-			while(DependencyIndex != INDEX_NONE)
-			{
-				if(CurNode->SubNodes[DependencyIndex].Name == DependencyName)
-				{
-					return &CurNode->SubNodes[DependencyIndex];
-				}
-				else
-				{
-					// Dig down the chain
-					CurNode = &CurNode->SubNodes[DependencyIndex];
-					DependencyIndex = CurNode->FindNodeIndex(DependencyName);
-				}
-			}
-		}
-		return nullptr;
-	};
-
-	static bool DoAllSubNodeHasDependency(const FProcessorDependencySolver::FNode& Node, const FName DependencyName, bool bBeforeDep)
-	{
-		bool bDepExistInAllSibling = true;
-		for (const FProcessorDependencySolver::FNode& SiblingSubNode : Node.SubNodes)
-		{
-			if (SiblingSubNode.Processor)
-			{
-				const TArray<FName>& ExecuteDep = bBeforeDep ? SiblingSubNode.Processor->GetExecutionOrder().ExecuteBefore : SiblingSubNode.Processor->GetExecutionOrder().ExecuteAfter;
-				if (ExecuteDep.Find(DependencyName) == INDEX_NONE)
-				{
-					bDepExistInAllSibling = false;
-					break;
-				}
-			}
-			else if (!DoAllSubNodeHasDependency(SiblingSubNode, DependencyName, bBeforeDep))
-			{
-				bDepExistInAllSibling = false;
-				break;
-			}
-		}
-		return bDepExistInAllSibling;
-	}
-
-	static bool DoesDependencyExistIndirectly(const TArray<const FProcessorDependencySolver::FNode*>& Parents, const FProcessorDependencySolver::FNode& Node, const FName DepNameToFind, bool bBeforeDep)
-	{
-		check(Node.Processor);
-		const TArray<FName>& ExecuteDep = bBeforeDep ? Node.Processor->GetExecutionOrder().ExecuteBefore : Node.Processor->GetExecutionOrder().ExecuteAfter;
-		for (const FName& DependencyName : ExecuteDep)
-		{
-			if (DependencyName == DepNameToFind)
-			{
-				continue;
-			}
-			if (const FProcessorDependencySolver::FNode* DepNode = FindDepNodeInParents(Parents, DependencyName))
-			{
-				if(DepNode->Processor)
-				{
-					const TArray<FName>& ExecuteDepDep = bBeforeDep ? DepNode->Processor->GetExecutionOrder().ExecuteBefore : DepNode->Processor->GetExecutionOrder().ExecuteAfter;
-					if(ExecuteDepDep.Find(DepNameToFind) != INDEX_NONE)
-					{
-						return true;
-					}
-				}
-				else if(DoAllSubNodeHasDependency(*DepNode, DepNameToFind, bBeforeDep))
-				{
-					return true;
-				}
-			}
-		}
-		return false;
-	}
-
-	static void RemoveAllProcessorFromSet(const FProcessorDependencySolver::FNode& DepNode, TSet<const FProcessorDependencySolver::FNode*>& Set)
-	{
-		if (DepNode.Processor)
-		{
-			Set.Remove(&DepNode);
-		}
-		else
-		{
-			for (const FProcessorDependencySolver::FNode& SubNode : DepNode.SubNodes)
-			{
-				RemoveAllProcessorFromSet(SubNode, Set);
-			}
-		}
-	}
-
-	static bool AreAllProcessorInSet(const FProcessorDependencySolver::FNode& DepNode, TSet<const FProcessorDependencySolver::FNode*>& Set)
-	{
-		if (DepNode.Processor)
-		{
-			if (!Set.Contains(&DepNode))
-			{
-				return false;
-			}
-		}
-		else
-		{
-			for (const FProcessorDependencySolver::FNode& SubNode : DepNode.SubNodes)
-			{
-				if (!AreAllProcessorInSet(SubNode, Set))
-				{
-					return false;
-				}
-			}
-		}
-		return true;
-	}
-
-	static void DumpGraphDependencies(FArchive& LogFile, const FProcessorDependencySolver::FNode& Node, TArray<const FProcessorDependencySolver::FNode*> Parents, TArray<const FProcessorDependencySolver::FNode*> CommonBeforeDep, TArray<const FProcessorDependencySolver::FNode*> CommonAfterDep, TSet<const FProcessorDependencySolver::FNode*>& DependsOnStart, TSet<const FProcessorDependencySolver::FNode*>& LinkToEnd)
-	{
-		const FString NodeName = Node.Name.ToString();
-		const FProcessorDependencySolver::FNode* ParentNode = Parents.Num() > 0 ? Parents.Last() : nullptr;
-		const FString ParentNodeName = ParentNode ? ParentNode->Name.ToString() : TEXT("");
-		const FString ParentGraphNodeName = ParentNodeName.Replace(TEXT("."), TEXT(" "));
-		const FString ParentClusterNodeName = ParentNodeName.Replace(TEXT("."), TEXT("_"));
-		if (Node.Processor)
-		{
-			for (const FName& DependencyName : Node.Processor->GetExecutionOrder().ExecuteBefore)
-			{
-				if (const FProcessorDependencySolver::FNode* DepNode = FindDepNodeInParents(Parents, DependencyName))
-				{
-					LinkToEnd.Remove(&Node);
-					RemoveAllProcessorFromSet(*DepNode, DependsOnStart);
-
-					if (CommonBeforeDep.Find(DepNode) != INDEX_NONE)
-					{
-						continue;
-					}
-					if (DoesDependencyExistIndirectly(Parents, Node, DependencyName, true))
-					{
-						continue;
-					}
-
-					const FString DepNodeName = DepNode->Name.ToString();
-					if (DepNode->Processor)
-					{
-						LogFile.Logf(TEXT("    \"%s\" -> \"%s\";"), *NodeName, *DepNodeName);
-					}
-					else
-					{
-						const FString DepClusterNodeName = DepNodeName.Replace(TEXT("."), TEXT("_"));
-						const FString DepGraphNodeName = DepNodeName.Replace(TEXT("."), TEXT(" "));
-						LogFile.Logf(TEXT("    \"%s\" -> \"%s Start\"[lhead=cluster_%s];"), *NodeName, *DepGraphNodeName, *DepClusterNodeName);
-					}
-
-				}
-			}
-			for (const FName& DependencyName : Node.Processor->GetExecutionOrder().ExecuteAfter)
-			{
-				if (const FProcessorDependencySolver::FNode* DepNode = FindDepNodeInParents(Parents, DependencyName))
-				{
-					DependsOnStart.Remove(&Node);
-					RemoveAllProcessorFromSet(*DepNode, LinkToEnd);
-
-					if (CommonAfterDep.Find(DepNode) != INDEX_NONE)
-					{
-						continue;
-					}
-					if (DoesDependencyExistIndirectly(Parents, Node, DependencyName, false))
-					{
-						continue;
-					}
-
-					const FString DepNodeName = DepNode->Name.ToString();
-					if (DepNode->Processor)
-					{
-						LogFile.Logf(TEXT("    \"%s\" -> \"%s\";"), *DepNodeName, *NodeName);
-					}
-					else
-					{
-						const FString DepClusterNodeName = DepNodeName.Replace(TEXT("."), TEXT("_"));
-						const FString DepGraphNodeName = DepNodeName.Replace(TEXT("."), TEXT(" "));
-						LogFile.Logf(TEXT("    \"%s End\" -> \"%s\"[ltail=cluster_%s];"), *DepGraphNodeName, *NodeName, *DepClusterNodeName);
-					}
-				}
-			}
-
-			// Layouting
-			LogFile.Logf(TEXT("    \"%s Start\" -> \"%s\" -> \"%s End\"[style=invis];"), *ParentGraphNodeName, *NodeName, *ParentGraphNodeName);
-		}
-		else if (ParentNode)
-		{
-
-			const FString ClusterNodeName = NodeName.Replace(TEXT("."), TEXT("_"));
-			const FString GraphNodeName = NodeName.Replace(TEXT("."), TEXT(" "));
-
-			// Find common dependency through out all sub nodes
- 			TSet<FName> BeforeDepOuputed;
- 			TSet<FName> AfterDepOutputed;
-			for (const FProcessorDependencySolver::FNode& SubNode : Node.SubNodes)
-			{
-				if(SubNode.Processor)
-				{
-					for (const FName& DependencyName : SubNode.Processor->GetExecutionOrder().ExecuteBefore)
-					{
-						if (const FProcessorDependencySolver::FNode* DepNode = FindDepNodeInParents(Parents, DependencyName))
-						{
-							if (DoAllSubNodeHasDependency(Node, DependencyName, true /*bBeforeDep*/))
-							{
-								CommonBeforeDep.Add(DepNode);
-
-								if (!DoesDependencyExistIndirectly(Parents, SubNode, DependencyName, true))
-								{
-									if (BeforeDepOuputed.Contains(DependencyName))
-									{
-										continue;
-									}
-									BeforeDepOuputed.Add(DependencyName);
-
-									const FString DepNodeName = DepNode->Name.ToString();
-									if (DepNode->Processor)
-									{
-										LogFile.Logf(TEXT("    \"%s End\" -> \"%s\"[ltail=cluster_%s];"), *GraphNodeName, *DepNodeName, *ClusterNodeName);
-									}
-									else
-									{
-										const FString DepClusterNodeName = DepNodeName.Replace(TEXT("."), TEXT("_"));
-										const FString DepGraphNodeName = DepNodeName.Replace(TEXT("."), TEXT(" "));
-										LogFile.Logf(TEXT("    \"%s End\" -> \"%s Start\"[ltail=cluster_%s, lhead=cluster_%s];"), *GraphNodeName, *DepGraphNodeName, *ClusterNodeName, *DepClusterNodeName);
-									}
-								}
-							}
-						}
-					}
-
-					for (const FName& DependencyName : SubNode.Processor->GetExecutionOrder().ExecuteAfter)
-					{
-						if (const FProcessorDependencySolver::FNode* DepNode = FindDepNodeInParents(Parents, DependencyName))
-						{
-							if (DoAllSubNodeHasDependency(Node, DependencyName, false /*bBeforeDep*/))
-							{
-								CommonAfterDep.Add(DepNode);
-
-								if (!DoesDependencyExistIndirectly(Parents, SubNode, DependencyName, false))
-								{
-									if (AfterDepOutputed.Contains(DependencyName))
-									{
-										continue;
-									}
-									AfterDepOutputed.Add(DependencyName);
-
-									const FString DepNodeName = DepNode->Name.ToString();
-									if (DepNode->Processor)
-									{
-										LogFile.Logf(TEXT("    \"%s\" -> \"%s Start\"[lhead=cluster_%s];"), *DepNodeName, *GraphNodeName, *ClusterNodeName);
-									}
-									else
-									{
-										const FString DepClusterNodeName = DepNodeName.Replace(TEXT("."), TEXT("_"));
-										const FString DepGraphNodeName = DepNodeName.Replace(TEXT("."), TEXT(" "));
-										LogFile.Logf(TEXT("    \"%s End\" -> \"%s Start\"[ltail=cluster_%s, lhead=cluster_%s];"), *DepGraphNodeName,*GraphNodeName, *DepClusterNodeName, *ClusterNodeName);
-									}
-								}
-							}
-						}
-					}
-				}
-			}
-
- 			// Layouting
-			LogFile.Logf(TEXT("    \"%s Start\" -> \"%s Start\" -> \"%s End\" -> \"%s End\"[style=invis];"), *ParentGraphNodeName, *GraphNodeName, *GraphNodeName, *ParentGraphNodeName);
-		}
-
-		Parents.Push(&Node);
-		for (const FProcessorDependencySolver::FNode& SubNode : Node.SubNodes)
-		{
-			DumpGraphDependencies(LogFile, SubNode, Parents, CommonBeforeDep, CommonAfterDep, DependsOnStart, LinkToEnd);
-		}
-	}
-
-	static void PromoteStartAndEndDependency(FArchive& LogFile, const FString& GraphName, const FProcessorDependencySolver::FNode& Node, TSet<const FProcessorDependencySolver::FNode*>& DependsOnStart, TSet<const FProcessorDependencySolver::FNode*>& LinkToEnd)
-	{
-		if (Node.Processor)
-		{
-			return;
-		}
-
-		bool bAllDependsOnStart = true;
-		bool bAllLinkedToEnd = true;
-		for (const FProcessorDependencySolver::FNode& SubNode : Node.SubNodes)
-		{
-			if (!AreAllProcessorInSet(SubNode, DependsOnStart))
-			{
-				bAllDependsOnStart = false;
-			}
-
-			if (!AreAllProcessorInSet(SubNode, LinkToEnd))
-			{
-				bAllLinkedToEnd = false;
-			}
-
-			if (!bAllDependsOnStart && !bAllLinkedToEnd)
-			{
-				break;
-			}
-		}
-
-		const FString NodeName = Node.Name.ToString();
-		const FString ClusterNodeName = NodeName.Replace(TEXT("."), TEXT("_"));
-		const FString GraphNodeName = NodeName.Replace(TEXT("."), TEXT(" "));
-		if (bAllDependsOnStart)
-		{
-			RemoveAllProcessorFromSet(Node, DependsOnStart);
-			LogFile.Logf(TEXT("    \"%s Start\" -> \"%s Start\"[lhead=cluster_%s, weight=0];"), *GraphName, *GraphNodeName, *ClusterNodeName);
-		}
-
-		if (bAllLinkedToEnd)
-		{
-			RemoveAllProcessorFromSet(Node, LinkToEnd);
-			LogFile.Logf(TEXT("    \"%s End\" -> \"%s End\"[ltail=cluster_%s, weight=0];"), *GraphNodeName, *GraphName, *ClusterNodeName);
-		}
-
-		if(!bAllDependsOnStart || !bAllLinkedToEnd)
-		{
-			for (const FProcessorDependencySolver::FNode& SubNode : Node.SubNodes)
-			{
-				PromoteStartAndEndDependency(LogFile, GraphName, SubNode, DependsOnStart, LinkToEnd);
-			}
-		}
-	}
-
-	static void DumpAllGraphDependencies(FArchive& LogFile, const FProcessorDependencySolver::FNode& GroupRootNode, const TSet<const FProcessorDependencySolver::FNode*>& AllNodes)
-	{
-		TSet<const FProcessorDependencySolver::FNode*> DependsOnStart = AllNodes;
-		TSet<const FProcessorDependencySolver::FNode*> LinkToEnd = AllNodes;
-		const FString GraphName = GroupRootNode.Name.ToString();
-		FDumpGraphDependencyUtils::DumpGraphDependencies(LogFile, GroupRootNode, TArray<const FProcessorDependencySolver::FNode*>(), TArray<const FProcessorDependencySolver::FNode*>(), TArray<const FProcessorDependencySolver::FNode*>(), DependsOnStart, LinkToEnd);
-		FDumpGraphDependencyUtils::PromoteStartAndEndDependency(LogFile, GraphName, GroupRootNode, DependsOnStart, LinkToEnd);
-		for (const FProcessorDependencySolver::FNode* Node : DependsOnStart)
-		{
-			LogFile.Logf(TEXT("    \"%s Start\" -> \"%s\"[weight=0];"), *GraphName, *Node->Name.ToString());
-		}
-		for (const FProcessorDependencySolver::FNode* Node : LinkToEnd)
-		{
-			LogFile.Logf(TEXT("    \"%s\" -> \"%s End\"[weight=0];"), *Node->Name.ToString(), *GraphName);
-		}
-	}
-};
-
-void FProcessorDependencySolver::DumpGraph(FArchive& LogFile) const
-{
-	TSet<const FNode*> AllNodes;
-	LogFile.Logf(TEXT("digraph MassProcessorGraph"));
-	LogFile.Logf(TEXT("{"));
-	LogFile.Logf(TEXT("    compound = true;"));
-	LogFile.Logf(TEXT("    newrank = true;"));
-	FDumpGraphDependencyUtils::DumpGraphNode(LogFile, GroupRootNode, 4/* Indent */, AllNodes, true/* bRoot */);
-	FDumpGraphDependencyUtils::DumpAllGraphDependencies(LogFile, GroupRootNode, AllNodes);
-	LogFile.Logf(TEXT("}"));
-}
-
-#endif // 0; disabled the graph building for now, leaving the old code here for reference
-
 void FMassProcessorDependencySolver::Solve(TArray<FMassProcessorOrderInfo>& OutResult)
 {
 	using UE::Mass::Private::NameViewToString;
