[Mass Debugger]
Added editable breakpoints view list

diff --git a/Engine/Source/Editor/MassEntityDebugger/Private/SMassBreakpointsView.cpp b/Engine/Source/Editor/MassEntityDebugger/Private/SMassBreakpointsView.cpp
--- a/Engine/Source/Editor/MassEntityDebugger/Private/SMassBreakpointsView.cpp
+++ b/Engine/Source/Editor/MassEntityDebugger/Private/SMassBreakpointsView.cpp
@@ -18,6 +18,12 @@ void SMassBreakpointsView::Construct(const FArguments& InArgs, TSharedRef<FMassD
 {
 	Initialize(InDebuggerModel);
 
+	FMassDebugger::OnBreakpointsChangedDelegate.AddSP(this, &SMassBreakpointsView::RefreshBreakpoints);
+
+	WriteBreakpointsListView = SNew(SListView<TSharedPtr<FBreakpointDisplay>>)
+		.ListItemsSource(&WriteBreakpoints)
+		.OnGenerateRow(this, &SMassBreakpointsView::OnGenerateBreakpointRow);
+
 	ChildSlot
 	[
 		SNew(SVerticalBox)
@@ -32,6 +38,12 @@ void SMassBreakpointsView::Construct(const FArguments& InArgs, TSharedRef<FMassD
 		+ SVerticalBox::Slot()
 		.AutoHeight()
 		.Padding(5)
+		[
+			WriteBreakpointsListView.ToSharedRef()
+		]
+		+ SVerticalBox::Slot()
+		.AutoHeight()
+		.Padding(5)
 		[
 			SNew(SButton)
 			.Text(LOCTEXT("ClearAllBreakpoints", "Clear All Breakpoints"))
@@ -40,6 +52,27 @@ void SMassBreakpointsView::Construct(const FArguments& InArgs, TSharedRef<FMassD
 	];
 }
 
+void SMassBreakpointsView::RefreshBreakpoints()
+{
+	WriteBreakpoints.Reset();
+	if (DebuggerModel && DebuggerModel->Environment)
+	{
+		FMassDebugger::FEnvironment* Env = FMassDebugger::FindEnvironmentForEntityManager(DebuggerModel->Environment->EntityManager.Pin().ToSharedRef().Get());
+		TMultiMap<const UScriptStruct*, FMassEntityHandle>::TIterator FragmentIterator = Env->FragmentWriteBreakpoints.CreateIterator();
+		for (; FragmentIterator; ++FragmentIterator)
+		{
+			TSharedPtr<FBreakpointDisplay> NewDisplay = MakeShared<FBreakpointDisplay>();
+			NewDisplay->WriteFragment = FragmentIterator.Key();
+			NewDisplay->Entity = FragmentIterator.Value();
+			WriteBreakpoints.Add(NewDisplay);
+		}
+	}
+	if (WriteBreakpointsListView)
+	{
+		WriteBreakpointsListView->RebuildList();
+	}
+}
+
 FReply SMassBreakpointsView::HandleAddWriteBreakpointClicked()
 {
 #if WITH_MASSENTITY_DEBUG
@@ -56,5 +89,70 @@ FReply SMassBreakpointsView::ClearBreakpointsClicked()
 	return FReply::Handled();
 }
 
+TSharedRef<ITableRow> SMassBreakpointsView::OnGenerateBreakpointRow(TSharedPtr<FBreakpointDisplay> InItem, const TSharedRef<STableViewBase>& OwnerTable)
+{
+	return SNew(STableRow<TSharedPtr<FMassEntityHandle>>, OwnerTable)
+		.Content()
+		[
+#if WITH_MASSENTITY_DEBUG
+            SNew(SHorizontalBox)
+			+ SHorizontalBox::Slot()
+            .AutoWidth()
+            .Padding(5.0f, 2.0f)
+            [
+                // Add a button to remove the breakpoint
+                SNew(SButton)
+                .OnClicked_Lambda([this, InItem]()
+                {
+					if (DebuggerModel.IsValid() && DebuggerModel->Environment.IsValid())
+					{
+						const FMassEntityManager* EntityManager = DebuggerModel->Environment->EntityManager.Pin().Get();
+						if (EntityManager)
+						{
+							if (InItem->Processor)
+							{
+								FMassDebugger::ClearProcessorBreakpoint(*EntityManager, InItem->Processor, InItem->Entity);
+							}
+							else if (InItem->WriteFragment)
+							{
+								FMassDebugger::ClearFragmentWriteBreak(*EntityManager, InItem->WriteFragment, InItem->Entity);
+							}
+						}
+					}
+                    return FReply::Handled();
+                })
+				[
+					SNew(SImage)
+						.Image(FCoreStyle::Get().GetBrush("Icons.Delete"))
+						.ColorAndOpacity(FSlateColor::UseForeground())
+				]
+            ]
+            + SHorizontalBox::Slot()
+            .AutoWidth()
+            .Padding(5.0f, 2.0f)
+            [
+                // Display WriteFragment or Processor name
+                SNew(STextBlock)
+                .Text(InItem->WriteFragment 
+                    ? FText::FromString(InItem->WriteFragment->GetName()) 
+                    : (InItem->Processor 
+                        ? FText::FromString(InItem->Processor->GetName()) 
+                        : LOCTEXT("UnknownBreakpoint", "Unknown")))
+            ]
+            + SHorizontalBox::Slot()
+            .AutoWidth()
+            .Padding(5.0f, 2.0f)
+            [
+                // Display Entity handle
+                SNew(STextBlock)
+				.Text(FText::FromString(InItem->Entity.DebugGetDescription()))
+            ]
+#else
+			SNew(STextBlock)
+				.Text(LOCTEXT("MassEntityDebuggingNotEnabled", "Mass Entity Debugging Not Enabled for this configuration"))
+#endif
+		];
+}
+
 #undef LOCTEXT_NAMESPACE
 
diff --git a/Engine/Source/Editor/MassEntityDebugger/Private/SMassBreakpointsView.h b/Engine/Source/Editor/MassEntityDebugger/Private/SMassBreakpointsView.h
--- a/Engine/Source/Editor/MassEntityDebugger/Private/SMassBreakpointsView.h
+++ b/Engine/Source/Editor/MassEntityDebugger/Private/SMassBreakpointsView.h
@@ -2,11 +2,15 @@
 #pragma once
 
 #include "SMassDebuggerViewBase.h"
+#include "Widgets/Views/SListView.h"
+#include "MassEntityHandle.h"
 
 struct FMassDebuggerModel;
 struct FMassDebuggerArchetypeData;
 struct FMassDebuggerProcessorData;
-
+class ITableRow;
+class STableViewBase;
+class UMassProcessor;
 
 class SMassBreakpointsView : public SMassDebuggerViewBase
 {
@@ -39,4 +43,16 @@ protected:
 private:
 	FReply HandleAddWriteBreakpointClicked();
 	FReply ClearBreakpointsClicked();
+
+	struct FBreakpointDisplay
+	{
+		const UMassProcessor* Processor;
+		const UScriptStruct* WriteFragment;
+		FMassEntityHandle Entity;
+	};
+	TArray<TSharedPtr<FBreakpointDisplay>> WriteBreakpoints;
+	TSharedRef<ITableRow> OnGenerateBreakpointRow(TSharedPtr<FBreakpointDisplay> InItem, const TSharedRef<STableViewBase>& OwnerTable);
+	TSharedPtr<SListView<TSharedPtr<FBreakpointDisplay>>> WriteBreakpointsListView;
+
+	void RefreshBreakpoints();
 };
diff --git a/Engine/Source/Runtime/MassEntity/Private/MassDebugger.cpp b/Engine/Source/Runtime/MassEntity/Private/MassDebugger.cpp
--- a/Engine/Source/Runtime/MassEntity/Private/MassDebugger.cpp
+++ b/Engine/Source/Runtime/MassEntity/Private/MassDebugger.cpp
@@ -466,6 +466,7 @@ namespace UE::Mass::Debug
 //----------------------------------------------------------------------//
 // FMassDebugger
 //----------------------------------------------------------------------//
+FMassDebugger::FOnBreakpointsChanged FMassDebugger::OnBreakpointsChangedDelegate;
 FMassDebugger::FOnEntitySelected FMassDebugger::OnEntitySelectedDelegate;
 
 FMassDebugger::FOnMassEntityManagerEvent FMassDebugger::OnEntityManagerInitialized;
@@ -843,6 +844,18 @@ void FMassDebugger::RegisterProcessorDataProvider(FName ProviderName, const TSha
 	OnProcessorProviderRegistered.Broadcast(ActiveEnvironments[Index]);
 }
 
+FMassDebugger::FEnvironment* FMassDebugger::FindEnvironmentForEntityManager(const FMassEntityManager& EntityManager)
+{
+	for (FMassDebugger::FEnvironment& Environment : ActiveEnvironments)
+	{
+		if (Environment.EntityManager.HasSameObject(&EntityManager))
+		{
+			return &Environment;
+		}
+	}
+	return nullptr;
+}
+
 bool FMassDebugger::DoesArchetypeMatchRequirements(const FMassArchetypeHandle& ArchetypeHandle, const FMassFragmentRequirements& Requirements, FOutputDevice& OutputDevice)
 {
 	if (const FMassArchetypeData* Archetype = FMassArchetypeHelper::ArchetypeDataFromHandle(ArchetypeHandle))
@@ -936,7 +949,8 @@ void FMassDebugger::SetProcessorBreakpoint(const FMassEntityManager& EntityManag
 	ActiveEnvironment.bHasBreakpoint = true;
 	bHasBreakpoint = true;
 
-	ActiveEnvironment.ProcessorBreakpoints.Add(Processor, Entity);
+	ActiveEnvironment.ProcessorBreakpoints.AddUnique(Processor, Entity);
+	OnBreakpointsChangedDelegate.Broadcast();
 }
 
 void FMassDebugger::SetFragmentWriteBreak(const FMassEntityManager& EntityManager, const UScriptStruct* FragmentType, FMassEntityHandle Entity)
@@ -946,7 +960,8 @@ void FMassDebugger::SetFragmentWriteBreak(const FMassEntityManager& EntityManage
 	ActiveEnvironment.bHasBreakpoint = true;
 	bHasBreakpoint = true;
 
-	ActiveEnvironment.FragmentWriteBreakpoints.Add(FragmentType, Entity);
+	ActiveEnvironment.FragmentWriteBreakpoints.AddUnique(FragmentType, Entity);
+	OnBreakpointsChangedDelegate.Broadcast();
 }
 
 void FMassDebugger::ClearProcessorBreakpoint(const FMassEntityManager& EntityManager, const UMassProcessor* Processor, FMassEntityHandle Entity)
@@ -956,12 +971,14 @@ void FMassDebugger::ClearProcessorBreakpoint(const FMassEntityManager& EntityMan
 	ActiveEnvironment.ProcessorBreakpoints.Remove(Processor, Entity);
 
 	UpdateHasBreakpoint();
+	OnBreakpointsChangedDelegate.Broadcast();
 }
 
 void FMassDebugger::ClearAllProcessorBreakpoints(const FMassEntityManager& EntityManager, const UMassProcessor* Processor)
 {
 	FEnvironment& ActiveEnvironment = GetActiveEnvironment(EntityManager);
 	ActiveEnvironment.ProcessorBreakpoints.Remove(Processor);
+	OnBreakpointsChangedDelegate.Broadcast();
 }
 
 void FMassDebugger::ClearFragmentWriteBreak(const FMassEntityManager& EntityManager, const UScriptStruct* FragmentType, FMassEntityHandle Entity)
@@ -969,12 +986,14 @@ void FMassDebugger::ClearFragmentWriteBreak(const FMassEntityManager& EntityMana
 	FEnvironment& ActiveEnvironment = GetActiveEnvironment(EntityManager);
 	ActiveEnvironment.FragmentWriteBreakpoints.Remove(FragmentType, Entity);
 	UpdateHasBreakpoint();
+	OnBreakpointsChangedDelegate.Broadcast();
 }
 
 void FMassDebugger::ClearAllFragmentWriteBreak(const FMassEntityManager& EntityManager, const UScriptStruct* FragmentType)
 {
 	FEnvironment& ActiveEnvironment = GetActiveEnvironment(EntityManager);
 	ActiveEnvironment.FragmentWriteBreakpoints.Remove(FragmentType);
+	OnBreakpointsChangedDelegate.Broadcast();
 }
 
 void FMassDebugger::ClearAllEntityBreakpoints(const FMassEntityManager& EntityManager, FMassEntityHandle Entity)
@@ -998,6 +1017,7 @@ void FMassDebugger::ClearAllEntityBreakpoints(const FMassEntityManager& EntityMa
 			FragmentIterator.RemoveCurrent();
 		}
 	}
+	OnBreakpointsChangedDelegate.Broadcast();
 }
  
 void FMassDebugger::BreakOnFragmentWriteForSelectedEntity(FName FragmentName)
@@ -1009,6 +1029,7 @@ void FMassDebugger::BreakOnFragmentWriteForSelectedEntity(FName FragmentName)
 			SetFragmentWriteBreak(*Environment.EntityManager.Pin(), GetFragmentTypeFromName(FragmentName), Environment.SelectedEntity);
 		}
 	}
+	OnBreakpointsChangedDelegate.Broadcast();
 }
 
 void FMassDebugger::ClearAllBreakpoints()
@@ -1018,6 +1039,7 @@ void FMassDebugger::ClearAllBreakpoints()
 		Environment.ClearBreakpoints();
 	}
 	bHasBreakpoint = false;
+	OnBreakpointsChangedDelegate.Broadcast();
 }
 
 const UScriptStruct* FMassDebugger::GetFragmentTypeFromName(FName FragmentName)
diff --git a/Engine/Source/Runtime/MassEntity/Public/MassDebugger.h b/Engine/Source/Runtime/MassEntity/Public/MassDebugger.h
--- a/Engine/Source/Runtime/MassEntity/Public/MassDebugger.h
+++ b/Engine/Source/Runtime/MassEntity/Public/MassDebugger.h
@@ -156,6 +156,7 @@ struct MASSENTITY_API FMassDebugger
 		void ClearBreakpoints();
 	};
 
+	DECLARE_TS_MULTICAST_DELEGATE(FOnBreakpointsChanged);
 	DECLARE_TS_MULTICAST_DELEGATE_TwoParams(FOnEntitySelected, const FMassEntityManager&, const FMassEntityHandle);
 	DECLARE_TS_MULTICAST_DELEGATE_OneParam(FOnMassEntityManagerEvent, const FMassEntityManager&);
 	DECLARE_TS_MULTICAST_DELEGATE_OneParam(FOnEnvironmentEvent, const FEnvironment&);
@@ -194,6 +195,7 @@ struct MASSENTITY_API FMassDebugger
 	static void HighlightEntity(const FMassEntityManager& EntityManager, const FMassEntityHandle EntityHandle);
 	static FMassEntityHandle GetHighlightedEntity(const FMassEntityManager& EntityManager);
 
+	static FOnBreakpointsChanged OnBreakpointsChangedDelegate;
 	static FOnEntitySelected OnEntitySelectedDelegate;
 
 	static FOnMassEntityManagerEvent OnEntityManagerInitialized;
@@ -230,6 +232,7 @@ struct MASSENTITY_API FMassDebugger
 	static void RegisterProcessorDataProvider(FName ProviderName, const TSharedRef<FMassEntityManager>& EntityManager, const UE::Mass::Debug::FProcessorProviderFunction& ProviderFunction);
 
 	static TConstArrayView<FEnvironment> GetEnvironments() { return ActiveEnvironments; }
+	static FEnvironment* FindEnvironmentForEntityManager(const FMassEntityManager& EntityManager);
 
 	/**
 	 * Determines whether given Archetype matches given Requirements. In case of a mismatch description of failed conditions will be added to OutputDevice.
