[Mass] added support for adding and removing shared fragments via FMassEntityManager::AddCompositionToEntity_GetDelta and FMassEntityManager::RemoveCompositionFromEntity, respectively.

diff --git a/Engine/Source/Runtime/MassEntity/Private/MassEntityManager.cpp b/Engine/Source/Runtime/MassEntity/Private/MassEntityManager.cpp
--- a/Engine/Source/Runtime/MassEntity/Private/MassEntityManager.cpp
+++ b/Engine/Source/Runtime/MassEntity/Private/MassEntityManager.cpp
@@ -1112,7 +1112,7 @@ void FMassEntityManager::AddFragmentListToEntity(FMassEntityHandle Entity, TCons
 	}
 }
 
-void FMassEntityManager::AddCompositionToEntity_GetDelta(FMassEntityHandle Entity, FMassArchetypeCompositionDescriptor& InOutDescriptor)
+void FMassEntityManager::AddCompositionToEntity_GetDelta(FMassEntityHandle Entity, FMassArchetypeCompositionDescriptor& InOutDescriptor, const FMassArchetypeSharedFragmentValues* AddedSharedFragmentValues)
 {
 	CHECK_SYNC_API();
 
@@ -1124,7 +1124,9 @@ void FMassEntityManager::AddCompositionToEntity_GetDelta(FMassEntityHandle Entit
 	InOutDescriptor.Remove(OldArchetype->GetCompositionDescriptor());
 
 	ensureMsgf(InOutDescriptor.ChunkFragments.IsEmpty(), TEXT("Adding new chunk fragments is not supported"));
-	ensureMsgf(InOutDescriptor.SharedFragments.IsEmpty(), TEXT("Adding new shared fragments is not supported"));
+	ensureMsgf(InOutDescriptor.SharedFragments.IsEmpty() 
+		|| (AddedSharedFragmentValues && AddedSharedFragmentValues->DoesMatchComposition(InOutDescriptor))
+		, TEXT("When adding new shared fragments it's required to provide values for said fragments"));
 
 	if (InOutDescriptor.IsEmpty() == false)
 	{
@@ -1138,7 +1140,18 @@ void FMassEntityManager::AddCompositionToEntity_GetDelta(FMassEntityHandle Entit
 			// Move the entity over
 			FMassArchetypeData& NewArchetype = FMassArchetypeHelper::ArchetypeDataFromHandleChecked(NewArchetypeHandle);
 			NewArchetype.CopyDebugNamesFrom(*OldArchetype);
-			OldArchetype->MoveEntityToAnotherArchetype(Entity, NewArchetype);
+			if (AddedSharedFragmentValues)
+			{
+				// we need to merge AddedSharedFragmentValues with OldArchetype's shared fragments
+				FMassArchetypeSharedFragmentValues CurrentSharedFragment = OldArchetype->GetSharedFragmentValues(Entity);
+				CurrentSharedFragment.Append(*AddedSharedFragmentValues);
+				CurrentSharedFragment.Sort();
+				OldArchetype->MoveEntityToAnotherArchetype(Entity, NewArchetype, &CurrentSharedFragment);
+			}
+			else
+			{
+				OldArchetype->MoveEntityToAnotherArchetype(Entity, NewArchetype);
+			}
 
 			GetEntityStorageInterface().SetArchetypeFromShared(Entity.Index, NewArchetypeHandle.DataPtr);
 
@@ -1181,7 +1194,18 @@ void FMassEntityManager::RemoveCompositionFromEntity(FMassEntityHandle Entity, c
 				// Move the entity over
 				FMassArchetypeData& NewArchetype = FMassArchetypeHelper::ArchetypeDataFromHandleChecked(NewArchetypeHandle);
 				NewArchetype.CopyDebugNamesFrom(*OldArchetype);
-				OldArchetype->MoveEntityToAnotherArchetype(Entity, NewArchetype);
+				if (InDescriptor.SharedFragments.IsEmpty() && InDescriptor.ConstSharedFragments.IsEmpty())
+				{
+					OldArchetype->MoveEntityToAnotherArchetype(Entity, NewArchetype);
+				}
+				else
+				{
+					// we need to remove the shared fragment values to match the new composition
+					FMassArchetypeSharedFragmentValues CurrentSharedFragment = OldArchetype->GetSharedFragmentValues(Entity);
+					CurrentSharedFragment.Remove(InDescriptor);
+					CurrentSharedFragment.Sort();
+					OldArchetype->MoveEntityToAnotherArchetype(Entity, NewArchetype, &CurrentSharedFragment);
+				}
 				GetEntityStorageInterface().SetArchetypeFromShared(Entity.Index, NewArchetypeHandle.DataPtr);
 			}
 		}
diff --git a/Engine/Source/Runtime/MassEntity/Private/MassEntityTypes.cpp b/Engine/Source/Runtime/MassEntity/Private/MassEntityTypes.cpp
--- a/Engine/Source/Runtime/MassEntity/Private/MassEntityTypes.cpp
+++ b/Engine/Source/Runtime/MassEntity/Private/MassEntityTypes.cpp
@@ -255,12 +255,9 @@ int32 FMassArchetypeSharedFragmentValues::Append(const FMassArchetypeSharedFragm
 		if (SharedFragmentBitSet.Contains(*StructType))
 		{
 			const int32 FragmentIndex = SharedFragments.IndexOfByPredicate(FStructTypeEqualOperator(StructType));
-			if (testableEnsureMsgf(FragmentIndex != INDEX_NONE, TEXT("%hs trying to switch const fragment to non-const, type %s")
-				, __FUNCTION__, *StructType->GetName()))
-			{
-				SharedFragments[FragmentIndex] = SharedStruct;
-				++AddedOrModifiedCount;
-			}
+			checkf(FragmentIndex != INDEX_NONE, TEXT("Mismatch between shared fragment bitset and stored values"));
+			SharedFragments[FragmentIndex] = SharedStruct;
+			++AddedOrModifiedCount;
 		}
 		else
 		{
@@ -276,12 +273,9 @@ int32 FMassArchetypeSharedFragmentValues::Append(const FMassArchetypeSharedFragm
 		if (ConstSharedFragmentBitSet.Contains(*StructType))
 		{
 			const int32 FragmentIndex = ConstSharedFragments.IndexOfByPredicate(FStructTypeEqualOperator(StructType));
-			if (testableEnsureMsgf(FragmentIndex != INDEX_NONE, TEXT("%hs trying to switch non-const fragment to const, type %s")
-				, __FUNCTION__, *StructType->GetName()))
-			{
-				ConstSharedFragments[FragmentIndex] = SharedStruct;
-				++AddedOrModifiedCount;
-			}
+			checkf(FragmentIndex != INDEX_NONE, TEXT("Mismatch between const shared fragment bitset and stored values"));
+			ConstSharedFragments[FragmentIndex] = SharedStruct;
+			++AddedOrModifiedCount;
 		}
 		else
 		{
diff --git a/Engine/Source/Runtime/MassEntity/Public/MassEntityManager.h b/Engine/Source/Runtime/MassEntity/Public/MassEntityManager.h
--- a/Engine/Source/Runtime/MassEntity/Public/MassEntityManager.h
+++ b/Engine/Source/Runtime/MassEntity/Public/MassEntityManager.h
@@ -466,11 +466,13 @@ public:
 	void BatchAddSharedFragmentsForEntities(TConstArrayView<FMassArchetypeEntityCollection> EntityCollections, const FMassArchetypeSharedFragmentValues& AddedFragmentValues);
 
 	/**
-	 * Adds fragments and tags indicated by InOutDescriptor to the entity indicated by EntityHandle. The function also figures out which elements
+	 * Adds elements indicated by InOutDescriptor to the entity indicated by EntityHandle. The function also figures out which elements
 	 * in InOutDescriptor are missing from the current composition of the given entity and then returns the resulting 
 	 * delta via InOutDescriptor.
+	 * If InOutDescriptor indicates shared fragments to be added the caller is required to provide matching values for the indicated
+	 * shared fragment types, via AddedSharedFragmentValues.
 	 */
-	void AddCompositionToEntity_GetDelta(FMassEntityHandle EntityHandle, FMassArchetypeCompositionDescriptor& InOutDescriptor);
+	void AddCompositionToEntity_GetDelta(FMassEntityHandle EntityHandle, FMassArchetypeCompositionDescriptor& InOutDescriptor, const FMassArchetypeSharedFragmentValues* AddedSharedFragmentValues = nullptr);
 	void RemoveCompositionFromEntity(FMassEntityHandle EntityHandle, const FMassArchetypeCompositionDescriptor& InDescriptor);
 
 	const FMassArchetypeCompositionDescriptor& GetArchetypeComposition(const FMassArchetypeHandle& ArchetypeHandle) const;
diff --git a/Engine/Source/Runtime/MassEntity/Public/MassEntityTypes.h b/Engine/Source/Runtime/MassEntity/Public/MassEntityTypes.h
--- a/Engine/Source/Runtime/MassEntity/Public/MassEntityTypes.h
+++ b/Engine/Source/Runtime/MassEntity/Public/MassEntityTypes.h
@@ -237,6 +237,15 @@ struct MASSENTITY_API FMassArchetypeSharedFragmentValues
 		return ConstSharedFragmentBitSet.HasAll(InConstSharedFragmentBitSet);
 	}
 
+	/**
+	 * @return whether the stored shared fragment values exactly match shared fragment types indicated by InDescriptor
+	 */
+	bool DoesMatchComposition(const FMassArchetypeCompositionDescriptor& InDescriptor) const 
+	{
+		return HasExactSharedFragmentTypesMatch(InDescriptor.SharedFragments)
+			&& HasExactConstSharedFragmentTypesMatch(InDescriptor.ConstSharedFragments);
+	}
+
 	FORCEINLINE bool IsEquivalent(const FMassArchetypeSharedFragmentValues& OtherSharedFragmentValues) const
 	{
 		return GetTypeHash(*this) == GetTypeHash(OtherSharedFragmentValues);
@@ -324,6 +333,15 @@ struct MASSENTITY_API FMassArchetypeSharedFragmentValues
 	 */
 	int32 Remove(const FMassConstSharedFragmentBitSet& ConstSharedFragmentToRemoveBitSet);
 
+	/**
+	 * Remove all the shared and const shared fragments indicated by InDescriptor
+	 * @return number of fragments types removed
+	 */
+	int32 Remove(const FMassArchetypeCompositionDescriptor& InDescriptor)
+	{
+		return Remove(InDescriptor.SharedFragments) + Remove(InDescriptor.ConstSharedFragments);
+	}
+
 	FORCEINLINE const TArray<FConstSharedStruct>& GetConstSharedFragments() const
 	{
 		return ConstSharedFragments;
