[Mass]
* Removed  dllstorage on template class

diff --git a/Engine/Source/Runtime/MassEntity/Public/MassBitSetRegistry.h b/Engine/Source/Runtime/MassEntity/Public/MassBitSetRegistry.h
--- a/Engine/Source/Runtime/MassEntity/Public/MassBitSetRegistry.h
+++ b/Engine/Source/Runtime/MassEntity/Public/MassBitSetRegistry.h
@@ -46,7 +46,7 @@ namespace UE::Mass
 		 * @param TUStructType the Unreal Engine struct type, default is UScriptStruct.
 		 */
 		template<typename T, typename TUStructType = UScriptStruct>
-		struct MASSENTITY_API TBitTypeRegistry
+		struct TBitTypeRegistry
 		{
 			/** Alias for the bitset builder specific to the type T */
 			using FBitSetBuilder = TTypeBitSetBuilder<T, TUStructType, /*bTestInheritanceAtRuntime=*/TTypeBitSetTraits<T>::RequiresBaseType>;
