Added the ability to retrieve all archetype handles that match a Mass query.

diff --git a/Engine/Source/Runtime/MassEntity/Private/MassEntityManager.cpp b/Engine/Source/Runtime/MassEntity/Private/MassEntityManager.cpp
--- a/Engine/Source/Runtime/MassEntity/Private/MassEntityManager.cpp
+++ b/Engine/Source/Runtime/MassEntity/Private/MassEntityManager.cpp
@@ -610,6 +610,11 @@ FMassArchetypeHandle FMassEntityManager::GetArchetypeForEntityUnsafe(FMassEntity
 	return FMassArchetypeHelper::ArchetypeHandleFromData(GetEntityStorageInterface().GetArchetypeAsShared(Entity.Index));
 }
 
+void FMassEntityManager::GetMatchingArchetypes(const FMassFragmentRequirements& Requirements, TArray<FMassArchetypeHandle>& OutValidArchetypes) const
+{
+	GetMatchingArchetypes(Requirements, OutValidArchetypes, 0);
+}
+
 void FMassEntityManager::ForEachArchetypeFragmentType(const FMassArchetypeHandle& ArchetypeHandle, TFunction< void(const UScriptStruct* /*FragmentType*/)> Function)
 {
 	const FMassArchetypeData& ArchetypeData = FMassArchetypeHelper::ArchetypeDataFromHandleChecked(ArchetypeHandle);
diff --git a/Engine/Source/Runtime/MassEntity/Public/MassEntityManager.h b/Engine/Source/Runtime/MassEntity/Public/MassEntityManager.h
--- a/Engine/Source/Runtime/MassEntity/Public/MassEntityManager.h
+++ b/Engine/Source/Runtime/MassEntity/Public/MassEntityManager.h
@@ -206,6 +206,13 @@ public:
 	 */
 	FMassArchetypeHandle GetArchetypeForEntityUnsafe(FMassEntityHandle EntityHandle) const;
 
+	/**
+	 * Searches through all known archetypes and matches them to the provided requirements. All archetypes that pass the requirement check are returned.
+	 * @param Requirements The set of fragments and tags that need to be on available in the request form on an archetype before it's added.
+	 * @param OutValidArchetypes Archetypes that pass the requirements test are added here.
+	 */
+	void GetMatchingArchetypes(const FMassFragmentRequirements& Requirements, TArray<FMassArchetypeHandle>& OutValidArchetypes) const;
+
 	/** Method to iterate on all the fragment types of an archetype */
 	static void ForEachArchetypeFragmentType(const FMassArchetypeHandle& ArchetypeHandle, TFunction< void(const UScriptStruct* /*FragmentType*/)> Function);
 
