[Mass] removed a no-longer-valid ensureMsgf from FMassEntityManager::RemoveCompositionFromEntity, it's been accidentally left behind during last change in that code.

diff --git a/Engine/Source/Runtime/MassEntity/Private/MassEntityManager.cpp b/Engine/Source/Runtime/MassEntity/Private/MassEntityManager.cpp
--- a/Engine/Source/Runtime/MassEntity/Private/MassEntityManager.cpp
+++ b/Engine/Source/Runtime/MassEntity/Private/MassEntityManager.cpp
@@ -1183,7 +1183,6 @@ void FMassEntityManager::RemoveCompositionFromEntity(FMassEntityHandle Entity, c
 		NewDescriptor.Remove(InDescriptor);
 
 		ensureMsgf(InDescriptor.ChunkFragments.IsEmpty(), TEXT("Removing chunk fragments is not supported"));
-		ensureMsgf(InDescriptor.SharedFragments.IsEmpty(), TEXT("Removing shared fragments is not supported"));
 
 		if (NewDescriptor.IsEquivalent(OldArchetype->GetCompositionDescriptor()) == false)
 		{
