[LLM] Mark Mass storage and manager allocations under separate LLM tag

#rnx

diff --git a/Engine/Source/Runtime/MassEntity/Internal/MassArchetypeData.h b/Engine/Source/Runtime/MassEntity/Internal/MassArchetypeData.h
--- a/Engine/Source/Runtime/MassEntity/Internal/MassArchetypeData.h
+++ b/Engine/Source/Runtime/MassEntity/Internal/MassArchetypeData.h
@@ -8,6 +8,7 @@
 #if UE_ENABLE_INCLUDE_ORDER_DEPRECATED_IN_5_6
 #include "MassEntityManager.h"
 #endif // UE_ENABLE_INCLUDE_ORDER_DEPRECATED_IN_5_6
+#include "HAL/LowLevelMemTracker.h"
 
 struct FMassEntityQuery;
 struct FMassExecutionContext;
@@ -38,6 +39,8 @@ public:
 		, ChunkFragmentData(InChunkFragmentTemplates)
 		, SharedFragmentValues(InSharedFragmentValues)
 	{
+		
+		LLM_SCOPE_BYNAME(TEXT("Mass/ArchetypeChunk"));
 		RawMemory = (uint8*)FMemory::Malloc(AllocSize);
 	}
 
diff --git a/Engine/Source/Runtime/MassEntity/Private/MassEntityManager.cpp b/Engine/Source/Runtime/MassEntity/Private/MassEntityManager.cpp
--- a/Engine/Source/Runtime/MassEntity/Private/MassEntityManager.cpp
+++ b/Engine/Source/Runtime/MassEntity/Private/MassEntityManager.cpp
@@ -103,6 +103,7 @@ void FMassEntityManager::FEntityCreationContext::AppendEntities(const TConstArra
 
 	if (EntitiesToAppend.Num())
 	{
+		LLM_SCOPE_BYNAME(TEXT("Mass/EntityManager"));
 		if (CreatedEntities.Num())
 		{
 			// since we already have entities in CreatedEntities (initially ensured to have no duplicates) we cannot 
@@ -126,7 +127,9 @@ void FMassEntityManager::FEntityCreationContext::AppendEntities(const TConstArra
 	{
 		return;
 	}
-	
+
+	LLM_SCOPE_BYNAME(TEXT("Mass/EntityManager"));
+
 	AppendEntities(EntitiesToAppend);
 
 	// this condition boils down to checking if this FEntityCreationContext instance only connects the just added EntitiesToAppend
@@ -248,6 +251,8 @@ void FMassEntityManager::Initialize(const FMassEntityManagerStorageInitParams& I
 		return;
 	}
 
+	LLM_SCOPE_BYNAME(TEXT("Mass/EntityManager"));
+
 	Visit(UE::Mass::Private::FEntityStorageInitializer{&EntityStorage}, InitializationParams);
 #if WITH_MASSENTITY_DEBUG
 	DebugEntityStoragePtr = &DebugGetEntityStorageInterface();
@@ -362,6 +367,7 @@ void FMassEntityManager::OnPostFork(EForkProcessRole Role)
 {
 	if (Role == EForkProcessRole::Child)
 	{
+		LLM_SCOPE_BYNAME(TEXT("Mass/EntityManager"));
 		for (TSharedPtr<FMassCommandBuffer>& CommandBuffer : DeferredCommandBuffers)
 		{
 			if (CommandBuffer)
@@ -440,6 +446,7 @@ FMassArchetypeHandle FMassEntityManager::GetOrCreateSuitableArchetype(const FMas
 
 FMassArchetypeHandle FMassEntityManager::CreateArchetype(const FMassArchetypeCompositionDescriptor& Composition, const FMassArchetypeCreationParams& CreationParams)
 {
+	LLM_SCOPE_BYNAME(TEXT("Mass/EntityManager"));
 	const uint32 TypeHash = HashCombine(Composition.CalculateHash(), GetTypeHash(UE::Mass::FArchetypeGroups()));
 
 	TArray<TSharedPtr<FMassArchetypeData>>& HashRow = FragmentHashToArchetypeMap.FindOrAdd(TypeHash);
@@ -526,6 +533,7 @@ FMassArchetypeHandle FMassEntityManager::InternalCreateSimilarArchetype(const TS
 
 FMassArchetypeHandle FMassEntityManager::InternalCreateSimilarArchetype(const FMassArchetypeData& SourceArchetypeRef, FMassArchetypeCompositionDescriptor&& NewComposition, const UE::Mass::FArchetypeGroups& Groups)
 {
+	LLM_SCOPE_BYNAME(TEXT("Mass/EntityManager"));
 	// we require Groups to be already shrunk. Shrinking is required to remove any trailing, invalid group IDs that would
 	// be there if IDs were added and removed to this specific Groups container instance
 	checkf(Groups.IsShrunk(), TEXT("A group container with invalid trailing IDs has been passed to archetype creation - this is not expected and will cause issues. Make sure to Shrink your Groups before passing to %hs"), __FUNCTION__);
@@ -1045,6 +1053,7 @@ UE::Mass::FArchetypeGroupHandle FMassEntityManager::GetGroupForEntity(FMassEntit
 
 UE::Mass::FArchetypeGroupType FMassEntityManager::FinOrAddArchetypeGroupType(const FName GroupName)
 {
+	LLM_SCOPE_BYNAME(TEXT("Mass/EntityManager"));
 	const int32* FoundGroupIndex = GroupNameToTypeIndex.Find(GroupName);
 	if (LIKELY(FoundGroupIndex))
 	{
@@ -2081,6 +2090,7 @@ void FMassEntityManager::AppendCommands(TSharedPtr<FMassCommandBuffer>& InOutCom
 	{
 		return;
 	}
+	LLM_SCOPE_BYNAME(TEXT("Mass/EntityManager"));
 	Defer().MoveAppend(*InOutCommandBuffer.Get());
 }
 
@@ -2102,6 +2112,7 @@ TSharedRef<FMassEntityManager::FEntityCreationContext> FMassEntityManager::GetOr
 TSharedRef<FMassEntityManager::FEntityCreationContext> FMassEntityManager::GetOrMakeCreationContext(TConstArrayView<FMassEntityHandle> ReservedEntities
 	, FMassArchetypeEntityCollection&& EntityCollection)
 {
+	LLM_SCOPE_BYNAME(TEXT("Mass/EntityManager"));
 	if (ActiveCreationContext.IsValid())
 	{
 		TSharedPtr<FEntityCreationContext> SharedContext = ActiveCreationContext.Pin();
diff --git a/Engine/Source/Runtime/MassEntity/Private/MassEntityManagerStorage.cpp b/Engine/Source/Runtime/MassEntity/Private/MassEntityManagerStorage.cpp
--- a/Engine/Source/Runtime/MassEntity/Private/MassEntityManagerStorage.cpp
+++ b/Engine/Source/Runtime/MassEntity/Private/MassEntityManagerStorage.cpp
@@ -111,6 +111,7 @@ namespace UE::Mass
 
 	FMassEntityHandle FSingleThreadedEntityStorage::AcquireOne()
 	{
+		LLM_SCOPE_BYNAME(TEXT("Mass/SingleThreadedStorage"));
 		const int32 SerialNumber = SerialNumberGenerator.fetch_add(1);
 		const int32 Index = (EntityFreeIndexList.Num() > 0) ? EntityFreeIndexList.Pop(EAllowShrinking::No) : Entities.Add();
 		Entities[Index].SerialNumber = SerialNumber;
@@ -123,6 +124,7 @@ namespace UE::Mass
 
 	int32 FSingleThreadedEntityStorage::Acquire(TArrayView<FMassEntityHandle> OutEntityHandles)
 	{
+		LLM_SCOPE_BYNAME(TEXT("Mass/SingleThreadedStorage"));
 		const int32 NumToAdd = OutEntityHandles.Num();
 
 		const int32 SerialNumber = SerialNumberGenerator.fetch_add(1);
@@ -162,6 +164,7 @@ namespace UE::Mass
 
 	int32 FSingleThreadedEntityStorage::Release(TConstArrayView<FMassEntityHandle> Handles)
 	{
+		LLM_SCOPE_BYNAME(TEXT("Mass/SingleThreadedStorage"));
 		int DeallocateCount = 0;
 
 		EntityFreeIndexList.Reserve(EntityFreeIndexList.Num() + Handles.Num());
@@ -187,6 +190,7 @@ namespace UE::Mass
 
 	int32 FSingleThreadedEntityStorage::ForceRelease(TConstArrayView<FMassEntityHandle> Handles)
 	{
+		LLM_SCOPE_BYNAME(TEXT("Mass/SingleThreadedStorage"));
 		EntityFreeIndexList.Reserve(EntityFreeIndexList.Num() + Handles.Num());
 		for (const FMassEntityHandle& Handle : Handles)
 		{
@@ -235,6 +239,7 @@ namespace UE::Mass
 
 	void FConcurrentEntityStorage::Initialize(const FMassEntityManager_InitParams_Concurrent& InInitializationParams)
 	{
+		LLM_SCOPE_BYNAME(TEXT("Mass/ConcurrentStorage"));
 		// Compute number of pages required
 		check(FMath::IsPowerOfTwo(InInitializationParams.MaxEntitiesPerPage));
 		check(FMath::IsPowerOfTwo(InInitializationParams.MaxEntityCount));
@@ -357,6 +362,7 @@ namespace UE::Mass
 
 	bool FConcurrentEntityStorage::AddPage()
 	{
+		LLM_SCOPE_BYNAME(TEXT("Mass/ConcurrentStorage"));
 		check(FreeListMutex.IsLocked());
 		UE::TUniqueLock PageAllocateLock(PageAllocateMutex);
 
@@ -489,6 +495,7 @@ namespace UE::Mass
 
 	int32 FConcurrentEntityStorage::Release(TConstArrayView<FMassEntityHandle> Handles)
 	{
+		LLM_SCOPE_BYNAME(TEXT("Mass/ConcurrentStorage"));
 		int32 DeallocateCount = 0;
 	
 		int32 BeginHandlesIndexToFree = 0;
@@ -552,6 +559,7 @@ namespace UE::Mass
 
 	int32 FConcurrentEntityStorage::ForceRelease(TConstArrayView<FMassEntityHandle> Handles)
 	{
+		LLM_SCOPE_BYNAME(TEXT("Mass/ConcurrentStorage"));
 		// ForceRelease assumes the caller knows all handles are allocated
 		// no need to have complexity of tracking "runs" of handles 
 		for (const FMassEntityHandle& Handle : Handles)
