[Mass] Added a warning and a temporary workaround for "grouped" entity queries run in parallel mode - calling `FMassEntityQuery::ParallelForEachEntityChunk` will emit a warning and call `FMassEntityQuery::ForEachEntityChunk` instead.

#rnx

diff --git a/Engine/Source/Runtime/MassEntity/Private/MassEntityQuery.cpp b/Engine/Source/Runtime/MassEntity/Private/MassEntityQuery.cpp
--- a/Engine/Source/Runtime/MassEntity/Private/MassEntityQuery.cpp
+++ b/Engine/Source/Runtime/MassEntity/Private/MassEntityQuery.cpp
@@ -189,7 +189,7 @@ void FMassEntityQuery::CacheArchetypes()
 				}
 			}
 
-			if (GroupSortingSteps.Num())
+			if (IsGrouping())
 			{
 				SortArchetypes(FirstNewArchetype);
 			}
@@ -334,6 +334,12 @@ void FMassEntityQuery::ParallelForEachEntityChunk(FMassExecutionContext& Executi
 		ForEachEntityChunk(ExecutionContext, ExecuteFunction);
 		return;
 	}
+	else if (IsGrouping())
+	{
+		UE_VLOG_UELOG(CachedEntityManager->GetOwner(), LogMass, Warning, TEXT("Calling %hs is not supported for grouping queries yet. Calling regular ForEachEntityChunk instead."), __FUNCTION__);
+		ForEachEntityChunk(ExecutionContext, ExecuteFunction);
+		return;
+	}
 
 	checkf(CachedEntityManager == ExecutionContext.GetSharedEntityManager() && CachedEntityManager.IsValid()
 		, TEXT("FMassEntityQuery needs to be initialized with a valid EntityManager and the execution context has to come from the same manager"));
diff --git a/Engine/Source/Runtime/MassEntity/Public/MassEntityQuery.h b/Engine/Source/Runtime/MassEntity/Public/MassEntityQuery.h
--- a/Engine/Source/Runtime/MassEntity/Public/MassEntityQuery.h
+++ b/Engine/Source/Runtime/MassEntity/Public/MassEntityQuery.h
@@ -186,6 +186,12 @@ public:
 	void GroupBy(UE::Mass::FArchetypeGroupType GroupType, const TFunction<bool(const UE::Mass::FArchetypeGroupID, const UE::Mass::FArchetypeGroupID)>& Predicate);
 	void ResetGrouping();
 
+	/** 
+	 * @return whether the query is configured to use archetype group information to group and sort
+	 *		archetypes to be processed.
+	 */
+	 bool IsGrouping() const;
+
 #if WITH_ARCHETYPE_MATCH_OVERRIDE
 	template<typename T> requires TArchetypeMatchOverrideConcept<T>
 	void SetArchetypeMatchOverride(const T& Override);
@@ -346,5 +352,10 @@ inline void FMassEntityQuery::SetChunkFilter(const FMassChunkConditionFunction&
 	ChunkCondition = InFunction;
 }
 
+inline bool FMassEntityQuery::IsGrouping() const
+{
+	return !GroupSortingSteps.IsEmpty();
+}
+
 #undef MASS_ACHETYPE_OVERRIDE_MAX_SIZE
 #undef MASS_ACHETYPE_OVERRIDE_MAX_ALIGNMENT
\ No newline at end of file
