Remove unused `GetTypeHash` method.

diff --git a/Engine/Source/Runtime/MassEntity/Public/TypeBitSetBuilder.h b/Engine/Source/Runtime/MassEntity/Public/TypeBitSetBuilder.h
--- a/Engine/Source/Runtime/MassEntity/Public/TypeBitSetBuilder.h
+++ b/Engine/Source/Runtime/MassEntity/Public/TypeBitSetBuilder.h
@@ -584,20 +584,6 @@ protected:
 	}
 #endif // WITH_STRUCTUTILS_DEBUG
 
-public:
-	/**
-	 * Hash function for the builder.
-	 * Combines the hash of the struct tracker and the bitset.
-	 * @param Instance the builder instance to hash.
-	 * @return The computed hash value.
-	 */
-	FORCEINLINE friend uint32 GetTypeHash(const TTypeBitSetBuilder& Instance)
-	{
-		const uint32 StructTrackerHash = GetTypeHash(Instance.StructTracker);
-		const uint32 BitArrayHash = GetTypeHash(Instance.StructTypesBitArray);
-		return HashCombine(StructTrackerHash, BitArrayHash);
-	}
-
 protected:
 	/**
 	 * Retrieves the struct tracker used by this builder.
