Raise the priority of Mass game thread processor tasks, they would often get stuck behind normal ticks which would delay the ability to kick off worker thread tasks

diff --git a/Engine/Source/Runtime/MassEntity/Private/MassProcessor.cpp b/Engine/Source/Runtime/MassEntity/Private/MassProcessor.cpp
--- a/Engine/Source/Runtime/MassEntity/Private/MassProcessor.cpp
+++ b/Engine/Source/Runtime/MassEntity/Private/MassProcessor.cpp
@@ -116,7 +116,8 @@ public:
 
 	static ENamedThreads::Type GetDesiredThread()
 	{
-		return ENamedThreads::GameThread;
+		// Use a high priority task so processor chains that touch the game thread will take priority over normal ticks
+		return ENamedThreads::SetTaskPriority(ENamedThreads::GameThread, ENamedThreads::HighTaskPriority);
 	}
 };
 
