[Mass] Tiny change to ensure we never have the same observer listed twice in a pipeline for a given observed operation (for given operation-type/element pair).

diff --git a/Engine/Source/Runtime/MassEntity/Private/MassObserverManager.cpp b/Engine/Source/Runtime/MassEntity/Private/MassObserverManager.cpp
--- a/Engine/Source/Runtime/MassEntity/Private/MassObserverManager.cpp
+++ b/Engine/Source/Runtime/MassEntity/Private/MassObserverManager.cpp
@@ -342,12 +342,16 @@ void FMassObserverManager::AddObserverInstance(const UScriptStruct& FragmentOrTa
 		Pipeline = &(*TagObservers[(uint8)Operation]).FindOrAdd(&FragmentOrTagType);
 		ObservedTags[(uint8)Operation].Add(FragmentOrTagType);
 	}
-	Pipeline->AppendProcessor(ObserverProcessor);
 
-	// calling initialize to ensure the given processor is related to the same EntityManager
-	if (UObject* Owner = EntityManager.GetOwner())
-	{	
-		ObserverProcessor.CallInitialize(Owner, EntityManager.AsShared());
+	// AppendUniqueProcessor will return true only if ObserverProcessor has not been a part of Pipeline yet.
+	// Otherwise, we don't need to CallInitialize.
+	if (Pipeline->AppendUniqueProcessor(ObserverProcessor))
+	{
+		// calling initialize to ensure the given processor is related to the same EntityManager
+		if (UObject* Owner = EntityManager.GetOwner())
+		{	
+			ObserverProcessor.CallInitialize(Owner, EntityManager.AsShared());
+		}
 	}
 }
 
diff --git a/Engine/Source/Runtime/MassEntity/Private/MassProcessingTypes.cpp b/Engine/Source/Runtime/MassEntity/Private/MassProcessingTypes.cpp
--- a/Engine/Source/Runtime/MassEntity/Private/MassProcessingTypes.cpp
+++ b/Engine/Source/Runtime/MassEntity/Private/MassProcessingTypes.cpp
@@ -241,9 +241,11 @@ void FMassRuntimePipeline::AppendProcessors(TArrayView<TObjectPtr<UMassProcessor
 	Processors.Append(InProcessors);
 }
 
-void FMassRuntimePipeline::AppendUniqueProcessor(UMassProcessor& Processor)
+bool FMassRuntimePipeline::AppendUniqueProcessor(UMassProcessor& Processor)
 {
+	const int32 PreviousCount = Processors.Num();
 	Processors.AddUnique(&Processor);
+	return PreviousCount != Processors.Num();
 }
 
 void FMassRuntimePipeline::AppendProcessor(TSubclassOf<UMassProcessor> ProcessorClass, UObject& InOwner)
diff --git a/Engine/Source/Runtime/MassEntity/Public/MassProcessingTypes.h b/Engine/Source/Runtime/MassEntity/Public/MassProcessingTypes.h
--- a/Engine/Source/Runtime/MassEntity/Public/MassProcessingTypes.h
+++ b/Engine/Source/Runtime/MassEntity/Public/MassProcessingTypes.h
@@ -121,7 +121,8 @@ public:
 	void AppendProcessor(UMassProcessor& InProcessor);
 	void AppendProcessors(TArrayView<TObjectPtr<UMassProcessor>> InProcessors);
 
-	void AppendUniqueProcessor(UMassProcessor& Processor);
+	/** @return true if the given processor has indeed been added, i.e. will return false if Processor was already part of the pipeline. */
+	bool AppendUniqueProcessor(UMassProcessor& Processor);
 
 	/** Creates an instance of ProcessorClass and adds it to Processors without any additional checks */
 	void AppendProcessor(TSubclassOf<UMassProcessor> ProcessorClass, UObject& InOwner);
