[Mass]
Added Mass trace channel for MassInsights plugin trace logging
First pass Mass trace markup

diff --git a/Engine/Source/Runtime/MassEntity/Internal/MassEntityTrace.h b/Engine/Source/Runtime/MassEntity/Internal/MassEntityTrace.h
new file mode 100644
--- /dev/null
+++ b/Engine/Source/Runtime/MassEntity/Internal/MassEntityTrace.h
@@ -0,0 +1,136 @@
+// Copyright Epic Games, Inc. All Rights Reserved.
+
+#pragma once
+
+#include "MassProcessingTypes.h"
+#include "Misc/NotNull.h"
+#include "ProfilingDebugging/TraceAuxiliary.h"
+#include "Trace/Config.h"
+#include "Trace/Trace.h"
+
+#if UE_TRACE_ENABLED && WITH_MASSENTITY_DEBUG && !IS_PROGRAM && !UE_BUILD_SHIPPING
+#define UE_MASS_TRACE_ENABLED 1
+#else
+#define UE_MASS_TRACE_ENABLED 0
+#endif
+
+#if UE_MASS_TRACE_ENABLED
+
+struct FMassEntityHandle;
+struct FMassEntityManager;
+struct FMassEntityQuery;
+struct FMassArchetypeData;
+struct FMassArchetypeHandle;
+struct FMassArchetypeCompositionDescriptor;
+class UMassProcessor;	
+
+UE_TRACE_CHANNEL_EXTERN(MassChannel, MASSENTITY_API);
+
+class FMassTrace : public FNoncopyable
+{
+public:
+	struct FScopedQueryForEachTrace
+	{
+		FScopedQueryForEachTrace(const FMassEntityQuery* InQuery);
+		~FScopedQueryForEachTrace();
+
+		void ReportArchetype(const FMassArchetypeData& Archetype);
+		const FMassEntityQuery* Query = nullptr;
+		int32 ArchetypeCount = 0;
+		int32 ChunkCount = 0;
+		int32 EntityCount = 0;
+	};
+
+	MASSENTITY_API static uint64_t OutputBeginPhaseWithID(const TCHAR* PhaseName);
+	MASSENTITY_API static void OutputBeginPhaseRegion(const TCHAR* PhaseName);
+	MASSENTITY_API static void OutputEndPhaseRegion(const TCHAR* PhaseName);
+	MASSENTITY_API static void OutputEndPhaseRegion(uint64 PhaseId);
+
+	MASSENTITY_API static void OnPhaseBegin(uint64 PhaseId);
+	MASSENTITY_API static void OnPhaseEnd(uint64 PhaseId);
+	
+	MASSENTITY_API static uint64 OutputRegisterArchetype(uint64 ArchetypeID, const FMassArchetypeCompositionDescriptor& CompositionDescriptor);
+	MASSENTITY_API static uint64 RegisterArchetype(const FMassArchetypeHandle& ArchetypeHandle);
+	MASSENTITY_API static uint64 RegisterArchetype(const FMassArchetypeData& Data);
+
+	MASSENTITY_API static void OutputRegisterFragment(const UScriptStruct* Struct);
+	MASSENTITY_API static void RegisterFragment(const UScriptStruct* Struct);
+
+	MASSENTITY_API static void PopulateEntities(const FMassEntityManager& EntityManager);
+	MASSENTITY_API static void EntityCreated(FMassEntityHandle Entity, const FMassArchetypeData& Archetype);
+	MASSENTITY_API static void EntityMoved(FMassEntityHandle Entity, const FMassArchetypeData& NewArchetype);
+	MASSENTITY_API static void EntityDestroyed(FMassEntityHandle Entity);
+	MASSENTITY_API static void EntitiesDestroyed(TConstArrayView<FMassEntityHandle> Entities);
+
+	MASSENTITY_API static void QueryCreated(const FMassEntityQuery* Query);
+	MASSENTITY_API static void QueryDestroyed(const FMassEntityQuery* Query);
+	MASSENTITY_API static void QueryRegisteredToProcessor(const FMassEntityQuery* Query, TNotNull<const UMassProcessor*> Processor);
+	MASSENTITY_API static void QueryArchetypeAdded(const FMassEntityQuery* Query, const FMassArchetypeHandle& Archetype);
+};
+
+#define UE_TRACE_MASS_PHASE_BEGIN(PhaseID) \
+	FMassTrace::OnPhaseBegin(PhaseID);
+
+#define UE_TRACE_MASS_PHASE_END(PhaseID) \
+	FMassTrace::OnPhaseEnd(PhaseID);
+
+#define UE_TRACE_MASS_ARCHETYPE_CREATED(Archetype) \
+	FMassTrace::RegisterArchetype(Archetype);
+
+#define UE_TRACE_MASS_ENTITY_CREATED(Entity, Archetype) \
+	FMassTrace::EntityCreated(Entity, Archetype);
+
+#define UE_TRACE_MASS_ENTITIES_CREATED(EntityHandles, Archetype) \
+	for (const FMassEntityHandle& Entity : EntityHandles) \
+	{ \
+		FMassTrace::EntityCreated(Entity, Archetype); \
+	}
+
+#define UE_TRACE_MASS_ENTITY_MOVED(Entity, NewArchetype) \
+	FMassTrace::EntityMoved(Entity, NewArchetype);
+
+#define UE_TRACE_MASS_ENTITY_DESTROYED(Entity) \
+	FMassTrace::EntityDestroyed(Entity);
+
+#define UE_TRACE_MASS_ENTITIES_DESTROYED(EntityHandles) \
+	for (const FMassEntityHandle& Entity : EntityHandles) \
+	{ \
+		FMassTrace::EntityDestroyed(Entity); \
+	}
+
+#define UE_TRACE_MASS_QUERY_CREATED() \
+	FMassTrace::QueryCreated(this);
+
+#define UE_TRACE_MASS_QUERY_DESTROYED() \
+	FMassTrace::QueryDestroyed(this);
+
+#define UE_TRACE_MASS_QUERY_REGISTERED_TO_PROCESSOR(Processor) \
+	FMassTrace::QueryRegisteredToProcessor(this, Processor);
+
+#define UE_TRACE_MASS_QUERY_ARCHETYPE_ADDED(Archetype) \
+	FMassTrace::QueryArchetypeAdded(this, Archetype);
+
+#define UE_TRACE_SCOPED_MASS_QUERY_FOR_EACH() \
+	FMassTrace::FScopedQueryForEachTrace _ScopedQueryForEachTrace(this);
+
+#define UE_TRACE_SCOPED_MASS_QUERY_FOR_EACH_REPORT_ARCHETYPE(Archetype) \
+	_ScopedQueryForEachTrace.ReportArchetype(Archetype);
+
+#else //UE_MASS_TRACE_ENABLED
+
+#define UE_TRACE_MASS_ARCHETYPE_CREATED(PhaseID)
+#define UE_TRACE_MASS_ENTITY_CREATED(Entity, Archetype)
+#define UE_TRACE_MASS_ENTITIES_CREATED(EntityHandles, Archetype)
+#define UE_TRACE_MASS_ENTITY_MOVED(Entity, NewArchetype)
+#define UE_TRACE_MASS_ENTITY_DESTROYED(Entity)
+#define UE_TRACE_MASS_ENTITIES_DESTROYED(EntityHandles)
+#define UE_TRACE_MASS_PHASE_BEGIN(PhaseID)
+#define UE_TRACE_MASS_PHASE_END(PhaseID)
+#define UE_TRACE_MASS_QUERY_CREATED()
+#define UE_TRACE_MASS_QUERY_DESTROYED()
+#define UE_TRACE_MASS_QUERY_REGISTERED_TO_PROCESSOR(Processor)
+#define UE_TRACE_MASS_QUERY_ARCHETYPE_ADDED(Archetype)
+#define UE_TRACE_SCOPED_MASS_QUERY_FOR_EACH()
+#define UE_TRACE_SCOPED_MASS_QUERY_FOR_EACH_REPORT_ARCHETYPE(Archetype)
+
+#endif //UE_MASS_TRACE_ENABLED
\ No newline at end of file
diff --git a/Engine/Source/Runtime/MassEntity/MassEntity.Build.cs b/Engine/Source/Runtime/MassEntity/MassEntity.Build.cs
--- a/Engine/Source/Runtime/MassEntity/MassEntity.Build.cs
+++ b/Engine/Source/Runtime/MassEntity/MassEntity.Build.cs
@@ -16,6 +16,7 @@ public class MassEntity : ModuleRules
 				"CoreUObject",
 				"Engine",
 				"DeveloperSettings",
+				"TraceLog",
 			}
 		);
 
diff --git a/Engine/Source/Runtime/MassEntity/Private/MassArchetypeData.cpp b/Engine/Source/Runtime/MassEntity/Private/MassArchetypeData.cpp
--- a/Engine/Source/Runtime/MassEntity/Private/MassArchetypeData.cpp
+++ b/Engine/Source/Runtime/MassEntity/Private/MassArchetypeData.cpp
@@ -462,6 +462,9 @@ void FMassArchetypeData::MoveEntityToAnotherArchetype(const FMassEntityHandle En
 	const int32 NewIndexWithinChunk = NewAbsoluteIndex % NewArchetype.NumEntitiesPerChunk;
 	FMassArchetypeChunk& NewChunk = NewArchetype.Chunks[NewChunkIndex];
 
+	UE_TRACE_MASS_ARCHETYPE_CREATED(NewArchetype)
+	UE_TRACE_MASS_ENTITY_MOVED(Entity, NewArchetype)
+
 	MoveFragmentsToAnotherArchetypeInternal(NewArchetype, { NewChunk.GetRawMemory(), NewIndexWithinChunk }, { Chunk.GetRawMemory(), IndexWithinChunk }, /*Count=*/1);
 
 	RemoveEntityInternal(AbsoluteIndex);
diff --git a/Engine/Source/Runtime/MassEntity/Private/MassDebugger.cpp b/Engine/Source/Runtime/MassEntity/Private/MassDebugger.cpp
--- a/Engine/Source/Runtime/MassEntity/Private/MassDebugger.cpp
+++ b/Engine/Source/Runtime/MassEntity/Private/MassDebugger.cpp
@@ -2,6 +2,7 @@
 
 #include "MassDebugger.h"
 #if WITH_MASSENTITY_DEBUG
+#include "Algo/ForEach.h"
 #include "MassProcessor.h"
 #include "MassEntityManager.h"
 #include "MassEntityManagerStorage.h"
@@ -15,6 +16,7 @@
 #include "Engine/Engine.h"
 #include "MassEntityUtils.h"
 #include "MassCommandBuffer.h"
+#include "MassEntityTrace.h"
 
 
 namespace UE::Mass::Debug
@@ -522,6 +524,17 @@ TArray<FMassEntityHandle> FMassDebugger::GetEntitiesMatchingQuery(const FMassEnt
 	return Entities;
 }
 
+void FMassDebugger::ForEachArchetype(const FMassEntityManager& EntityManager, const UE::Mass::Debug::FArchetypeFunction& Function)
+{
+	for (auto& KVP : EntityManager.FragmentHashToArchetypeMap)
+	{
+		for (const TSharedPtr<FMassArchetypeData>& Archetype : KVP.Value)
+		{
+			Function(FMassArchetypeHelper::ArchetypeHandleFromData(Archetype));
+		}
+	}
+}
+
 TArray<FMassArchetypeHandle> FMassDebugger::GetAllArchetypes(const FMassEntityManager& EntityManager)
 {
 	TArray<FMassArchetypeHandle> Archetypes;
@@ -543,6 +556,38 @@ const FMassArchetypeCompositionDescriptor& FMassDebugger::GetArchetypeCompositio
 	return ArchetypeData.CompositionDescriptor;
 }
 
+uint64 FMassDebugger::GetArchetypeTraceID(const FMassArchetypeData& ArchetypeData)
+{
+	return reinterpret_cast<uint64>(&ArchetypeData);
+}
+
+uint64 FMassDebugger::GetArchetypeTraceID(const FMassArchetypeHandle& ArchetypeHandle)
+{ 
+	const FMassArchetypeData& ArchetypeData = FMassArchetypeHelper::ArchetypeDataFromHandleChecked(ArchetypeHandle);
+	return GetArchetypeTraceID(ArchetypeData);
+}
+
+TConstArrayView<FMassEntityHandle> FMassDebugger::GetEntitiesViewOfArchetype(const FMassArchetypeData& ArchetypeData, const FMassArchetypeChunk& Chunk)
+{
+	FMassArchetypeChunk& MutableChunk = const_cast<FMassArchetypeChunk&>(Chunk);
+	TConstArrayView<FMassEntityHandle> View(&MutableChunk.GetEntityArrayElementRef(ArchetypeData.EntityListOffsetWithinChunk, 0), Chunk.GetNumInstances());
+	return View;
+}
+
+const FMassArchetypeData* FMassDebugger::GetArchetypeData(const FMassArchetypeHandle& ArchetypeHandle)
+{
+	FMassArchetypeData* ArchetypeData = FMassArchetypeHelper::ArchetypeDataFromHandle(ArchetypeHandle);
+	return ArchetypeData;
+}
+
+void FMassDebugger::EnumerateChunks(const FMassArchetypeData& Archetype, TFunctionRef<void(const FMassArchetypeChunk&)> Fn)
+{
+	for (const FMassArchetypeChunk& Chunk : Archetype.Chunks)
+	{
+		Fn(Chunk);
+	}
+}
+
 void FMassDebugger::GetArchetypeEntityStats(const FMassArchetypeHandle& ArchetypeHandle, UE::Mass::Debug::FArchetypeStats& OutStats)
 {
 	const FMassArchetypeData& ArchetypeData = FMassArchetypeHelper::ArchetypeDataFromHandleChecked(ArchetypeHandle);
@@ -1162,6 +1207,37 @@ FMassDebugger::FEnvironment& FMassDebugger::GetActiveEnvironment(const FMassEnti
 	return ActiveEnvironments[Index];
 }
 
+FMassDebugger::FEnvironment::FEnvironment(const FMassEntityManager& InEntityManager)
+	: EntityManager(InEntityManager.AsWeak())
+{
+#if UE_MASS_TRACE_ENABLED
+	TraceStartedDelegateHandle = FTraceAuxiliary::OnTraceStarted.AddLambda([this](FTraceAuxiliary::EConnectionType TraceType, const FString& TraceDestination)
+	{
+		if (!EntityManager.IsValid())
+		{
+			return;
+		}
+
+		const FMassEntityManager& Manager = *EntityManager.Pin();
+
+		ForEachArchetype(Manager, [](FMassArchetypeHandle ArchetypeHandle)
+		{
+			UE_TRACE_MASS_ARCHETYPE_CREATED(ArchetypeHandle)
+		});
+
+		FMassTrace::PopulateEntities(Manager);
+
+	});
+#endif
+}
+
+FMassDebugger::FEnvironment::~FEnvironment()
+{
+#if UE_MASS_TRACE_ENABLED
+	FTraceAuxiliary::OnTraceStarted.Remove(TraceStartedDelegateHandle);
+#endif
+}
+
 void FMassDebugger::FEnvironment::ClearBreakpoints()
 {
 	ProcessorBreakpoints.Reset();
diff --git a/Engine/Source/Runtime/MassEntity/Private/MassEntityManager.cpp b/Engine/Source/Runtime/MassEntity/Private/MassEntityManager.cpp
--- a/Engine/Source/Runtime/MassEntity/Private/MassEntityManager.cpp
+++ b/Engine/Source/Runtime/MassEntity/Private/MassEntityManager.cpp
@@ -492,6 +492,7 @@ FMassArchetypeHandle FMassEntityManager::CreateArchetype(const FMassArchetypeCom
 		}
 
 		OnNewArchetypeEvent.Broadcast(FMassArchetypeHandle(ArchetypeDataPtr));
+		UE_TRACE_MASS_ARCHETYPE_CREATED(ArchetypeDataPtr)
 	}
 
 	return FMassArchetypeHelper::ArchetypeHandleFromData(ArchetypeDataPtr);
@@ -573,6 +574,7 @@ FMassArchetypeHandle FMassEntityManager::InternalCreateSimilarArchetype(const FM
 		}
 
 		OnNewArchetypeEvent.Broadcast(FMassArchetypeHandle(ArchetypeDataPtr));
+		UE_TRACE_MASS_ARCHETYPE_CREATED(ArchetypeDataPtr)
 	}
 
 	return FMassArchetypeHelper::ArchetypeHandleFromData(ArchetypeDataPtr);
@@ -817,6 +819,7 @@ TSharedRef<FMassEntityManager::FEntityCreationContext> FMassEntityManager::Batch
 	}
 
 	TargetArchetypeHandle.DataPtr->BatchAddEntities(EntityHandles, SharedFragmentValues, TargetArchetypeEntityRanges);
+	UE_TRACE_MASS_ENTITIES_CREATED(EntityHandles, *TargetArchetypeHandle.DataPtr.Get())
 
 	if (EncodedEntitiesWithPayload.GetPayload().IsEmpty() == false)
 	{
@@ -885,6 +888,8 @@ TSharedRef<FMassEntityManager::FEntityCreationContext> FMassEntityManager::Inter
 	FMassArchetypeEntityCollection::FEntityRangeArray TargetArchetypeEntityRanges;
 	ArchetypeData->BatchAddEntities(ReservedEntities, SharedFragmentValues, TargetArchetypeEntityRanges);
 
+	UE_TRACE_MASS_ENTITIES_CREATED(ReservedEntities, *ArchetypeData)
+
 	return GetOrMakeCreationContext(ReservedEntities, FMassArchetypeEntityCollection(ArchetypeHandle, MoveTemp(TargetArchetypeEntityRanges)));
 }
 
@@ -901,6 +906,8 @@ void FMassEntityManager::DestroyEntity(FMassEntityHandle Entity)
 		ObserverManager.OnPreEntityDestroyed(Archetype->GetCompositionDescriptor(), Entity);
 		Archetype->RemoveEntity(Entity);
 	}
+	
+	UE_TRACE_MASS_ENTITY_DESTROYED(Entity)
 
 	InternalReleaseEntity(Entity);
 }
@@ -932,6 +939,8 @@ void FMassEntityManager::BatchDestroyEntities(TConstArrayView<FMassEntityHandle>
 		}
 		// else it's a "reserved" entity so it has not been assigned to an archetype yet, no archetype nor observers to notify
 	}
+	
+	UE_TRACE_MASS_ENTITIES_DESTROYED(InEntities)
 
 	GetEntityStorageInterface().Release(InEntities);
 }
@@ -1236,6 +1245,8 @@ void FMassEntityManager::InternalBuildEntity(FMassEntityHandle Entity, const FMa
 	const TSharedPtr<FMassArchetypeData>& NewArchetype = ArchetypeHandle.DataPtr;
 	GetEntityStorageInterface().SetArchetypeFromShared(Entity.Index, ArchetypeHandle.DataPtr);
 	NewArchetype->AddEntity(Entity, SharedFragmentValues);
+	
+	UE_TRACE_MASS_ENTITY_CREATED(Entity, *NewArchetype)
 
 	if (IsAllowedToTriggerObservers())
 	{
diff --git a/Engine/Source/Runtime/MassEntity/Private/MassEntityQuery.cpp b/Engine/Source/Runtime/MassEntity/Private/MassEntityQuery.cpp
--- a/Engine/Source/Runtime/MassEntity/Private/MassEntityQuery.cpp
+++ b/Engine/Source/Runtime/MassEntity/Private/MassEntityQuery.cpp
@@ -10,6 +10,7 @@
 #include "Async/ParallelFor.h"
 #include "Containers/UnrealString.h"
 #include "MassProcessor.h"
+#include "MassEntityTrace.h"
 
 #include UE_INLINE_GENERATED_CPP_BY_NAME(MassEntityQuery)
 
@@ -71,17 +72,20 @@ struct FScopedEntityQueryContext
 //-----------------------------------------------------------------------------
 FMassEntityQuery::FMassEntityQuery(UMassProcessor& Owner)
 {
+	UE_TRACE_MASS_QUERY_CREATED()
 	RegisterWithProcessor(Owner);
 }
 
 FMassEntityQuery::FMassEntityQuery(const TSharedPtr<FMassEntityManager>& EntityManager)
 	: FMassFragmentRequirements(EntityManager)
 {
+	UE_TRACE_MASS_QUERY_CREATED()
 }
 
 FMassEntityQuery::FMassEntityQuery(const TSharedRef<FMassEntityManager>& EntityManager, std::initializer_list<UScriptStruct*> InitList)
 	: FMassFragmentRequirements(EntityManager)
 {
+	UE_TRACE_MASS_QUERY_CREATED()
 	for (const UScriptStruct* FragmentType : InitList)
 	{
 		AddRequirement(FragmentType, EMassFragmentAccess::ReadWrite, EMassFragmentPresence::All);
@@ -91,6 +95,7 @@ FMassEntityQuery::FMassEntityQuery(const TSharedRef<FMassEntityManager>& EntityM
 FMassEntityQuery::FMassEntityQuery(const TSharedRef<FMassEntityManager>& EntityManager, TConstArrayView<const UScriptStruct*> InitList)
 	: FMassFragmentRequirements(EntityManager)
 {
+	UE_TRACE_MASS_QUERY_CREATED()
 	for (const UScriptStruct* FragmentType : InitList)
 	{
 		AddRequirement(FragmentType, EMassFragmentAccess::ReadWrite, EMassFragmentPresence::All);
@@ -99,6 +104,8 @@ FMassEntityQuery::FMassEntityQuery(const TSharedRef<FMassEntityManager>& EntityM
 
 void FMassEntityQuery::RegisterWithProcessor(UMassProcessor& Owner)
 {
+	UE_TRACE_MASS_QUERY_REGISTERED_TO_PROCESSOR(&Owner)
+
 	ExpectedContextType = EMassExecutionContextType::Processor;
 	Owner.RegisterQuery(*this);
 #if WITH_MASSENTITY_DEBUG
@@ -221,6 +228,8 @@ void FMassEntityQuery::ForEachEntityChunk(const FMassArchetypeEntityCollection&
 
 void FMassEntityQuery::ForEachEntityChunk(FMassExecutionContext& ExecutionContext, const FMassExecuteFunction& ExecuteFunction)
 {
+	UE_TRACE_SCOPED_MASS_QUERY_FOR_EACH()
+
 	checkf(CachedEntityManager == ExecutionContext.GetSharedEntityManager() && CachedEntityManager.IsValid()
 		, TEXT("FMassEntityQuery needs to be initialized with a valid EntityManager and the execution context has to come from the same manager"));
 
@@ -281,6 +290,9 @@ void FMassEntityQuery::ForEachEntityChunk(FMassExecutionContext& ExecutionContex
 			ExecutionContext.ApplyFragmentRequirements(*this);
 
 			FMassArchetypeData& ArchetypeData = FMassArchetypeHelper::ArchetypeDataFromHandleChecked(ArchetypeHandle);
+			
+			UE_TRACE_SCOPED_MASS_QUERY_FOR_EACH_REPORT_ARCHETYPE(ArchetypeData)
+			
 			ArchetypeData.ExecuteFunction(ExecutionContext, ExecuteFunction
 				, GetRequirementsMappingForArchetype(ArchetypeHandle)
 				, ExecutionContext.GetEntityCollection().GetRanges()
@@ -297,6 +309,9 @@ void FMassEntityQuery::ForEachEntityChunk(FMassExecutionContext& ExecutionContex
 			{
 				const FMassArchetypeHandle& ArchetypeHandle = ValidArchetypes[ArchetypeIndex];
 				FMassArchetypeData& ArchetypeData = FMassArchetypeHelper::ArchetypeDataFromHandleChecked(ArchetypeHandle);
+				
+				UE_TRACE_SCOPED_MASS_QUERY_FOR_EACH_REPORT_ARCHETYPE(ArchetypeData)
+
 				ArchetypeData.ExecuteFunction(ExecutionContext, ExecuteFunction, ArchetypeFragmentMapping[ArchetypeIndex], ChunkCondition);
 				ExecutionContext.ClearFragmentViews(*this);
 			}
diff --git a/Engine/Source/Runtime/MassEntity/Private/MassEntityTrace.cpp b/Engine/Source/Runtime/MassEntity/Private/MassEntityTrace.cpp
new file mode 100644
--- /dev/null
+++ b/Engine/Source/Runtime/MassEntity/Private/MassEntityTrace.cpp
@@ -0,0 +1,473 @@
+// Copyright Epic Games, Inc. All Rights Reserved.
+
+#include "MassEntityTrace.h"
+
+#if UE_MASS_TRACE_ENABLED
+
+#include "Algo/ForEach.h"
+
+#include "MassArchetypeData.h"
+#include "MassEntityTypes.h"
+#include "MassEntityQuery.h"
+#include "MassEntityUtils.h"
+#include "MassProcessor.h"
+#include "MassDebugger.h"
+
+#include "Trace/Trace.inl"
+#include "TraceFilter.h"
+
+UE_TRACE_CHANNEL_DEFINE(MassChannel);
+
+UE_TRACE_EVENT_BEGIN(MassTrace, MassPhaseBegin)
+	UE_TRACE_EVENT_FIELD(uint64, Cycle)
+	UE_TRACE_EVENT_FIELD(UE::Trace::WideString, PhaseName)
+	UE_TRACE_EVENT_FIELD(uint64, PhaseId)
+UE_TRACE_EVENT_END()
+
+UE_TRACE_EVENT_BEGIN(MassTrace, MassPhaseEnd)
+	UE_TRACE_EVENT_FIELD(uint64, Cycle)
+	UE_TRACE_EVENT_FIELD(UE::Trace::WideString, PhaseName)
+	UE_TRACE_EVENT_FIELD(uint64, PhaseId)
+UE_TRACE_EVENT_END()
+
+UE_TRACE_EVENT_BEGIN(MassTrace, RegisterMassArchetype)
+	UE_TRACE_EVENT_FIELD(uint64, ArchetypeID)
+	UE_TRACE_EVENT_FIELD(uint64[], Fragments)
+UE_TRACE_EVENT_END()
+
+UE_TRACE_EVENT_BEGIN(MassTrace, RegisterMassFragment)
+	UE_TRACE_EVENT_FIELD(uint64, FragmentId)
+	UE_TRACE_EVENT_FIELD(UE::Trace::WideString, FragmentName)
+	UE_TRACE_EVENT_FIELD(uint32, FragmentSize)
+	UE_TRACE_EVENT_FIELD(uint8, FragmentType)
+UE_TRACE_EVENT_END()
+
+UE_TRACE_EVENT_BEGIN(MassTrace, MassPhaseExecutionBegin)
+	UE_TRACE_EVENT_FIELD(uint64, PhaseId)
+	UE_TRACE_EVENT_FIELD(uint64, Cycle)
+UE_TRACE_EVENT_END()
+
+UE_TRACE_EVENT_BEGIN(MassTrace, MassPhaseExecutionEnd)
+	UE_TRACE_EVENT_FIELD(uint64, PhaseId)
+	UE_TRACE_EVENT_FIELD(uint64, Cycle)
+UE_TRACE_EVENT_END()
+
+UE_TRACE_EVENT_BEGIN(MassTrace, MassExecuteChunk)
+	UE_TRACE_EVENT_FIELD(uint64, Cycle)
+	UE_TRACE_EVENT_FIELD(uint64, ChunkId)
+	UE_TRACE_EVENT_FIELD(uint64, QueryID)
+	UE_TRACE_EVENT_FIELD(int32, EntityCount)
+UE_TRACE_EVENT_END()
+
+UE_TRACE_EVENT_BEGIN(MassTrace, MassExecuteChunkEnd)
+	UE_TRACE_EVENT_FIELD(uint64, Cycle)
+	UE_TRACE_EVENT_FIELD(uint64, ChunkId)
+	UE_TRACE_EVENT_FIELD(uint64, QueryID)
+UE_TRACE_EVENT_END()
+
+UE_TRACE_EVENT_BEGIN(MassTrace, MassBulkAddEntity)
+	UE_TRACE_EVENT_FIELD(uint64, Cycle)
+	UE_TRACE_EVENT_FIELD(uint64[], Entities)
+	UE_TRACE_EVENT_FIELD(uint64[], ArchetypeIDs)
+UE_TRACE_EVENT_END()
+
+UE_TRACE_EVENT_BEGIN(MassTrace, MassBulkEntityDestroyed)
+	UE_TRACE_EVENT_FIELD(uint64, Cycle)
+	UE_TRACE_EVENT_FIELD(uint64[], Entities)
+UE_TRACE_EVENT_END()
+
+UE_TRACE_EVENT_BEGIN(MassTrace, MassEntityMoved)
+	UE_TRACE_EVENT_FIELD(uint64, Cycle)
+	UE_TRACE_EVENT_FIELD(uint64, Entity)
+	UE_TRACE_EVENT_FIELD(uint64, NewArchetypeID)
+UE_TRACE_EVENT_END()
+
+UE_TRACE_EVENT_BEGIN(MassTrace, QueryCreated)
+	UE_TRACE_EVENT_FIELD(uint64, Cycle)
+	UE_TRACE_EVENT_FIELD(uint64, QueryID)
+	UE_TRACE_EVENT_FIELD(UE::Trace::WideString, Name)
+UE_TRACE_EVENT_END()
+
+UE_TRACE_EVENT_BEGIN(MassTrace, QueryDestroyed)
+	UE_TRACE_EVENT_FIELD(uint64, Cycle)
+	UE_TRACE_EVENT_FIELD(uint64, QueryID)
+UE_TRACE_EVENT_END()
+
+UE_TRACE_EVENT_BEGIN(MassTrace, QueryRegisteredToProcessor)
+	UE_TRACE_EVENT_FIELD(uint64, QueryID)
+	UE_TRACE_EVENT_FIELD(uint64, ProcessorID)
+	UE_TRACE_EVENT_FIELD(UE::Trace::WideString, ProcessorName)
+UE_TRACE_EVENT_END()
+
+UE_TRACE_EVENT_BEGIN(MassTrace, QueryArchetypeAdded)
+	UE_TRACE_EVENT_FIELD(uint64, QueryID)
+	UE_TRACE_EVENT_FIELD(uint64, ArchetypeID)
+UE_TRACE_EVENT_END()
+
+UE_TRACE_EVENT_BEGIN(MassTrace, QueryForEachStarted)
+	UE_TRACE_EVENT_FIELD(uint64, Cycle)
+	UE_TRACE_EVENT_FIELD(uint64, QueryID)
+UE_TRACE_EVENT_END()
+
+UE_TRACE_EVENT_BEGIN(MassTrace, QueryForEachComplete)
+	UE_TRACE_EVENT_FIELD(uint64, Cycle)
+	UE_TRACE_EVENT_FIELD(uint64, QueryID)
+	UE_TRACE_EVENT_FIELD(int32, ArchetypeCount)
+	UE_TRACE_EVENT_FIELD(int32, ChunkCount)
+	UE_TRACE_EVENT_FIELD(int32, EntityCount)
+UE_TRACE_EVENT_END()
+
+enum class FFragmentType : uint8
+{
+	Unknown = 0,
+	Fragment,
+	Tag,
+	SharedFragment
+};
+
+void FMassTrace::QueryCreated(const FMassEntityQuery* Query)
+{
+	if (UE_TRACE_CHANNELEXPR_IS_ENABLED(MassChannel))
+	{
+		UE_TRACE_LOG(MassTrace, QueryCreated, MassChannel)
+			<< QueryCreated.Cycle(FPlatformTime::Cycles64())
+			<< QueryCreated.QueryID(reinterpret_cast<uint64>(Query));
+	}
+}
+
+void FMassTrace::QueryDestroyed(const FMassEntityQuery* Query)
+{
+	if (UE_TRACE_CHANNELEXPR_IS_ENABLED(MassChannel))
+	{
+		UE_TRACE_LOG(MassTrace, QueryDestroyed, MassChannel)
+			<< QueryDestroyed.Cycle(FPlatformTime::Cycles64())
+			<< QueryDestroyed.QueryID(reinterpret_cast<uint64>(Query));
+	}
+}
+
+void FMassTrace::QueryRegisteredToProcessor(const FMassEntityQuery* Query, TNotNull<const UMassProcessor*> Processor)
+{
+	if (UE_TRACE_CHANNELEXPR_IS_ENABLED(MassChannel))
+	{
+		UE_TRACE_LOG(MassTrace, QueryRegisteredToProcessor, MassChannel)
+			<< QueryRegisteredToProcessor.QueryID(reinterpret_cast<uint64>(Query))
+			<< QueryRegisteredToProcessor.ProcessorID(reinterpret_cast<uint64>((const UMassProcessor*)(Processor)))
+			<< QueryRegisteredToProcessor.ProcessorName(*Processor->GetProcessorName());
+	}
+}
+
+void FMassTrace::QueryArchetypeAdded(const FMassEntityQuery* Query, const FMassArchetypeHandle& Archetype)
+{
+	if (UE_TRACE_CHANNELEXPR_IS_ENABLED(MassChannel))
+	{
+		uint64 ArchetypeID = reinterpret_cast<uint64>(FMassArchetypeHelper::ArchetypeDataFromHandle(Archetype));
+
+		UE_TRACE_LOG(MassTrace, QueryArchetypeAdded, MassChannel)
+			<< QueryArchetypeAdded.QueryID(reinterpret_cast<uint64>(Query))
+			<< QueryArchetypeAdded.ArchetypeID(ArchetypeID);
+	}
+}
+
+FMassTrace::FScopedQueryForEachTrace::FScopedQueryForEachTrace(const FMassEntityQuery* InQuery)
+	: Query(InQuery)
+	, ArchetypeCount(0)
+	, ChunkCount(0)
+	, EntityCount(0)
+{
+	if (UE_TRACE_CHANNELEXPR_IS_ENABLED(MassChannel))
+	{
+		UE_TRACE_LOG(MassTrace, QueryForEachStarted, MassChannel)
+			<< QueryForEachStarted.Cycle(FPlatformTime::Cycles64())
+			<< QueryForEachStarted.QueryID(reinterpret_cast<uint64>(Query));
+	}
+}
+
+FMassTrace::FScopedQueryForEachTrace::~FScopedQueryForEachTrace()
+{
+	if (UE_TRACE_CHANNELEXPR_IS_ENABLED(MassChannel))
+	{
+		UE_TRACE_LOG(MassTrace, QueryForEachComplete, MassChannel)
+			<< QueryForEachComplete.Cycle(FPlatformTime::Cycles64())
+			<< QueryForEachComplete.QueryID(reinterpret_cast<uint64>(Query))
+			<< QueryForEachComplete.ArchetypeCount(ArchetypeCount)
+			<< QueryForEachComplete.ChunkCount(ChunkCount)
+			<< QueryForEachComplete.EntityCount(EntityCount);
+	}
+}
+
+void FMassTrace::FScopedQueryForEachTrace::ReportArchetype(const FMassArchetypeData& Archetype)
+{
+	if (UE_TRACE_CHANNELEXPR_IS_ENABLED(MassChannel))
+	{
+		++ArchetypeCount;
+		ChunkCount += Archetype.GetChunkCount();
+		EntityCount += Archetype.GetNumEntities();
+	}
+}
+
+void FMassTrace::OutputRegisterFragment(const UScriptStruct* Struct)
+{
+	const uint64 FragmentId = reinterpret_cast<uint64>(Struct);
+	const FFragmentType FragmentType = [](const UScriptStruct* Struct)
+		{
+			if (Struct->IsChildOf<FMassFragment>())
+			{
+				return FFragmentType::Fragment;
+			}
+			else if (Struct->IsChildOf<FMassTag>())
+			{
+				return FFragmentType::Tag;
+			}
+			else if (Struct->IsChildOf<FMassSharedFragment>())
+			{
+				return FFragmentType::SharedFragment;
+			}
+			else
+			{
+				return FFragmentType::Unknown;
+			}
+		}(Struct);
+
+	UE_TRACE_LOG(MassTrace, RegisterMassFragment, MassChannel)
+		<< RegisterMassFragment.FragmentId(FragmentId)
+		<< RegisterMassFragment.FragmentName(*Struct->GetName())
+		<< RegisterMassFragment.FragmentSize(Struct->GetStructureSize())
+		<< RegisterMassFragment.FragmentType(static_cast<uint8>(FragmentType));
+}
+
+void FMassTrace::OnPhaseBegin(uint64 PhaseId)
+{
+	if (UE_TRACE_CHANNELEXPR_IS_ENABLED(MassChannel))
+	{
+		const FString EnumName = StaticEnum<EMassProcessingPhase>()->GetNameStringByValue(PhaseId);
+		OutputBeginPhaseRegion(*EnumName);
+	}
+}
+
+void FMassTrace::OnPhaseEnd(uint64 PhaseId)
+{
+	if (UE_TRACE_CHANNELEXPR_IS_ENABLED(MassChannel))
+	{
+		const FString EnumName = StaticEnum<EMassProcessingPhase>()->GetNameStringByValue(PhaseId);
+		OutputBeginPhaseRegion(*EnumName);
+	}
+}
+
+uint64 FMassTrace::OutputRegisterArchetype(uint64 ArchetypeID, const FMassArchetypeCompositionDescriptor& CompositionDescriptor)
+{
+	TArray<uint64> FragmentsScratch;
+	FragmentsScratch.Reserve(CompositionDescriptor.Fragments.CountStoredTypes() + CompositionDescriptor.Tags.CountStoredTypes());
+
+	auto FragmentIterator = CompositionDescriptor.Fragments.GetIndexIterator();
+	while (FragmentIterator)
+	{
+		const UScriptStruct* FragmentStruct = CompositionDescriptor.Fragments.GetTypeAtIndex(*FragmentIterator);
+
+		OutputRegisterFragment(FragmentStruct);
+
+		// TODO Should have utility function to go from UScriptStruct to the fragment ID
+		FragmentsScratch.Add(reinterpret_cast<uint64>(FragmentStruct));
+		++FragmentIterator;
+	}
+
+	auto TagIterator = CompositionDescriptor.Tags.GetIndexIterator();
+	while (TagIterator)
+	{
+		const UScriptStruct* FragmentStruct = CompositionDescriptor.Tags.GetTypeAtIndex(*TagIterator);
+		{
+			OutputRegisterFragment(FragmentStruct);
+		}
+		// TODO Should have utility function to go from UScriptStruct to the fragment ID
+		FragmentsScratch.Add(reinterpret_cast<uint64>(FragmentStruct));
+		++TagIterator;
+	}
+
+	UE_TRACE_LOG(MassTrace, RegisterMassArchetype, MassChannel)
+		<< RegisterMassArchetype.ArchetypeID(ArchetypeID)
+		<< RegisterMassArchetype.Fragments(FragmentsScratch.GetData(), FragmentsScratch.Num());
+
+	return ArchetypeID;
+}
+
+uint64 FMassTrace::RegisterArchetype(const FMassArchetypeHandle& ArchetypeHandle)
+{
+	if (UE_TRACE_CHANNELEXPR_IS_ENABLED(MassChannel))
+	{
+		const FMassArchetypeCompositionDescriptor& CompositionDescriptor = FMassDebugger::GetArchetypeComposition(ArchetypeHandle);
+		const uint64 ArchetypeID = FMassDebugger::GetArchetypeTraceID(ArchetypeHandle);
+
+		return OutputRegisterArchetype(ArchetypeID, CompositionDescriptor);
+	}
+	return 0;
+}
+
+uint64 FMassTrace::RegisterArchetype(const FMassArchetypeData& Data)
+{
+	if (UE_TRACE_CHANNELEXPR_IS_ENABLED(MassChannel))
+	{
+		const FMassArchetypeCompositionDescriptor& CompositionDescriptor = Data.GetCompositionDescriptor();
+		const uint64 ArchetypeID = FMassDebugger::GetArchetypeTraceID(Data);
+
+		return OutputRegisterArchetype(ArchetypeID, CompositionDescriptor);
+	}
+	return 0;
+}
+
+void FMassTrace::RegisterFragment(const UScriptStruct* Struct)
+{
+	if (UE_TRACE_CHANNELEXPR_IS_ENABLED(MassChannel))
+	{
+		OutputRegisterFragment(Struct);
+	}
+}
+
+void FMassTrace::PopulateEntities(const FMassEntityManager& EntityManager)
+{
+	if (UE_TRACE_CHANNELEXPR_IS_ENABLED(MassChannel))
+	{
+		constexpr int32 ScratchSpace = 1024;
+
+		TArray<uint64, TInlineAllocator<1024>> EntityDataScratch;
+		TArray<uint64, TInlineAllocator<1024>> ArchetypeIDsScratch;
+
+		const uint64 Cycle = FPlatformTime::Cycles64();
+
+		FMassDebugger::ForEachArchetype(EntityManager, [Cycle, &EntityDataScratch, &ArchetypeIDsScratch](const FMassArchetypeHandle ArchetypeHandle)
+		{
+			const FMassArchetypeData* ArchetypeData = FMassDebugger::GetArchetypeData(ArchetypeHandle);
+
+			if (ArchetypeData)
+			{
+				uint64 ArchetypeID = FMassDebugger::GetArchetypeTraceID(*ArchetypeData);
+				FMassDebugger::EnumerateChunks(*ArchetypeData, [Cycle, ArchetypeID, ArchetypeData, &EntityDataScratch, &ArchetypeIDsScratch](const FMassArchetypeChunk& Chunk)
+					{
+						TConstArrayView<FMassEntityHandle> EntitiesHandles = FMassDebugger::GetEntitiesViewOfArchetype(*ArchetypeData, Chunk);
+
+						for (FMassEntityHandle EntityHandle : EntitiesHandles)
+						{
+							EntityDataScratch.Add(EntityHandle.AsNumber());
+							ArchetypeIDsScratch.Add(ArchetypeID);
+
+							if (EntityDataScratch.Num() == ScratchSpace)
+							{
+								UE_TRACE_LOG(MassTrace, MassBulkAddEntity, MassChannel)
+									<< MassBulkAddEntity.Cycle(Cycle)
+									<< MassBulkAddEntity.Entities(EntityDataScratch.GetData(), EntityDataScratch.Num())
+									<< MassBulkAddEntity.ArchetypeIDs(ArchetypeIDsScratch.GetData(), ArchetypeIDsScratch.Num());
+								EntityDataScratch.Reset();
+								ArchetypeIDsScratch.Reset();
+							}
+						}
+
+					});
+			}
+		});
+
+		UE_TRACE_LOG(MassTrace, MassBulkAddEntity, MassChannel)
+			<< MassBulkAddEntity.Cycle(Cycle)
+			<< MassBulkAddEntity.Entities(EntityDataScratch.GetData(), EntityDataScratch.Num())
+			<< MassBulkAddEntity.ArchetypeIDs(ArchetypeIDsScratch.GetData(), ArchetypeIDsScratch.Num());
+	}
+}
+
+void FMassTrace::EntityCreated(FMassEntityHandle Entity, const FMassArchetypeData& Archetype)
+{
+	if (UE_TRACE_CHANNELEXPR_IS_ENABLED(MassChannel))
+	{
+		const uint64 Cycle = FPlatformTime::Cycles64();
+		const uint64 EntityAsU64 = Entity.AsNumber();
+		const uint64 ArchetypeID = FMassDebugger::GetArchetypeTraceID(Archetype);
+		UE_TRACE_LOG(MassTrace, MassBulkAddEntity, MassChannel)
+			<< MassBulkAddEntity.Cycle(Cycle)
+			<< MassBulkAddEntity.Entities(&EntityAsU64, 1)
+			<< MassBulkAddEntity.ArchetypeIDs(&ArchetypeID, 1);
+	}
+}
+
+void FMassTrace::EntityMoved(FMassEntityHandle Entity, const FMassArchetypeData& NewArchetype)
+{
+	if (UE_TRACE_CHANNELEXPR_IS_ENABLED(MassChannel))
+	{
+		const uint64 Cycle = FPlatformTime::Cycles64();
+		const uint64 EntityAsU64 = Entity.AsNumber();
+		UE_TRACE_LOG(MassTrace, MassEntityMoved, MassChannel)
+			<< MassEntityMoved.Cycle(Cycle)
+			<< MassEntityMoved.Entity(EntityAsU64)
+			<< MassEntityMoved.NewArchetypeID(FMassDebugger::GetArchetypeTraceID(NewArchetype));
+	}
+}
+
+void FMassTrace::EntityDestroyed(FMassEntityHandle Entity)
+{
+	if (UE_TRACE_CHANNELEXPR_IS_ENABLED(MassChannel))
+	{
+		const uint64 Cycle = FPlatformTime::Cycles64();
+		const uint64 EntityAsU64 = Entity.AsNumber();
+		UE_TRACE_LOG(MassTrace, MassBulkEntityDestroyed, MassChannel)
+			<< MassBulkEntityDestroyed.Cycle(Cycle)
+			<< MassBulkEntityDestroyed.Entities(&EntityAsU64, 1);
+	}
+}
+
+void FMassTrace::EntitiesDestroyed(TConstArrayView<FMassEntityHandle> Entities)
+{
+	if (UE_TRACE_CHANNELEXPR_IS_ENABLED(MassChannel))
+	{
+		const uint64 Cycle = FPlatformTime::Cycles64();
+
+		TConstArrayView<uint64> EntitiesAsU64(
+			reinterpret_cast<const uint64*>(Entities.GetData()),
+			Entities.Num());
+		UE_TRACE_LOG(MassTrace, MassBulkEntityDestroyed, MassChannel)
+			<< MassBulkEntityDestroyed.Cycle(Cycle)
+			<< MassBulkEntityDestroyed.Entities(EntitiesAsU64.GetData(), EntitiesAsU64.Num());
+	}
+}
+
+uint64_t FMassTrace::OutputBeginPhaseWithID(const TCHAR* PhaseName)
+{
+	if (UE_TRACE_CHANNELEXPR_IS_ENABLED(MassChannel))
+	{
+		const uint64 PhaseId = FPlatformTime::Cycles64();
+		UE_TRACE_LOG(MassTrace, MassPhaseBegin, MassChannel)
+			<< MassPhaseBegin.Cycle(FPlatformTime::Cycles64())
+			<< MassPhaseBegin.PhaseName(PhaseName)
+			<< MassPhaseBegin.PhaseId(PhaseId);
+		return PhaseId;
+	}
+	return 0;
+}
+
+void FMassTrace::OutputBeginPhaseRegion(const TCHAR* PhaseName)
+{
+	if (UE_TRACE_CHANNELEXPR_IS_ENABLED(MassChannel))
+	{
+		UE_TRACE_LOG(MassTrace, MassPhaseBegin, MassChannel)
+			<< MassPhaseBegin.Cycle(FPlatformTime::Cycles64())
+			<< MassPhaseBegin.PhaseName(PhaseName)
+			<< MassPhaseBegin.PhaseId(0);
+	}
+}
+
+void FMassTrace::OutputEndPhaseRegion(const TCHAR* PhaseName)
+{
+	if (UE_TRACE_CHANNELEXPR_IS_ENABLED(MassChannel))
+	{
+		UE_TRACE_LOG(MassTrace, MassPhaseEnd, MassChannel)
+			<< MassPhaseEnd.Cycle(FPlatformTime::Cycles64())
+			<< MassPhaseEnd.PhaseName(PhaseName)
+			<< MassPhaseEnd.PhaseId(0);
+	}
+}
+
+void FMassTrace::OutputEndPhaseRegion(uint64 PhaseId)
+{
+	if (UE_TRACE_CHANNELEXPR_IS_ENABLED(MassChannel))
+	{
+		UE_TRACE_LOG(MassTrace, MassPhaseEnd, MassChannel)
+			<< MassPhaseEnd.Cycle(FPlatformTime::Cycles64())
+			<< MassPhaseEnd.PhaseId(PhaseId);
+	}
+}
+
+#endif //UE_MASS_TRACE_ENABLED
\ No newline at end of file
diff --git a/Engine/Source/Runtime/MassEntity/Private/MassProcessingPhaseManager.cpp b/Engine/Source/Runtime/MassEntity/Private/MassProcessingPhaseManager.cpp
--- a/Engine/Source/Runtime/MassEntity/Private/MassProcessingPhaseManager.cpp
+++ b/Engine/Source/Runtime/MassEntity/Private/MassProcessingPhaseManager.cpp
@@ -10,6 +10,7 @@
 #include "VisualLogger/VisualLogger.h"
 #include "Engine/World.h"
 #include "MassCommandBuffer.h"
+#include "MassEntityTrace.h"
 
 #include UE_INLINE_GENERATED_CPP_BY_NAME(MassProcessingPhaseManager)
 
@@ -525,6 +526,8 @@ void FMassProcessingPhaseManager::OnPhaseStart(FMassProcessingPhase& Phase)
 		HandlePendingDynamicProcessorOperations(PhaseAsInt);
 	}
 
+	UE_TRACE_MASS_PHASE_BEGIN(PhaseAsInt)
+
 	if (Owner.IsValid()
 		&& ensure(Phase.Phase != EMassProcessingPhase::MAX)
 		&& (ProcessingGraphBuildStates[PhaseAsInt].bNewArchetypes || ProcessingGraphBuildStates[PhaseAsInt].bProcessorsNeedRebuild)
@@ -568,6 +571,7 @@ void FMassProcessingPhaseManager::OnPhaseStart(FMassProcessingPhase& Phase)
 void FMassProcessingPhaseManager::OnPhaseEnd(FMassProcessingPhase& Phase)
 {
 	ensure(CurrentPhase == Phase.Phase);
+	UE_TRACE_MASS_PHASE_END(static_cast<int32>(CurrentPhase))
 	CurrentPhase = EMassProcessingPhase::MAX;
 
 	// The VERY LAST thing we do in FrameEnd is change the Pause state if needed.
diff --git a/Engine/Source/Runtime/MassEntity/Public/MassDebugger.h b/Engine/Source/Runtime/MassEntity/Public/MassDebugger.h
--- a/Engine/Source/Runtime/MassEntity/Public/MassDebugger.h
+++ b/Engine/Source/Runtime/MassEntity/Public/MassDebugger.h
@@ -3,6 +3,7 @@
 #pragma once
 
 #include "MassProcessingTypes.h"
+#include "MassEntityTrace.h"
 #if WITH_MASSENTITY_DEBUG
 #include "Containers/ContainersFwd.h"
 #include "MassEntityQuery.h"
@@ -19,6 +20,7 @@ class UMassProcessor;
 struct FMassEntityQuery;
 struct FMassEntityManager;
 struct FMassArchetypeHandle;
+struct FMassArchetypeChunk;
 struct FMassFragmentRequirements;
 struct FMassFragmentRequirementDescription;
 enum class EMassFragmentAccess : uint8;
@@ -71,6 +73,7 @@ namespace UE::Mass::Debug
 	extern MASSENTITY_API bool bAllowProceduralDebuggedEntitySelection;
 	extern MASSENTITY_API bool bAllowBreakOnDebuggedEntity;
 	extern MASSENTITY_API bool bTestSelectedEntityAgainstProcessorQueries;
+	using FArchetypeFunction = TFunction<void(FMassArchetypeHandle)>;
 } // namespace UE::Mass::Debug
 
 #define MASS_IF_ENTITY_DEBUGGED(Manager, EntityHandle) (FMassDebugger::GetSelectedEntity(Manager) == EntityHandle)
@@ -148,9 +151,12 @@ struct MASSENTITY_API FMassDebugger
 		TMultiMap<const UMassProcessor*, FMassEntityHandle> ProcessorBreakpoints;
 		TMultiMap<const UScriptStruct*, FMassEntityHandle> FragmentWriteBreakpoints;
 
-		explicit FEnvironment(const FMassEntityManager& InEntityManager)
-			: EntityManager(InEntityManager.AsWeak())
-		{}
+#if UE_MASS_TRACE_ENABLED
+		FDelegateHandle TraceStartedDelegateHandle;
+#endif
+
+		explicit FEnvironment(const FMassEntityManager& InEntityManager);
+		~FEnvironment();
 
 		bool IsValid() const { return EntityManager.IsValid(); }
 
@@ -172,9 +178,20 @@ struct MASSENTITY_API FMassDebugger
 	static void GetQueryExecutionRequirements(const FMassEntityQuery& Query, FMassExecutionRequirements& OutExecutionRequirements);
 	static TArray<FMassEntityHandle> GetEntitiesMatchingQuery(const FMassEntityManager& EntityManager, const FMassEntityQuery& Query);
 
+	static void ForEachArchetype(const FMassEntityManager& EntitySubsystem, const UE::Mass::Debug::FArchetypeFunction& Function);
 	static TArray<FMassArchetypeHandle> GetAllArchetypes(const FMassEntityManager& EntitySubsystem);
 	static const FMassArchetypeCompositionDescriptor& GetArchetypeComposition(const FMassArchetypeHandle& ArchetypeHandle);
 
+	static uint64 GetArchetypeTraceID(const FMassArchetypeData& ArchetypeData);
+	static uint64 GetArchetypeTraceID(const FMassArchetypeHandle& ArchetypeHandle);
+	
+	static TConstArrayView<FMassEntityHandle> GetEntitiesViewOfArchetype(
+		const FMassArchetypeData& ArchetypeData,
+		const FMassArchetypeChunk& Chunk);
+
+	static const FMassArchetypeData* GetArchetypeData(const FMassArchetypeHandle& ArchetypeHandle);
+	static void EnumerateChunks(const FMassArchetypeData& Archetype, TFunctionRef<void(const FMassArchetypeChunk&)> Fn);
+
 	static void GetArchetypeEntityStats(const FMassArchetypeHandle & ArchetypeHandle, UE::Mass::Debug::FArchetypeStats & OutStats);
 	static const TConstArrayView<FName> GetArchetypeDebugNames(const FMassArchetypeHandle& ArchetypeHandle);
 	static TArray<FMassEntityHandle> GetEntitiesOfArchetype(const FMassArchetypeHandle& ArchetypeHandle);
