[Mass]
Cleanup extraneous conditionals around trace output

diff --git a/Engine/Source/Runtime/MassEntity/Private/MassEntityTrace.cpp b/Engine/Source/Runtime/MassEntity/Private/MassEntityTrace.cpp
--- a/Engine/Source/Runtime/MassEntity/Private/MassEntityTrace.cpp
+++ b/Engine/Source/Runtime/MassEntity/Private/MassEntityTrace.cpp
@@ -127,33 +127,24 @@ enum class FFragmentType : uint8
 
 void FMassTrace::QueryCreated(const FMassEntityQuery* Query)
 {
-	if (UE_TRACE_CHANNELEXPR_IS_ENABLED(MassChannel))
-	{
-		UE_TRACE_LOG(MassTrace, QueryCreated, MassChannel)
-			<< QueryCreated.Cycle(FPlatformTime::Cycles64())
-			<< QueryCreated.QueryID(reinterpret_cast<uint64>(Query));
-	}
+	UE_TRACE_LOG(MassTrace, QueryCreated, MassChannel)
+		<< QueryCreated.Cycle(FPlatformTime::Cycles64())
+		<< QueryCreated.QueryID(reinterpret_cast<uint64>(Query));
 }
 
 void FMassTrace::QueryDestroyed(const FMassEntityQuery* Query)
 {
-	if (UE_TRACE_CHANNELEXPR_IS_ENABLED(MassChannel))
-	{
-		UE_TRACE_LOG(MassTrace, QueryDestroyed, MassChannel)
-			<< QueryDestroyed.Cycle(FPlatformTime::Cycles64())
-			<< QueryDestroyed.QueryID(reinterpret_cast<uint64>(Query));
-	}
+	UE_TRACE_LOG(MassTrace, QueryDestroyed, MassChannel)
+		<< QueryDestroyed.Cycle(FPlatformTime::Cycles64())
+		<< QueryDestroyed.QueryID(reinterpret_cast<uint64>(Query));
 }
 
 void FMassTrace::QueryRegisteredToProcessor(const FMassEntityQuery* Query, TNotNull<const UMassProcessor*> Processor)
 {
-	if (UE_TRACE_CHANNELEXPR_IS_ENABLED(MassChannel))
-	{
-		UE_TRACE_LOG(MassTrace, QueryRegisteredToProcessor, MassChannel)
-			<< QueryRegisteredToProcessor.QueryID(reinterpret_cast<uint64>(Query))
-			<< QueryRegisteredToProcessor.ProcessorID(reinterpret_cast<uint64>((const UMassProcessor*)(Processor)))
-			<< QueryRegisteredToProcessor.ProcessorName(*Processor->GetProcessorName());
-	}
+	UE_TRACE_LOG(MassTrace, QueryRegisteredToProcessor, MassChannel)
+		<< QueryRegisteredToProcessor.QueryID(reinterpret_cast<uint64>(Query))
+		<< QueryRegisteredToProcessor.ProcessorID(reinterpret_cast<uint64>((const UMassProcessor*)(Processor)))
+		<< QueryRegisteredToProcessor.ProcessorName(*Processor->GetProcessorName());
 }
 
 void FMassTrace::QueryArchetypeAdded(const FMassEntityQuery* Query, const FMassArchetypeHandle& Archetype)
@@ -174,25 +165,19 @@ FMassTrace::FScopedQueryForEachTrace::FScopedQueryForEachTrace(const FMassEntity
 	, ChunkCount(0)
 	, EntityCount(0)
 {
-	if (UE_TRACE_CHANNELEXPR_IS_ENABLED(MassChannel))
-	{
-		UE_TRACE_LOG(MassTrace, QueryForEachStarted, MassChannel)
-			<< QueryForEachStarted.Cycle(FPlatformTime::Cycles64())
-			<< QueryForEachStarted.QueryID(reinterpret_cast<uint64>(Query));
-	}
+	UE_TRACE_LOG(MassTrace, QueryForEachStarted, MassChannel)
+		<< QueryForEachStarted.Cycle(FPlatformTime::Cycles64())
+		<< QueryForEachStarted.QueryID(reinterpret_cast<uint64>(Query));
 }
 
 FMassTrace::FScopedQueryForEachTrace::~FScopedQueryForEachTrace()
 {
-	if (UE_TRACE_CHANNELEXPR_IS_ENABLED(MassChannel))
-	{
-		UE_TRACE_LOG(MassTrace, QueryForEachComplete, MassChannel)
-			<< QueryForEachComplete.Cycle(FPlatformTime::Cycles64())
-			<< QueryForEachComplete.QueryID(reinterpret_cast<uint64>(Query))
-			<< QueryForEachComplete.ArchetypeCount(ArchetypeCount)
-			<< QueryForEachComplete.ChunkCount(ChunkCount)
-			<< QueryForEachComplete.EntityCount(EntityCount);
-	}
+	UE_TRACE_LOG(MassTrace, QueryForEachComplete, MassChannel)
+		<< QueryForEachComplete.Cycle(FPlatformTime::Cycles64())
+		<< QueryForEachComplete.QueryID(reinterpret_cast<uint64>(Query))
+		<< QueryForEachComplete.ArchetypeCount(ArchetypeCount)
+		<< QueryForEachComplete.ChunkCount(ChunkCount)
+		<< QueryForEachComplete.EntityCount(EntityCount);
 }
 
 void FMassTrace::FScopedQueryForEachTrace::ReportArchetype(const FMassArchetypeData& Archetype)
@@ -440,34 +425,25 @@ uint64_t FMassTrace::OutputBeginPhaseWithID(const TCHAR* PhaseName)
 
 void FMassTrace::OutputBeginPhaseRegion(const TCHAR* PhaseName)
 {
-	if (UE_TRACE_CHANNELEXPR_IS_ENABLED(MassChannel))
-	{
-		UE_TRACE_LOG(MassTrace, MassPhaseBegin, MassChannel)
-			<< MassPhaseBegin.Cycle(FPlatformTime::Cycles64())
-			<< MassPhaseBegin.PhaseName(PhaseName)
-			<< MassPhaseBegin.PhaseId(0);
-	}
+	UE_TRACE_LOG(MassTrace, MassPhaseBegin, MassChannel)
+		<< MassPhaseBegin.Cycle(FPlatformTime::Cycles64())
+		<< MassPhaseBegin.PhaseName(PhaseName)
+		<< MassPhaseBegin.PhaseId(0);
 }
 
 void FMassTrace::OutputEndPhaseRegion(const TCHAR* PhaseName)
 {
-	if (UE_TRACE_CHANNELEXPR_IS_ENABLED(MassChannel))
-	{
-		UE_TRACE_LOG(MassTrace, MassPhaseEnd, MassChannel)
-			<< MassPhaseEnd.Cycle(FPlatformTime::Cycles64())
-			<< MassPhaseEnd.PhaseName(PhaseName)
-			<< MassPhaseEnd.PhaseId(0);
-	}
+	UE_TRACE_LOG(MassTrace, MassPhaseEnd, MassChannel)
+		<< MassPhaseEnd.Cycle(FPlatformTime::Cycles64())
+		<< MassPhaseEnd.PhaseName(PhaseName)
+		<< MassPhaseEnd.PhaseId(0);
 }
 
 void FMassTrace::OutputEndPhaseRegion(uint64 PhaseId)
 {
-	if (UE_TRACE_CHANNELEXPR_IS_ENABLED(MassChannel))
-	{
-		UE_TRACE_LOG(MassTrace, MassPhaseEnd, MassChannel)
-			<< MassPhaseEnd.Cycle(FPlatformTime::Cycles64())
-			<< MassPhaseEnd.PhaseId(PhaseId);
-	}
+	UE_TRACE_LOG(MassTrace, MassPhaseEnd, MassChannel)
+		<< MassPhaseEnd.Cycle(FPlatformTime::Cycles64())
+		<< MassPhaseEnd.PhaseId(PhaseId);
 }
 
 #endif //UE_MASS_TRACE_ENABLED
\ No newline at end of file
