[Mass] Mass entity types comments spelling fixes

#rnx

diff --git a/Engine/Source/Runtime/MassEntity/Public/MassEntityTypes.h b/Engine/Source/Runtime/MassEntity/Public/MassEntityTypes.h
--- a/Engine/Source/Runtime/MassEntity/Public/MassEntityTypes.h
+++ b/Engine/Source/Runtime/MassEntity/Public/MassEntityTypes.h
@@ -39,7 +39,7 @@ namespace UE::Mass
 }
 
 /** The type summarily describing a composition of an entity or an archetype. It contains information on both the
- *  fragments as well as tags */
+ *  fragments and tags */
 struct FMassArchetypeCompositionDescriptor
 {
 	FMassArchetypeCompositionDescriptor() = default;
@@ -428,7 +428,7 @@ struct MASSENTITY_API FMassArchetypeSharedFragmentValues
 protected:
 	mutable uint32 HashCache = UINT32_MAX;
 	/**
-	 * We consider empty FMassArchetypeSharedFragmentValues a sorted containter.Same goes for a container containing
+	 * We consider empty FMassArchetypeSharedFragmentValues a sorted container.Same goes for a container containing
 	 * a single element, @see DirtyHashCache
 	 */ 
 	mutable bool bSorted = true; 
@@ -454,7 +454,7 @@ enum class EMassObservedOperation : uint8
 {
 	Add,
 	Remove,
-	// @todo Keeping this here as a indication of design intent. For now we handle entity destruction like removal, but 
+	// @todo Keeping this here as an indication of design intent. For now, we handle entity destruction like removal, but 
 	// there might be computationally expensive cases where we might want to avoid for soon-to-be-dead entities. 
 	// Destroy,
 	// @todo another planned supported operation type
@@ -541,7 +541,7 @@ private:
 namespace UE::Mass
 {
 	/**
-	 * A statically-typed list of of related types. Used mainly to differentiate type collections at compile-type as well as
+	 * A statically-typed list of related types. Used mainly to differentiate type collections at compile-type as well as
 	 * efficiently produce TStructTypeBitSet representing given collection.
 	 */
 	template<typename T, typename... TOthers>
