[Mass] Minor refactor of the Observer Manager, where functions containing 'FragmentOrTag' in its name got that part renamed to 'Element'

#rnx

diff --git a/Engine/Source/Runtime/MassEntity/Private/MassObserverManager.cpp b/Engine/Source/Runtime/MassEntity/Private/MassObserverManager.cpp
--- a/Engine/Source/Runtime/MassEntity/Private/MassObserverManager.cpp
+++ b/Engine/Source/Runtime/MassEntity/Private/MassObserverManager.cpp
@@ -278,21 +278,21 @@ bool FMassObserverManager::OnCompositionChanged(const FMassEntityHandle Entity,
 	return bHasFragmentsOverlap || bHasTagsOverlap;
 }
 
-void FMassObserverManager::OnFragmentOrTagOperation(const UScriptStruct& FragmentOrTagType, const FMassArchetypeEntityCollection& EntityCollection, const EMassObservedOperation Operation)
+void FMassObserverManager::OnElementOperation(TNotNull<const UScriptStruct*> ElementType, const FMassArchetypeEntityCollection& EntityCollection, const EMassObservedOperation Operation)
 {
-	const bool bIsFragment = UE::Mass::IsA<FMassFragment>(&FragmentOrTagType);
-	check(bIsFragment || UE::Mass::IsA<FMassTag>(&FragmentOrTagType));
+	const bool bIsFragment = UE::Mass::IsA<FMassFragment>(ElementType);
+	check(bIsFragment || UE::Mass::IsA<FMassTag>(ElementType));
 
 	if (bIsFragment)
 	{
-		if (ObservedFragments[(uint8)Operation].Contains(FragmentOrTagType))
+		if (ObservedFragments[(uint8)Operation].Contains(*ElementType))
 		{
-			HandleSingleEntityImpl(FragmentOrTagType, EntityCollection, FragmentObservers[(uint8)Operation]);
+			HandleSingleEntityImpl(*ElementType, EntityCollection, FragmentObservers[(uint8)Operation]);
 		}
 	}
-	else if (ObservedTags[(uint8)Operation].Contains(FragmentOrTagType))
+	else if (ObservedTags[(uint8)Operation].Contains(*ElementType))
 	{
-		HandleSingleEntityImpl(FragmentOrTagType, EntityCollection, TagObservers[(uint8)Operation]);
+		HandleSingleEntityImpl(*ElementType, EntityCollection, TagObservers[(uint8)Operation]);
 	}
 }
 
@@ -325,22 +325,22 @@ void FMassObserverManager::HandleSingleEntityImpl(const UScriptStruct& FragmentT
 	UE::Mass::Executor::RunProcessorsView(Pipeline.GetMutableProcessors(), ProcessingContext, MakeArrayView(&EntityCollection, 1));
 }
 
-void FMassObserverManager::AddObserverInstance(const UScriptStruct& FragmentOrTagType, const EMassObservedOperation Operation, UMassProcessor& ObserverProcessor)
+void FMassObserverManager::AddObserverInstance(const UScriptStruct& ElementType, const EMassObservedOperation Operation, UMassProcessor& ObserverProcessor)
 {
-	const bool bIsFragment = UE::Mass::IsA<FMassFragment>(&FragmentOrTagType);
-	checkSlow(bIsFragment || UE::Mass::IsA<FMassTag>(&FragmentOrTagType));
+	const bool bIsFragment = UE::Mass::IsA<FMassFragment>(&ElementType);
+	checkSlow(bIsFragment || UE::Mass::IsA<FMassTag>(&ElementType));
 
 	FMassRuntimePipeline* Pipeline = nullptr;
 
 	if (bIsFragment)
 	{
-		Pipeline = &(*FragmentObservers[(uint8)Operation]).FindOrAdd(&FragmentOrTagType);
-		ObservedFragments[(uint8)Operation].Add(FragmentOrTagType);
+		Pipeline = &(*FragmentObservers[(uint8)Operation]).FindOrAdd(&ElementType);
+		ObservedFragments[(uint8)Operation].Add(ElementType);
 	}
 	else
 	{
-		Pipeline = &(*TagObservers[(uint8)Operation]).FindOrAdd(&FragmentOrTagType);
-		ObservedTags[(uint8)Operation].Add(FragmentOrTagType);
+		Pipeline = &(*TagObservers[(uint8)Operation]).FindOrAdd(&ElementType);
+		ObservedTags[(uint8)Operation].Add(ElementType);
 	}
 
 	// AppendUniqueProcessor will return true only if ObserverProcessor has not been a part of Pipeline yet.
@@ -355,11 +355,11 @@ void FMassObserverManager::AddObserverInstance(const UScriptStruct& FragmentOrTa
 	}
 }
 
-void FMassObserverManager::RemoveObserverInstance(const UScriptStruct& FragmentOrTagType, const EMassObservedOperation Operation, UMassProcessor& ObserverProcessor)
+void FMassObserverManager::RemoveObserverInstance(const UScriptStruct& ElementType, const EMassObservedOperation Operation, UMassProcessor& ObserverProcessor)
 {
-	const bool bIsFragmentObserver = UE::Mass::IsA<FMassFragment>(&FragmentOrTagType);
+	const bool bIsFragmentObserver = UE::Mass::IsA<FMassFragment>(&ElementType);
 
-	if (!ensure(bIsFragmentObserver || UE::Mass::IsA<FMassTag>(&FragmentOrTagType)))
+	if (!ensure(bIsFragmentObserver || UE::Mass::IsA<FMassTag>(&ElementType)))
 	{
 		return;
 	}
@@ -367,7 +367,7 @@ void FMassObserverManager::RemoveObserverInstance(const UScriptStruct& FragmentO
 	TMap<TObjectPtr<const UScriptStruct>, FMassRuntimePipeline>& ObserversMap =
 		bIsFragmentObserver ? *FragmentObservers[(uint8)Operation] : *TagObservers[(uint8)Operation];
 
-	FMassRuntimePipeline* Pipeline = ObserversMap.Find(&FragmentOrTagType);
+	FMassRuntimePipeline* Pipeline = ObserversMap.Find(&ElementType);
 	if (!ensureMsgf(Pipeline, TEXT("Trying to remove an observer for a fragment/tag that does not seem to be observed.")))
 	{
 		return;
@@ -376,14 +376,32 @@ void FMassObserverManager::RemoveObserverInstance(const UScriptStruct& FragmentO
 
 	if (Pipeline->Num() == 0)
 	{
-		ObserversMap.Remove(&FragmentOrTagType);
+		ObserversMap.Remove(&ElementType);
 		if (bIsFragmentObserver)
 		{
-			ObservedFragments[(uint8)Operation].Remove(FragmentOrTagType);
+			ObservedFragments[(uint8)Operation].Remove(ElementType);
 		}
 		else
 		{
-			ObservedTags[(uint8)Operation].Remove(FragmentOrTagType);
+			ObservedTags[(uint8)Operation].Remove(ElementType);
 		}
 	}
 }
+
+//----------------------------------------------------------------------//
+// FMassObserversMap
+//----------------------------------------------------------------------//
+void FMassObserverManager::OnPostFragmentOrTagAdded(const UScriptStruct& FragmentOrTagType, const FMassArchetypeEntityCollection& EntityCollection)
+{
+	OnPostElementAdded(&FragmentOrTagType, EntityCollection);
+}
+
+void FMassObserverManager::OnPreFragmentOrTagRemoved(const UScriptStruct& FragmentOrTagType, const FMassArchetypeEntityCollection& EntityCollection)
+{
+	OnPreElementRemoved(&FragmentOrTagType, EntityCollection);
+}
+
+void FMassObserverManager::OnFragmentOrTagOperation(const UScriptStruct& FragmentOrTagType, const FMassArchetypeEntityCollection& EntityCollection, const EMassObservedOperation Operation)
+{
+	OnElementOperation(&FragmentOrTagType, EntityCollection, Operation);
+}
diff --git a/Engine/Source/Runtime/MassEntity/Public/MassObserverManager.h b/Engine/Source/Runtime/MassEntity/Public/MassObserverManager.h
--- a/Engine/Source/Runtime/MassEntity/Public/MassObserverManager.h
+++ b/Engine/Source/Runtime/MassEntity/Public/MassObserverManager.h
@@ -116,20 +116,18 @@ public:
 		return OnCompositionChanged(LocalContext, EntityCollection, Composition, Operation);
 	}
 
-	void OnPostFragmentOrTagAdded(const UScriptStruct& FragmentOrTagType, const FMassArchetypeEntityCollection& EntityCollection)
+	void OnPostElementAdded(TNotNull<const UScriptStruct*> ElementType, const FMassArchetypeEntityCollection& EntityCollection)
 	{
-		OnFragmentOrTagOperation(FragmentOrTagType, EntityCollection, EMassObservedOperation::Add);
+		OnElementOperation(ElementType, EntityCollection, EMassObservedOperation::Add);
 	}
-	void OnPreFragmentOrTagRemoved(const UScriptStruct& FragmentOrTagType, const FMassArchetypeEntityCollection& EntityCollection)
+	void OnPreElementRemoved(TNotNull<const UScriptStruct*> ElementType, const FMassArchetypeEntityCollection& EntityCollection)
 	{
-		OnFragmentOrTagOperation(FragmentOrTagType, EntityCollection, EMassObservedOperation::Remove);
+		OnElementOperation(ElementType, EntityCollection, EMassObservedOperation::Remove);
 	}
-	// @todo I don't love this name. Alternatively could be OnSingleFragmentOrTagOperation. Long term we'll switch to
-	// OnSingleFragmentOperation and advertise that Tags are a type of Fragment as well (conceptually).
-	void OnFragmentOrTagOperation(const UScriptStruct& FragmentOrTagType, const FMassArchetypeEntityCollection& EntityCollection, const EMassObservedOperation Operation);
+	void OnElementOperation(TNotNull<const UScriptStruct*> ElementType, const FMassArchetypeEntityCollection& EntityCollection, const EMassObservedOperation Operation);
 
-	void AddObserverInstance(const UScriptStruct& FragmentOrTagType, const EMassObservedOperation Operation, UMassProcessor& ObserverProcessor);
-	void RemoveObserverInstance(const UScriptStruct& FragmentOrTagType, const EMassObservedOperation Operation, UMassProcessor& ObserverProcessor);
+	void AddObserverInstance(const UScriptStruct& ElementType, const EMassObservedOperation Operation, UMassProcessor& ObserverProcessor);
+	void RemoveObserverInstance(const UScriptStruct& ElementType, const EMassObservedOperation Operation, UMassProcessor& ObserverProcessor);
 
 	void DebugGatherUniqueProcessors(TArray<const UMassProcessor*>& OutProcessors) const;
 
@@ -159,18 +157,14 @@ protected:
 	FMassEntityManager& EntityManager;
 
 public:
-	UE_DEPRECATED(5.4, "This flavor of OnCompositionChanged is deprecated. Please use the one taking a FMassProcessingContext& parameter instead")
-	bool OnCompositionChanged(const FMassArchetypeEntityCollection& EntityCollection, const FMassArchetypeCompositionDescriptor& Composition, const EMassObservedOperation Operation, FMassProcessingContext* InProcessingContext)
-	{
-		if (InProcessingContext)
-		{
-			return OnCompositionChanged(*InProcessingContext, EntityCollection, Composition, Operation);
-		}
-		return OnCompositionChanged(EntityCollection, Composition, Operation);
-	}
-
 	UE_DEPRECATED(5.5, "This flavor of OnPostEntitiesCreated is deprecated. Please use the one taking a TConstArrayView<FMassArchetypeEntityCollection> parameter instead")
 	bool OnPostEntitiesCreated(FMassProcessingContext& InProcessingContext, const FMassArchetypeEntityCollection& EntityCollection);
+	UE_DEPRECATED(5.5, "*FragmentOrTag* functions have been renamed to *Element*. Please use those instead.")
+	void OnPostFragmentOrTagAdded(const UScriptStruct& FragmentOrTagType, const FMassArchetypeEntityCollection& EntityCollection);
+	UE_DEPRECATED(5.5, "*FragmentOrTag* functions have been renamed to *Element*. Please use those instead.")
+	void OnPreFragmentOrTagRemoved(const UScriptStruct& FragmentOrTagType, const FMassArchetypeEntityCollection& EntityCollection);
+	UE_DEPRECATED(5.5, "*FragmentOrTag* functions have been renamed to *Element*. Please use those instead.")
+	void OnFragmentOrTagOperation(const UScriptStruct& FragmentOrTagType, const FMassArchetypeEntityCollection& EntityCollection, const EMassObservedOperation Operation);
 };
 
 template<>
