[Mass] Made MassEntityManager.IsDuringCreation public

diff --git a/Engine/Source/Runtime/MassEntity/Public/MassEntityManager.h b/Engine/Source/Runtime/MassEntity/Public/MassEntityManager.h
--- a/Engine/Source/Runtime/MassEntity/Public/MassEntityManager.h
+++ b/Engine/Source/Runtime/MassEntity/Public/MassEntityManager.h
@@ -1025,6 +1025,8 @@ PRAGMA_ENABLE_DEPRECATION_WARNINGS
 	UWorld* GetWorld() const { return Owner.IsValid() ? Owner->GetWorld() : nullptr; }
 	UObject* GetOwner() const { return Owner.Get(); }
 
+	bool IsDuringEntityCreation() const;
+
 	void SetDebugName(const FString& NewDebugGame);
 #if WITH_MASSENTITY_DEBUG
 	void DebugPrintArchetypes(FOutputDevice& Ar, const bool bIncludeEmpty = true) const;
@@ -1138,8 +1140,7 @@ private:
 	 * Otherwise ActiveCreationContext will get extended with ReservedEntities and EntityCollection, and returned by the function.
 	 */
 	TSharedRef<FEntityCreationContext> GetOrMakeCreationContext(TConstArrayView<FMassEntityHandle> ReservedEntities, FMassArchetypeEntityCollection&& EntityCollection);
-	bool IsDuringEntityCreation() const { return ActiveCreationContext.IsValid(); }
-	
+
 	/** 
 	 * This type is used in entity mutating batched API to ensure the active FEntityCreationContext gets dirtied 
 	 * upon function's end (since the mutating operations render FEntityCreationContext.EntityCollections invalid).
@@ -1408,4 +1409,9 @@ inline const UE::Mass::FTypeManager& FMassEntityManager::GetTypeManager() const
 inline UE::Mass::FTypeManager& FMassEntityManager::GetTypeManager()
 {
 	return *TypeManager;
-}
\ No newline at end of file
+}
+
+inline bool FMassEntityManager::IsDuringEntityCreation() const
+{
+	return ActiveCreationContext.IsValid();
+}
