[Mass] Introduced the concept of debug-time "color" to MassArchetypeData. To be used in debug drawing.

diff --git a/Engine/Source/Runtime/MassEntity/Internal/MassArchetypeData.h b/Engine/Source/Runtime/MassEntity/Internal/MassArchetypeData.h
--- a/Engine/Source/Runtime/MassEntity/Internal/MassArchetypeData.h
+++ b/Engine/Source/Runtime/MassEntity/Internal/MassArchetypeData.h
@@ -187,8 +187,15 @@ private:
 	uint32 CreatedArchetypeDataVersion = 0;
 
 #if WITH_MASSENTITY_DEBUG
-	// Arrays of names the archetype is referred as.
+	/** Arrays of names the archetype is referred as. */
 	TArray<FName> DebugNames;
+
+	/**
+	 * Color to be used when representing this archetype. If not set with FMassArchetypeCreationParams
+	 * will be deterministically set based on archetype's composition. Can be overridden at any point 
+	 * via SetDebugColor.
+	 */
+	FColor DebugColor;
 #endif // WITH_MASSENTITY_DEBUG
 
 	// Defaults to UMassEntitySettings.ChunkMemorySize. In near future will support being set via constructor.
@@ -347,6 +354,8 @@ public:
 	void DebugPrintEntity(FMassEntityHandle Entity, FOutputDevice& Ar, const TCHAR* InPrefix = TEXT("")) const;
 #endif // WITH_MASSENTITY_DEBUG
 
+	void SetDebugColor(const FColor InDebugColor);
+
 	void REMOVEME_GetArrayViewForFragmentInChunk(int32 ChunkIndex, const UScriptStruct* FragmentType, void*& OutChunkBase, int32& OutNumEntities);
 
 	//////////////////////////////////////////////////////////////////////
diff --git a/Engine/Source/Runtime/MassEntity/Private/MassArchetypeData.cpp b/Engine/Source/Runtime/MassEntity/Private/MassArchetypeData.cpp
--- a/Engine/Source/Runtime/MassEntity/Private/MassArchetypeData.cpp
+++ b/Engine/Source/Runtime/MassEntity/Private/MassArchetypeData.cpp
@@ -39,6 +39,7 @@ FMassArchetypeData::FMassArchetypeData(const FMassArchetypeCreationParams& Creat
 {
 #if WITH_MASSENTITY_DEBUG
 	DebugNames.Add(CreationParams.DebugName);
+	DebugColor = CreationParams.DebugColor;
 #endif // WITH_MASSENTITY_DEBUG
 }
 
@@ -89,6 +90,10 @@ void FMassArchetypeData::Initialize(const FMassEntityManager& EntityManager, con
 	CompositionDescriptor.ConstSharedFragments = InCompositionDescriptor.ConstSharedFragments;
 
 	EntityListOffsetWithinChunk = 0;
+
+#if WITH_MASSENTITY_DEBUG
+	SetDebugColor(DebugColor);
+#endif // WITH_MASSENTITY_DEBUG
 }
 
 void FMassArchetypeData::InitializeWithSimilar(const FMassEntityManager& EntityManager, const FMassArchetypeData& BaseArchetype
@@ -116,6 +121,10 @@ void FMassArchetypeData::InitializeWithSimilar(const FMassEntityManager& EntityM
 	Groups = InGroups;
 
 	EntityListOffsetWithinChunk = 0;
+
+#if WITH_MASSENTITY_DEBUG
+	SetDebugColor(DebugColor);
+#endif // WITH_MASSENTITY_DEBUG
 }
 
 void FMassArchetypeData::ConfigureFragments(const FMassEntityManager& EntityManager)
@@ -479,6 +488,9 @@ void FMassArchetypeData::ExecuteFunction(FMassExecutionContext& RunContext, cons
 
 	// @todo do we really want users to check composition of the archetype being processed at the moment?
 	RunContext.SetCurrentArchetypeCompositionDescriptor(GetCompositionDescriptor());
+#if WITH_MASSENTITY_DEBUG
+	RunContext.DebugSetColor(DebugColor);
+#endif // WITH_MASSENTITY_DEBUG
 
 	uint32 PrevSharedFragmentValuesHash = UINT32_MAX;
 	for (FMassArchetypeChunkIterator ChunkIterator(EntityRangeContainer); ChunkIterator; ++ChunkIterator)
@@ -518,6 +530,9 @@ void FMassArchetypeData::ExecuteFunction(FMassExecutionContext& RunContext, cons
 	}
 
 	RunContext.SetCurrentArchetypeCompositionDescriptor(GetCompositionDescriptor());
+#if WITH_MASSENTITY_DEBUG
+	RunContext.DebugSetColor(DebugColor);
+#endif // WITH_MASSENTITY_DEBUG
 
 	uint32 PrevSharedFragmentValuesHash = UINT32_MAX;
 	for (FMassArchetypeChunk& Chunk : Chunks)
@@ -556,6 +571,10 @@ void FMassArchetypeData::ExecutionFunctionForChunk(FMassExecutionContext& RunCon
 
 		RunContext.SetCurrentArchetypeCompositionDescriptor(GetCompositionDescriptor());
 		RunContext.SetCurrentChunkSerialModificationNumber(Chunk.GetSerialModificationNumber());
+#if WITH_MASSENTITY_DEBUG
+		RunContext.DebugSetColor(DebugColor);
+#endif // WITH_MASSENTITY_DEBUG
+
 		BindChunkFragmentRequirements(RunContext, RequirementMapping.ChunkFragments, Chunk);
 
 		if (!ChunkCondition || ChunkCondition(RunContext))
@@ -1017,6 +1036,29 @@ void FMassArchetypeData::DebugPrintEntity(FMassEntityHandle Entity, FOutputDevic
 
 #endif // WITH_MASSENTITY_DEBUG
 
+void FMassArchetypeData::SetDebugColor(const FColor InDebugColor)
+{
+#if WITH_MASSENTITY_DEBUG
+	if (InDebugColor == FColor{0})
+	{
+		// pick a color based on the composition
+		const uint32 Compositionhash = CompositionDescriptor.CalculateHash();
+		const uint8* Bytes = reinterpret_cast<const uint8*>(&Compositionhash);
+		
+		const FLinearColor AdjustedColor = FLinearColor::MakeFromHSV8(
+			static_cast<uint8>((Bytes[0] >> 1) + (Bytes[1] >> 1)),
+			static_cast<uint8>((Bytes[2] >> 1) + 128),
+			static_cast<uint8>((Bytes[3] >> 1) + 128)
+		);
+		DebugColor = AdjustedColor.ToFColorSRGB();
+	}
+	else
+	{
+		DebugColor = InDebugColor;
+	}
+#endif // WITH_MASSENTITY_DEBUG
+}
+
 void FMassArchetypeData::REMOVEME_GetArrayViewForFragmentInChunk(int32 ChunkIndex, const UScriptStruct* FragmentType, void*& OutChunkBase, int32& OutNumEntities)
 {
 	const FMassArchetypeChunk& Chunk = Chunks[ChunkIndex];
diff --git a/Engine/Source/Runtime/MassEntity/Private/MassExecutionContext.cpp b/Engine/Source/Runtime/MassEntity/Private/MassExecutionContext.cpp
--- a/Engine/Source/Runtime/MassEntity/Private/MassExecutionContext.cpp
+++ b/Engine/Source/Runtime/MassEntity/Private/MassExecutionContext.cpp
@@ -58,6 +58,9 @@ void FMassExecutionContext::ClearExecutionData()
 	CurrentArchetypeCompositionDescriptor.Reset();
 	EntityListView = {};
 	ChunkSerialModificationNumber = -1;
+#if WITH_MASSENTITY_DEBUG
+	DebugColor = FColor();
+#endif // WITH_MASSENTITY_DEBUG
 }
 
 bool FMassExecutionContext::CacheSubsystemRequirements(const FMassSubsystemRequirements& SubsystemRequirements)
diff --git a/Engine/Source/Runtime/MassEntity/Public/MassEntityTypes.h b/Engine/Source/Runtime/MassEntity/Public/MassEntityTypes.h
--- a/Engine/Source/Runtime/MassEntity/Public/MassEntityTypes.h
+++ b/Engine/Source/Runtime/MassEntity/Public/MassEntityTypes.h
@@ -676,6 +676,10 @@ struct FMassArchetypeCreationParams
 
 	/** Name to identify the archetype while debugging*/
 	FName DebugName;
+
+#if WITH_MASSENTITY_DEBUG
+	FColor DebugColor{0};
+#endif // WITH_MASSENTITY_DEBUG
 };
 
 //-----------------------------------------------------------------------------
diff --git a/Engine/Source/Runtime/MassEntity/Public/MassExecutionContext.h b/Engine/Source/Runtime/MassEntity/Public/MassExecutionContext.h
--- a/Engine/Source/Runtime/MassEntity/Public/MassExecutionContext.h
+++ b/Engine/Source/Runtime/MassEntity/Public/MassExecutionContext.h
@@ -67,6 +67,9 @@ private:
 	float DeltaTimeSeconds = 0.0f;
 	int32 ChunkSerialModificationNumber = -1;
 	FMassArchetypeCompositionDescriptor CurrentArchetypeCompositionDescriptor;
+#if WITH_MASSENTITY_DEBUG
+	FColor DebugColor;
+#endif // WITH_MASSENTITY_DEBUG
 
 	TSharedRef<FMassEntityManager> EntityManager;
 
@@ -313,6 +316,13 @@ public:
 		return CurrentArchetypeCompositionDescriptor.Tags.Contains<T>();
 	}
 
+#if WITH_MASSENTITY_DEBUG
+	FColor DebugGetArchetypeColor() const
+	{
+		return DebugColor;
+	}
+#endif // WITH_MASSENTITY_DEBUG
+
 	/** Chunk related operations */
 	void SetCurrentChunkSerialModificationNumber(const int32 SerialModificationNumber) { ChunkSerialModificationNumber = SerialModificationNumber; }
 	int32 GetChunkSerialModificationNumber() const { return ChunkSerialModificationNumber; }
@@ -517,6 +527,13 @@ public:
 		CurrentArchetypeCompositionDescriptor = Descriptor;
 	}
 
+#if WITH_MASSENTITY_DEBUG
+	void DebugSetColor(FColor InColor)
+	{
+		DebugColor = InColor;
+	}
+#endif // WITH_MASSENTITY_DEBUG
+
 	/** 
 	 * Processes SubsystemRequirements to fetch and cache all the indicated subsystems. If a UWorld is required to fetch
 	 * a specific subsystem then the one associated with the stored EntityManager will be used.
