[Mass] Introducing "versioned" archetype handle, that will become "out of date" when the target archetype goes through changes that change the ordering of the hosted entities (essentially removing entities or compaction)

diff --git a/Engine/Source/Runtime/MassEntity/Internal/MassArchetypeData.h b/Engine/Source/Runtime/MassEntity/Internal/MassArchetypeData.h
--- a/Engine/Source/Runtime/MassEntity/Internal/MassArchetypeData.h
+++ b/Engine/Source/Runtime/MassEntity/Internal/MassArchetypeData.h
@@ -183,9 +183,18 @@ private:
 	SIZE_T TotalBytesPerEntity = 0;
 	int32 EntityListOffsetWithinChunk;
 
-	// Archetype version at which this archetype was created, useful for query to do incremental archetype matching
+	/**
+	 * Archetype version at which this archetype was created, useful for query to do incremental archetype matching.
+	 * Note that it's set once and never changed afterward.
+	 */
 	uint32 CreatedArchetypeDataVersion = 0;
 
+	/**
+	 * Incremented whenever an operation modifies the order of hosted entities, for example entity removal and compaction.
+	 * This value is used to validate stored entity ranges, including FMassArchetypeEntityCollection.
+	 */
+	uint32 EntityOrderVersion = 0;
+
 #if WITH_MASSENTITY_DEBUG
 	/** Arrays of names the archetype is referred as. */
 	TArray<FName> DebugNames;
@@ -267,7 +276,8 @@ public:
 	int32 GetChunkCount() const { return Chunks.Num(); }
 	int32 GetNonEmptyChunkCount() const;
 
-	uint32 GetCreatedArchetypeDataVersion() const { return CreatedArchetypeDataVersion; }
+	uint32 GetCreatedArchetypeDataVersion() const;
+	uint32 GetEntityOrderVersion() const;
 
 	void ExecuteFunction(FMassExecutionContext& RunContext, const FMassExecuteFunction& Function, const FMassQueryRequirementIndicesMapping& RequirementMapping
 		, FMassArchetypeEntityCollection::FConstEntityRangeArrayView EntityRangeContainer, const FMassChunkConditionFunction& ChunkCondition);
@@ -513,3 +523,13 @@ inline bool FMassArchetypeData::IsInGroupOfType(const UE::Mass::FArchetypeGroupT
 {
 	return Groups.ContainsType(GroupType);
 }
+
+inline uint32 FMassArchetypeData::GetCreatedArchetypeDataVersion() const
+{
+	return CreatedArchetypeDataVersion;
+}
+
+inline uint32 FMassArchetypeData::GetEntityOrderVersion() const
+{
+	return EntityOrderVersion;
+}
diff --git a/Engine/Source/Runtime/MassEntity/Private/MassArchetypeData.cpp b/Engine/Source/Runtime/MassEntity/Private/MassArchetypeData.cpp
--- a/Engine/Source/Runtime/MassEntity/Private/MassArchetypeData.cpp
+++ b/Engine/Source/Runtime/MassEntity/Private/MassArchetypeData.cpp
@@ -285,6 +285,8 @@ void FMassArchetypeData::RemoveEntity(FMassEntityHandle Entity)
 
 void FMassArchetypeData::RemoveEntityInternal(const int32 AbsoluteIndex)
 {
+	++EntityOrderVersion;
+
 	const int32 ChunkIndex = AbsoluteIndex / NumEntitiesPerChunk;
 	const int32 IndexWithinChunk = AbsoluteIndex % NumEntitiesPerChunk;
 
@@ -664,6 +666,11 @@ int32 FMassArchetypeData::CompactEntities(const double TimeAllowed)
 		}
 	}
 
+	if (TotalEntitiesMoved > 0)
+	{
+		++EntityOrderVersion;
+	}
+
 	return TotalEntitiesMoved;
 }
 
@@ -1293,6 +1300,13 @@ FMassArchetypeEntityCollection::FArchetypeEntityRange FMassArchetypeData::Prepar
 
 void FMassArchetypeData::BatchRemoveEntitiesInternal(const int32 ChunkIndex, const int32 StartIndexWithinChunk, const int32 NumberToRemove)
 {
+	if (UNLIKELY(NumberToRemove <= 0))
+	{
+		return;
+	}
+
+	++EntityOrderVersion;
+
 	FMassArchetypeChunk& Chunk = Chunks[ChunkIndex];
 	
 	const int32 NumberToMove = FMath::Min(Chunk.GetNumInstances() - (StartIndexWithinChunk + NumberToRemove), NumberToRemove);
diff --git a/Engine/Source/Runtime/MassEntity/Private/MassArchetypeTypes.cpp b/Engine/Source/Runtime/MassEntity/Private/MassArchetypeTypes.cpp
--- a/Engine/Source/Runtime/MassEntity/Private/MassArchetypeTypes.cpp
+++ b/Engine/Source/Runtime/MassEntity/Private/MassArchetypeTypes.cpp
@@ -14,6 +14,26 @@ uint32 GetTypeHash(const FMassArchetypeHandle& Instance)
 	return GetTypeHash(Instance.DataPtr.Get());
 }
 
+//-----------------------------------------------------------------------------
+// FMassArchetypeVersionedHandle
+//-----------------------------------------------------------------------------
+FMassArchetypeVersionedHandle::FMassArchetypeVersionedHandle(const FMassArchetypeHandle& InHandle)
+	: ArchetypeHandle(InHandle)
+	, HandleVersion(ArchetypeHandle.IsValid() ? ArchetypeHandle.DataPtr->GetEntityOrderVersion() : 0)
+{
+}
+
+FMassArchetypeVersionedHandle::FMassArchetypeVersionedHandle(FMassArchetypeHandle&& InHandle)
+	: ArchetypeHandle(MoveTemp(InHandle))
+	, HandleVersion(ArchetypeHandle.IsValid() ? ArchetypeHandle.DataPtr->GetEntityOrderVersion() : 0)
+{
+}
+
+bool FMassArchetypeVersionedHandle::IsUpToDate() const
+{
+	return ArchetypeHandle.IsValid() && (ArchetypeHandle.DataPtr->GetEntityOrderVersion() == HandleVersion);
+}
+
 //-----------------------------------------------------------------------------
 // FMassArchetypeEntityCollection 
 //-----------------------------------------------------------------------------
diff --git a/Engine/Source/Runtime/MassEntity/Private/MassEntityQuery.cpp b/Engine/Source/Runtime/MassEntity/Private/MassEntityQuery.cpp
--- a/Engine/Source/Runtime/MassEntity/Private/MassEntityQuery.cpp
+++ b/Engine/Source/Runtime/MassEntity/Private/MassEntityQuery.cpp
@@ -250,7 +250,7 @@ void FMassEntityQuery::ForEachEntityChunk(FMassExecutionContext& ExecutionContex
 
 	if (FMassFragmentRequirements::IsEmpty())
 	{
-		if (ensureMsgf(ExecutionContext.GetEntityCollection().IsSet(), TEXT("Using empty queries is only supported in combination with Entity Collections that explicitly indicate entities to process")))
+		if (ensureMsgf(ExecutionContext.GetEntityCollection().IsEmpty() == false, TEXT("Using empty queries is only supported in combination with Entity Collections that explicitly indicate entities to process")))
 		{
 			static const FMassQueryRequirementIndicesMapping EmptyMapping;
 
@@ -269,8 +269,8 @@ void FMassEntityQuery::ForEachEntityChunk(FMassExecutionContext& ExecutionContex
 			// do need those calculations.
 		CacheArchetypes();
 
-		// if there's a chunk collection set by the external code - use that
-		if (ExecutionContext.GetEntityCollection().IsSet())
+		// if there's a chunk collection set by the external code - usse that
+		if (ExecutionContext.GetEntityCollection().IsEmpty() == false)
 		{
 			const FMassArchetypeHandle& ArchetypeHandle = ExecutionContext.GetEntityCollection().GetArchetype();
 			const int32 ArchetypeIndex = ValidArchetypes.Find(ArchetypeHandle);
@@ -384,7 +384,7 @@ void FMassEntityQuery::ParallelForEachEntityChunk(FMassExecutionContext& Executi
 
 	if (FMassFragmentRequirements::IsEmpty())
 	{
-		if (ensureMsgf(ExecutionContext.GetEntityCollection().IsSet(), TEXT("Using empty queries is only supported in combination with Entity Collections that explicitly indicate entities to process")))
+		if (ensureMsgf(ExecutionContext.GetEntityCollection().IsEmpty() == false, TEXT("Using empty queries is only supported in combination with Entity Collections that explicitly indicate entities to process")))
 		{
 			const FMassArchetypeHandle& ArchetypeHandle = ExecutionContext.GetEntityCollection().GetArchetype();
 			FMassArchetypeData& ArchetypeRef = FMassArchetypeHelper::ArchetypeDataFromHandleChecked(ArchetypeHandle);
@@ -403,7 +403,7 @@ void FMassEntityQuery::ParallelForEachEntityChunk(FMassExecutionContext& Executi
 		CacheArchetypes();
 
 		// if there's a chunk collection set by the external code - use that
-		if (ExecutionContext.GetEntityCollection().IsSet())
+		if (ExecutionContext.GetEntityCollection().IsEmpty() == false)
 		{
 			const FMassArchetypeHandle& ArchetypeHandle = ExecutionContext.GetEntityCollection().GetArchetype();
 			const int32 ArchetypeIndex = ValidArchetypes.Find(ArchetypeHandle);
diff --git a/Engine/Source/Runtime/MassEntity/Private/MassExecutionContext.cpp b/Engine/Source/Runtime/MassEntity/Private/MassExecutionContext.cpp
--- a/Engine/Source/Runtime/MassEntity/Private/MassExecutionContext.cpp
+++ b/Engine/Source/Runtime/MassEntity/Private/MassExecutionContext.cpp
@@ -77,6 +77,7 @@ void FMassExecutionContext::SetEntityCollection(const FMassArchetypeEntityCollec
 void FMassExecutionContext::SetEntityCollection(FMassArchetypeEntityCollection&& InEntityCollection)
 {
 	check(EntityCollection.IsEmpty());
+	check(InEntityCollection.IsUpToDate());
 	EntityCollection = MoveTemp(InEntityCollection);
 }
 
diff --git a/Engine/Source/Runtime/MassEntity/Public/MassArchetypeTypes.h b/Engine/Source/Runtime/MassEntity/Public/MassArchetypeTypes.h
--- a/Engine/Source/Runtime/MassEntity/Public/MassArchetypeTypes.h
+++ b/Engine/Source/Runtime/MassEntity/Public/MassArchetypeTypes.h
@@ -23,6 +23,7 @@ struct FMassArchetypeEntityCollection;
 struct FMassEntityView;
 struct FMassDebugger;
 struct FMassArchetypeHelper;
+struct FMassArchetypeVersionedHandle;
 
 using FMassEntityExecuteFunction = TFunction< void(FMassExecutionContext& /*ExecutionContext*/, int32 /*EntityIndex*/) >;
 using FMassExecuteFunction = TFunction< void(FMassExecutionContext& /*ExecutionContext*/) >;
@@ -31,7 +32,7 @@ using FMassChunkConditionFunction = TFunction< bool(const FMassExecutionContext&
 //-----------------------------------------------------------------------------
 // FMassArchetypeHandle
 //-----------------------------------------------------------------------------
-// An opaque handle to an archetype
+/** An opaque handle to an archetype */
 struct FMassArchetypeHandle final
 {
 	FMassArchetypeHandle() = default;
@@ -40,15 +41,50 @@ struct FMassArchetypeHandle final
 	bool operator==(const FMassArchetypeHandle& Other) const { return DataPtr == Other.DataPtr; }
 	bool operator!=(const FMassArchetypeHandle& Other) const { return DataPtr != Other.DataPtr; }
 
+	bool IsUpToData() const;
+
 	MASSENTITY_API friend uint32 GetTypeHash(const FMassArchetypeHandle& Instance);
 private:
 	FMassArchetypeHandle(const TSharedPtr<FMassArchetypeData>& InDataPtr)
 	: DataPtr(InDataPtr)
 	{}
 	TSharedPtr<FMassArchetypeData> DataPtr;
+	/**
+	 * This value indicates whether the target archetype had its entities moved around since the handle creations.
+	 * The information is useful in a couple of scenarios (like making sure an entity collection is up to date),
+	 * but in most cases the users should not concern themselves with this value.
+	 * Note that the value is not used as part of hash calculation, it's effectively transient. 
+	 */
+	uint32 HandleVersion = 0;
 
 	friend FMassArchetypeHelper;
 	friend FMassEntityManager;
+	friend FMassArchetypeVersionedHandle;
+};
+
+struct FMassArchetypeVersionedHandle final
+{
+	FMassArchetypeVersionedHandle() = default;
+	FMassArchetypeVersionedHandle(const FMassArchetypeHandle& InHandle);
+	FMassArchetypeVersionedHandle(FMassArchetypeHandle&& InHandle);
+
+	bool IsValid() const;
+	bool operator==(const FMassArchetypeVersionedHandle& Other) const;
+	bool operator!=(const FMassArchetypeVersionedHandle& Other) const;
+	bool IsUpToDate() const;
+
+	operator FMassArchetypeHandle() const;
+
+	MASSENTITY_API friend uint32 GetTypeHash(const FMassArchetypeHandle& Instance);
+private:
+	FMassArchetypeHandle ArchetypeHandle;
+	/**
+	 * This value indicates whether the target archetype had its entities moved around since the handle creations.
+	 * The information is useful in a couple of scenarios (like making sure an entity collection is up to date),
+	 * but in most cases the users should not concern themselves with this value.
+	 * Note that the value is not used as part of hash calculation, it's effectively transient. 
+	 */
+	uint32 HandleVersion = 0;
 };
 
 //-----------------------------------------------------------------------------
@@ -113,7 +149,7 @@ public:
 private:
 	FEntityRangeArray Ranges;
 	/** entity indices indicated by EntityRanges are only valid with given Archetype */
-	FMassArchetypeHandle Archetype;
+	FMassArchetypeVersionedHandle Archetype;
 
 public:
 	FMassArchetypeEntityCollection() = default;
@@ -125,13 +161,17 @@ public:
 		, Archetype(InArchetypeHandle)
 	{}
 
-	FConstEntityRangeArrayView GetRanges() const { return Ranges; }
-	const FMassArchetypeHandle& GetArchetype() const { return Archetype; }
-	bool IsEmpty() const { return Ranges.Num() == 0 && Archetype.IsValid() == false; }
-	bool IsSet() const { return Archetype.IsValid(); }
+	FConstEntityRangeArrayView GetRanges() const;
+	FMassArchetypeHandle GetArchetype() const;
+	bool IsEmpty() const;
+	bool IsUpToDate() const;
+
+	UE_DEPRECATED(5.6, "This function is deprecated. Use !IsEmpty() instead.")
+	bool IsSet() const;
+
 	void Reset() 
 	{ 
-		Archetype = FMassArchetypeHandle();
+		Archetype = FMassArchetypeVersionedHandle();
 		Ranges.Reset();
 	}
 
@@ -252,6 +292,51 @@ struct FMassQueryRequirementIndicesMapping
 //-----------------------------------------------------------------------------
 // INLINES
 //-----------------------------------------------------------------------------
+inline bool FMassArchetypeVersionedHandle::IsValid() const
+{
+	return ArchetypeHandle.IsValid();
+}
+
+inline bool FMassArchetypeVersionedHandle::operator==(const FMassArchetypeVersionedHandle& Other) const 
+{ 
+	return ArchetypeHandle == Other.ArchetypeHandle && HandleVersion == Other.HandleVersion; 
+}
+
+inline bool FMassArchetypeVersionedHandle::operator!=(const FMassArchetypeVersionedHandle& Other) const 
+{ 
+	return !(*this == Other);
+}
+
+inline FMassArchetypeEntityCollection::FConstEntityRangeArrayView FMassArchetypeEntityCollection::GetRanges() const
+{
+	return Ranges;
+}
+
+inline FMassArchetypeHandle FMassArchetypeEntityCollection::GetArchetype() const
+{
+	return Archetype;
+}
+
+inline bool FMassArchetypeEntityCollection::IsEmpty() const
+{
+	return Ranges.Num() == 0 && Archetype.IsValid() == false;
+}
+
+inline bool FMassArchetypeEntityCollection::IsUpToDate() const
+{
+	return IsEmpty() || Archetype.IsUpToDate();
+}
+
+inline bool FMassArchetypeEntityCollection::IsSet() const
+{
+	return !IsEmpty();
+}
+
+inline FMassArchetypeVersionedHandle::operator FMassArchetypeHandle() const
+{
+	return ArchetypeHandle;
+}
+
 inline FMassRawEntityInChunkData::FMassRawEntityInChunkData(uint8* InChunkRawMemory, const int32 InIndexWithinChunk)
 	: ChunkRawMemory(InChunkRawMemory), IndexWithinChunk(InIndexWithinChunk)
 {}
