[MassDebugger]
Added warning message when a breakpoint is set but a debugger is not attached.

diff --git a/Engine/Source/Runtime/MassEntity/Private/MassDebugger.cpp b/Engine/Source/Runtime/MassEntity/Private/MassDebugger.cpp
--- a/Engine/Source/Runtime/MassEntity/Private/MassDebugger.cpp
+++ b/Engine/Source/Runtime/MassEntity/Private/MassDebugger.cpp
@@ -17,7 +17,10 @@
 #include "MassEntityUtils.h"
 #include "MassCommandBuffer.h"
 #include "MassEntityTrace.h"
+#include "Misc/MessageDialog.h"
+#include "HAL/PlatformMisc.h"
 
+#define LOCTEXT_NAMESPACE "MassDebugger"
 
 namespace UE::Mass::Debug
 {
@@ -1004,6 +1007,11 @@ bool FMassDebugger::HasAnyFragmentWriteBreakpoints(const FMassEntityManager& Ent
 
 void FMassDebugger::SetProcessorBreakpoint(const FMassEntityManager& EntityManager, TNotNull<const UMassProcessor*> Processor, FMassEntityHandle Entity)
 {
+	if (!FPlatformMisc::IsDebuggerPresent())
+	{
+		FMessageDialog::Open(EAppMsgType::Ok, LOCTEXT("NoDebuggerAttached", "Breakpoint set but no debugger is attached."));
+	}
+
 	FEnvironment& ActiveEnvironment = GetActiveEnvironment(EntityManager);
 
 	ActiveEnvironment.bHasBreakpoint = true;
@@ -1015,6 +1023,11 @@ void FMassDebugger::SetProcessorBreakpoint(const FMassEntityManager& EntityManag
 
 void FMassDebugger::SetFragmentWriteBreak(const FMassEntityManager& EntityManager, const UScriptStruct* FragmentType, FMassEntityHandle Entity)
 {
+	if (!FPlatformMisc::IsDebuggerPresent())
+	{
+		FMessageDialog::Open(EAppMsgType::Ok, LOCTEXT("NoDebuggerAttached", "Breakpoint set but no debugger is attached."));
+	}
+
 	FEnvironment& ActiveEnvironment = GetActiveEnvironment(EntityManager);
 
 	ActiveEnvironment.bHasBreakpoint = true;
@@ -1245,5 +1258,5 @@ void FMassDebugger::FEnvironment::ClearBreakpoints()
 	bHasBreakpoint = false;
 }
 
-
+#undef LOCTEXT_NAMESPACE
 #endif // WITH_MASSENTITY_DEBUG
\ No newline at end of file
