[MassEntity] fixed some typos

diff --git a/Engine/Source/Runtime/MassEntity/Private/MassEntityManager.cpp b/Engine/Source/Runtime/MassEntity/Private/MassEntityManager.cpp
--- a/Engine/Source/Runtime/MassEntity/Private/MassEntityManager.cpp
+++ b/Engine/Source/Runtime/MassEntity/Private/MassEntityManager.cpp
@@ -97,11 +97,11 @@ void FMassEntityManager::FEntityCreationContext::MarkDirty()
 	EntityCollections.Reset();
 }
 
-void FMassEntityManager::FEntityCreationContext::AppendEntities(const TConstArrayView<FMassEntityHandle> EntitiesToAppend)
+void FMassEntityManager::FEntityCreationContext::AppendEntities(const TConstArrayView<FMassEntityHandle> InEntitiesToAppend)
 {
 	checkf(OwnerThreadId == FPlatformTLS::GetCurrentThreadId(), TEXT("%hs: all FEntityCreationContext operations ere expected to be run in a single thread"), __FUNCTION__);
 
-	if (EntitiesToAppend.Num())
+	if (InEntitiesToAppend.Num())
 	{
 		LLM_SCOPE_BYNAME(TEXT("Mass/EntityManager"));
 		if (CreatedEntities.Num())
@@ -114,26 +114,26 @@ void FMassEntityManager::FEntityCreationContext::AppendEntities(const TConstArra
 		// else, if there are no entities the resulting state will be "dirty" by design
 		ensureMsgf(EntityCollections.IsEmpty(), TEXT("Having a non-empty array of entity collections is unexpected at this point!"));
 
-		CreatedEntities.Append(EntitiesToAppend);
+		CreatedEntities.Append(InEntitiesToAppend);
 		ensure(IsDirty());
 	}
 }
 
-void FMassEntityManager::FEntityCreationContext::AppendEntities(const TConstArrayView<FMassEntityHandle> EntitiesToAppend, FMassArchetypeEntityCollection&& InEntityCollection)
+void FMassEntityManager::FEntityCreationContext::AppendEntities(const TConstArrayView<FMassEntityHandle> InEntitiesToAppend, FMassArchetypeEntityCollection&& InEntityCollection)
 {
-	checkf(OwnerThreadId == FPlatformTLS::GetCurrentThreadId(), TEXT("%hs: all FEntityCreationContext operations ere expected to be run in a single thread"), __FUNCTION__);
+	checkf(OwnerThreadId == FPlatformTLS::GetCurrentThreadId(), TEXT("%hs: all FEntityCreationContext operations are expected to be run in a single thread"), __FUNCTION__);
 
-	if (EntitiesToAppend.Num() == 0)
+	if (InEntitiesToAppend.Num() == 0)
 	{
 		return;
 	}
 
 	LLM_SCOPE_BYNAME(TEXT("Mass/EntityManager"));
 
-	AppendEntities(EntitiesToAppend);
+	AppendEntities(InEntitiesToAppend);
 
 	// this condition boils down to checking if this FEntityCreationContext instance only connects the just added EntitiesToAppend
-	if (CreatedEntities.Num() == EntitiesToAppend.Num())
+	if (CreatedEntities.Num() == InEntitiesToAppend.Num())
 	{
 		checkf(EntityCollections.Num() == 0, TEXT("We never expect EntityCollections to be non-empty while there are no entities in CreatedEntities."));
 		EntityCollections.Add(MoveTemp(InEntityCollection));
diff --git a/Engine/Source/Runtime/MassEntity/Public/MassEntityManager.h b/Engine/Source/Runtime/MassEntity/Public/MassEntityManager.h
--- a/Engine/Source/Runtime/MassEntity/Public/MassEntityManager.h
+++ b/Engine/Source/Runtime/MassEntity/Public/MassEntityManager.h
@@ -245,7 +245,7 @@ public:
 	private:
 		FEntityCreationContext();
 		explicit FEntityCreationContext(FMassEntityManager& InManager, const TConstArrayView<FMassEntityHandle> InCreatedEntities = {});
-		FEntityCreationContext(FMassEntityManager& InManager, const TConstArrayView<FMassEntityHandle> InCreatedEntities, FMassArchetypeEntityCollection&& EntityCollection);
+		FEntityCreationContext(FMassEntityManager& InManager, const TConstArrayView<FMassEntityHandle> InCreatedEntities, FMassArchetypeEntityCollection&& InEntityCollection);
 		
 	public:
 		~FEntityCreationContext();
@@ -255,8 +255,8 @@ public:
 		int32 GetSpawnedNum() const { return CreatedEntities.Num(); }
 		void MarkDirty();
 		FORCEINLINE bool IsDirty() const { return EntityCollections.IsEmpty() && (CreatedEntities.IsEmpty() == false); }
-		void AppendEntities(const TConstArrayView<FMassEntityHandle> EntitiesToAppend);
-		void AppendEntities(const TConstArrayView<FMassEntityHandle> EntitiesToAppend, FMassArchetypeEntityCollection&& EntityCollection);
+		void AppendEntities(const TConstArrayView<FMassEntityHandle> InEntitiesToAppend);
+		void AppendEntities(const TConstArrayView<FMassEntityHandle> InEntitiesToAppend, FMassArchetypeEntityCollection&& InEntityCollection);
 
 		/** Function for debugging/testing purposes. We don't expect users to ever call it, always get collections via GetEntityCollections */
 		bool DebugAreEntityCollectionsUpToDate() const { return EntityCollections.IsEmpty() == CreatedEntities.IsEmpty(); }
diff --git a/Engine/Source/Runtime/MassEntity/Public/MassExecutionContext.h b/Engine/Source/Runtime/MassEntity/Public/MassExecutionContext.h
--- a/Engine/Source/Runtime/MassEntity/Public/MassExecutionContext.h
+++ b/Engine/Source/Runtime/MassEntity/Public/MassExecutionContext.h
@@ -252,8 +252,8 @@ public:
 	};
 
 	/**
-	 * Creats and Entity Iterator for the current chunk.
-	 * Supports ranged for and can be used directly as an entity index for the current chunk.
+	 * Creates an Entity Iterator for the current chunk.
+	 * Supports range-based for loop and can be used directly as an entity index for the current chunk.
 	 */
 	FEntityIterator CreateEntityIterator();
 
