[Mass] Made debug-time requirement access detection by entity queries a bit more automatic (with a scoped "detector") which fixes an edge case of not cleaning up the detection on failed `FScopedEntityQueryContext` initialization (can happen when required subsystems are missing, for example).

diff --git a/Engine/Source/Runtime/MassEntity/Private/MassEntityQuery.cpp b/Engine/Source/Runtime/MassEntity/Private/MassEntityQuery.cpp
--- a/Engine/Source/Runtime/MassEntity/Private/MassEntityQuery.cpp
+++ b/Engine/Source/Runtime/MassEntity/Private/MassEntityQuery.cpp
@@ -41,6 +41,7 @@ struct FScopedEntityQueryContext
 {
 	FScopedEntityQueryContext(FMassEntityQuery& InQuery, FMassExecutionContext& ExecutionContext)
 		: Query(InQuery), CachedExecutionContext(ExecutionContext)
+		, ScopedAccessDetector(InQuery)
 	{
 		CachedExecutionContext.PushQuery(InQuery);
 
@@ -64,6 +65,7 @@ struct FScopedEntityQueryContext
 
 	FMassEntityQuery& Query;
 	FMassExecutionContext& CachedExecutionContext;
+	UE::Mass::Debug::FScopedRequirementAccessDetector ScopedAccessDetector;
 	bool bSubsystemRequirementsCached;
 };
 
@@ -236,8 +238,6 @@ void FMassEntityQuery::ForEachEntityChunk(FMassExecutionContext& ExecutionContex
 #if WITH_MASSENTITY_DEBUG
 	checkf(ExecutionContext.GetExecutionType() == ExpectedContextType && (ExpectedContextType == EMassExecutionContextType::Local || bRegistered)
 		, TEXT("ExecutionContextType mismatch, make sure all the queries run as part of processor execution are registered with some processor with a FMassEntityQuery::RegisterWithProcessor call"));
-
-	CachedEntityManager->GetRequirementAccessDetector().RequireAccess(*this);
 #endif
 
 	FScopedEntityQueryContext ScopedQueryContext(*this, ExecutionContext);
@@ -278,12 +278,9 @@ void FMassEntityQuery::ForEachEntityChunk(FMassExecutionContext& ExecutionContex
 			// if given ArchetypeHandle cannot be found in ValidArchetypes then it doesn't match the query's requirements
 			if (ArchetypeIndex == INDEX_NONE)
 			{
-				UE_VLOG_UELOG(CachedEntityManager->GetOwner(), LogMass, Log, TEXT("Attempted to execute FMassEntityQuery with an incompatible Archetype: %s")
+				UE_VLOG_UELOG(CachedEntityManager->GetOwner(), LogMass, Verbose, TEXT("Attempted to execute FMassEntityQuery with an incompatible Archetype: %s. Note that this is fine for observers.")
 					, *FMassDebugger::GetArchetypeRequirementCompatibilityDescription(*this, ArchetypeHandle));
 
-#if WITH_MASSENTITY_DEBUG
-				CachedEntityManager->GetRequirementAccessDetector().ReleaseAccess(*this);
-#endif // WITH_MASSENTITY_DEBUG
 				return;
 			}
 
@@ -317,10 +314,6 @@ void FMassEntityQuery::ForEachEntityChunk(FMassExecutionContext& ExecutionContex
 			}
 		}
 	}
-
-#if WITH_MASSENTITY_DEBUG
-	CachedEntityManager->GetRequirementAccessDetector().ReleaseAccess(*this);
-#endif
 }
 
 void FMassEntityQuery::ParallelForEachEntityChunkInCollection(TConstArrayView<FMassArchetypeEntityCollection> EntityCollections
@@ -362,8 +355,6 @@ void FMassEntityQuery::ParallelForEachEntityChunk(FMassExecutionContext& Executi
 #if WITH_MASSENTITY_DEBUG
 	checkf(ExecutionContext.GetExecutionType() == ExpectedContextType && (ExpectedContextType == EMassExecutionContextType::Local || bRegistered)
 		, TEXT("ExecutionContextType mismatch, make sure all the queries run as part of processor execution are registered with some processor with a FMassEntityQuery::RegisterWithProcessor call"));
-
-	CachedEntityManager->GetRequirementAccessDetector().RequireAccess(*this);
 #endif
 
 	FScopedEntityQueryContext ScopedQueryContext(*this, ExecutionContext);
@@ -411,12 +402,9 @@ void FMassEntityQuery::ParallelForEachEntityChunk(FMassExecutionContext& Executi
 			// if given ArchetypeHandle cannot be found in ValidArchetypes then it doesn't match the query's requirements
 			if (ArchetypeIndex == INDEX_NONE)
 			{
-				UE_VLOG_UELOG(CachedEntityManager->GetOwner(), LogMass, Log, TEXT("Attempted to execute FMassEntityQuery with an incompatible Archetype: %s")
+				UE_VLOG_UELOG(CachedEntityManager->GetOwner(), LogMass, Verbose, TEXT("Attempted to execute FMassEntityQuery with an incompatible Archetype: %s. Note that this is fine for observers.")
 					, *FMassDebugger::GetArchetypeRequirementCompatibilityDescription(*this, ExecutionContext.GetEntityCollection().GetArchetype()));
 
-#if WITH_MASSENTITY_DEBUG
-				CachedEntityManager->GetRequirementAccessDetector().ReleaseAccess(*this);
-#endif // WITH_MASSENTITY_DEBUG
 				return;
 			}
 
@@ -498,10 +486,6 @@ void FMassEntityQuery::ParallelForEachEntityChunk(FMassExecutionContext& Executi
 				}, ParallelForFlags);
 		}
 	}
-
-#if WITH_MASSENTITY_DEBUG
-	CachedEntityManager->GetRequirementAccessDetector().ReleaseAccess(*this);
-#endif
 }
 
 int32 FMassEntityQuery::GetNumMatchingEntities()
diff --git a/Engine/Source/Runtime/MassEntity/Private/MassRequirementAccessDetector.cpp b/Engine/Source/Runtime/MassEntity/Private/MassRequirementAccessDetector.cpp
--- a/Engine/Source/Runtime/MassEntity/Private/MassRequirementAccessDetector.cpp
+++ b/Engine/Source/Runtime/MassEntity/Private/MassRequirementAccessDetector.cpp
@@ -4,6 +4,7 @@
 #if WITH_MASSENTITY_DEBUG
 #include "MassEntityQuery.h"
 #include "HAL/IConsoleManager.h"
+#include "MassEntityManager.h"
 
 namespace UE::Mass::Private
 {
@@ -60,4 +61,22 @@ void FMassRequirementAccessDetector::ReleaseAccess(const FMassEntityQuery& Query
 	}
 }
 
-#endif // WITH_MASSENTITY_DEBUG
\ No newline at end of file
+namespace UE::Mass::Debug
+{
+	//-----------------------------------------------------------------------------
+	// FScopedRequirementAccessDetector
+	//-----------------------------------------------------------------------------
+	FScopedRequirementAccessDetector::FScopedRequirementAccessDetector(const FMassEntityQuery& InQuery)
+		: EntityManager(InQuery.GetEntityManager())
+		, Query(InQuery)
+	{
+		EntityManager->GetRequirementAccessDetector().RequireAccess(InQuery);
+	}
+
+	FScopedRequirementAccessDetector::~FScopedRequirementAccessDetector()
+	{
+		EntityManager->GetRequirementAccessDetector().ReleaseAccess(Query);
+	}
+}
+
+#endif // WITH_MASSENTITY_DEBUG
diff --git a/Engine/Source/Runtime/MassEntity/Public/MassEntityQuery.h b/Engine/Source/Runtime/MassEntity/Public/MassEntityQuery.h
--- a/Engine/Source/Runtime/MassEntity/Public/MassEntityQuery.h
+++ b/Engine/Source/Runtime/MassEntity/Public/MassEntityQuery.h
@@ -197,6 +197,8 @@ public:
 	void SetArchetypeMatchOverride(const T& Override);
 #endif
 
+	const TSharedPtr<FMassEntityManager>& GetEntityManager() const;
+
 	/** 
 	 * If ArchetypeHandle is among ValidArchetypes then the function retrieves requirements mapping cached for it,
 	 * otherwise an empty mapping will be returned (and the requirements binding will be done the slow way).
@@ -357,5 +359,10 @@ inline bool FMassEntityQuery::IsGrouping() const
 	return !GroupSortingSteps.IsEmpty();
 }
 
+inline const TSharedPtr<FMassEntityManager>& FMassEntityQuery::GetEntityManager() const
+{
+	return CachedEntityManager;
+}
+
 #undef MASS_ACHETYPE_OVERRIDE_MAX_SIZE
 #undef MASS_ACHETYPE_OVERRIDE_MAX_ALIGNMENT
\ No newline at end of file
diff --git a/Engine/Source/Runtime/MassEntity/Public/MassRequirementAccessDetector.h b/Engine/Source/Runtime/MassEntity/Public/MassRequirementAccessDetector.h
--- a/Engine/Source/Runtime/MassEntity/Public/MassRequirementAccessDetector.h
+++ b/Engine/Source/Runtime/MassEntity/Public/MassRequirementAccessDetector.h
@@ -8,6 +8,9 @@
 #include "MassEntityQuery.h"
 #include "Misc/MTAccessDetector.h"
 
+
+struct FMassEntityManager;
+
 struct MASSENTITY_API FMassRequirementAccessDetector
 {
 	void Initialize();
@@ -72,4 +75,23 @@ private:
 
 	TMap<const UStruct*, TSharedRef<FRWAccessDetector>> Detectors;
 };
-#endif // WITH_MASSENTITY_DEBUG
\ No newline at end of file
+
+#endif // WITH_MASSENTITY_DEBUG
+
+namespace UE::Mass::Debug
+{
+	struct FScopedRequirementAccessDetector
+	{
+#if WITH_MASSENTITY_DEBUG
+		MASSENTITY_API FScopedRequirementAccessDetector(const FMassEntityQuery& InQuery);
+		MASSENTITY_API ~FScopedRequirementAccessDetector();
+
+		TSharedPtr<FMassEntityManager> EntityManager;
+		const FMassEntityQuery& Query;
+#else
+		FScopedRequirementAccessDetector(const FMassEntityQuery&)
+		{
+		}
+#endif // WITH_MASSENTITY_DEBUG
+	};
+}
