[Mass] Improved Mass processor dependency cycles detection and reporting

diff --git a/Engine/Source/Runtime/MassEntity/Private/MassProcessorDependencySolver.cpp b/Engine/Source/Runtime/MassEntity/Private/MassProcessorDependencySolver.cpp
--- a/Engine/Source/Runtime/MassEntity/Private/MassProcessorDependencySolver.cpp
+++ b/Engine/Source/Runtime/MassEntity/Private/MassProcessorDependencySolver.cpp
@@ -7,6 +7,9 @@
 #include "Logging/MessageLog.h"
 #include "HAL/FileManager.h"
 #include "Engine/World.h"
+#if WITH_MASSENTITY_DEBUG
+#include "Algo/Count.h"
+#endif // WITH_MASSENTITY_DEBUG
 
 #define LOCTEXT_NAMESPACE "Mass"
 
@@ -37,6 +40,72 @@ namespace UE::Mass::Private
 		}
 		return false;
 	}
+
+#if WITH_MASSENTITY_DEBUG
+	void LogCycle(TArray<FMassProcessorDependencySolver::FNode> AllNodes, TConstArrayView<int32> CycleIndices, TArray<uint32>& InOutReportedCycleHashes)
+	{
+		check(CycleIndices.Num());
+		// We extract unique indices involved in the cycle below.
+		// Note that we want to preserve the order since it will provide more meaningful debugging context.
+		// But we do find the "lowest node index" so that we can generate a deterministic
+		// hash representing the cycle, regardless of which node was being processed when the cycle was found.
+		// We use this information to not report the same cycle multiple times.
+
+		// Finding the cycle start (the first node that has been encountered more than once)
+		int32 CycleStartElementIndex = INDEX_NONE;
+		for (const int32 CycleNodeIndex : CycleIndices)
+		{
+			if (Algo::Count(CycleIndices, CycleNodeIndex) > 1)
+			{
+				CycleStartElementIndex = CycleIndices.Find(CycleNodeIndex);
+				break;
+			}
+		}
+
+		// Finding the cycle length by finding the other occurence of cycle-starting node
+		const int32 CycleLength = MakeArrayView(&CycleIndices[CycleStartElementIndex + 1], CycleIndices.Num() - CycleStartElementIndex - 1).Find(CycleIndices[CycleStartElementIndex]) + 1;
+		check(CycleLength > 0);
+
+		TConstArrayView<int32> CycleView = MakeArrayView(&CycleIndices[CycleStartElementIndex], CycleLength);
+		// Find the deterministic cycle start, only used for hash generation.
+		const int32* MinElementIndex = Algo::MinElement(CycleView);
+		check(MinElementIndex);
+		const int32 LowestCycleElementIndex = CycleView.Find(*MinElementIndex);
+
+		// Calculate cycle's hash
+		int32 ElementIndex = LowestCycleElementIndex;
+		uint32 CycleHash = static_cast<uint32>(CycleView[ElementIndex++]);
+		for (int32 CycleCounter = 1; CycleCounter < CycleView.Num(); ++CycleCounter)
+		{
+			ElementIndex %= CycleView.Num();
+			CycleHash = HashCombine(CycleHash, CycleView[ElementIndex++]);
+		}
+
+		if (InOutReportedCycleHashes.Find(CycleHash) == INDEX_NONE)
+		{
+			InOutReportedCycleHashes.Add(CycleHash);
+
+			UE_LOG(LogMass, Error, TEXT("Detected processing dependency cycle:"));
+
+			for (const int32 CycleNodeIndex : CycleView)
+			{
+				if (const UMassProcessor* Processor = AllNodes[CycleNodeIndex].Processor)
+				{
+					UE_LOG(LogMass, Warning, TEXT("\t%s, group: %s, before: %s, after %s")
+						, *Processor->GetName()
+						, *Processor->GetExecutionOrder().ExecuteInGroup.ToString()
+						, *NameViewToString(Processor->GetExecutionOrder().ExecuteBefore)
+						, *NameViewToString(Processor->GetExecutionOrder().ExecuteAfter));
+				}
+				else
+				{
+					// group
+					UE_LOG(LogMass, Warning, TEXT("\tGroup %s"), *AllNodes[CycleNodeIndex].Name.ToString());
+				}
+			}
+		}
+	}
+#endif // WITH_MASSENTITY_DEBUG
 }
 
 //----------------------------------------------------------------------//
@@ -261,12 +330,14 @@ void FMassProcessorDependencySolver::FResourceUsage::SubmitNode(const int32 Node
 //----------------------------------------------------------------------//
 //  FProcessorDependencySolver::FNode
 //----------------------------------------------------------------------//
-void FMassProcessorDependencySolver::FNode::IncreaseWaitingNodesCount(TArrayView<FMassProcessorDependencySolver::FNode> InAllNodes)
+bool FMassProcessorDependencySolver::FNode::IncreaseWaitingNodesCount(TArrayView<FMassProcessorDependencySolver::FNode> InAllNodes
+	, const int32 IterationsLimit, TArray<int32>& OutCycleIndices)
 {
 	// cycle-protection check. If true it means we have a cycle and the whole algorithm result will be unreliable 
-	if (TotalWaitingNodes >= FMath::Square(InAllNodes.Num()))
+	if (IterationsLimit < 0)
 	{
-		return;
+		OutCycleIndices.Add(NodeIndex);
+		return false;
 	}
 
 	++TotalWaitingNodes;
@@ -274,8 +345,14 @@ void FMassProcessorDependencySolver::FNode::IncreaseWaitingNodesCount(TArrayView
 	for (const int32 DependencyIndex : OriginalDependencies)
 	{
 		check(&InAllNodes[DependencyIndex] != this);
-		InAllNodes[DependencyIndex].IncreaseWaitingNodesCount(InAllNodes);
+		if (InAllNodes[DependencyIndex].IncreaseWaitingNodesCount(InAllNodes, IterationsLimit - 1, OutCycleIndices) == false)
+		{
+			OutCycleIndices.Add(NodeIndex);
+			return false;
+		}
 	}
+
+	return true;
 }
 //----------------------------------------------------------------------//
 //  FProcessorDependencySolver
@@ -596,21 +673,38 @@ void FMassProcessorDependencySolver::Solve(TArray<FMassProcessorOrderInfo>& OutR
 		return;
 	}
 
+	// for more efficient cycle detection and breaking it will be useful to know how many
+	// nodes do not depend on anything - we can use this number as a limit for the longest dependency chain
+	int32 TotalDependingNodes = 0;
+
 	for (FNode& Node : AllNodes)
 	{
 		Node.TransientDependencies = Node.OriginalDependencies;
 		Node.TotalWaitingNodes = 0;
+		TotalDependingNodes += (Node.OriginalDependencies.Num() > 0) ? 1 : 0;
 	}
 
+	TArray<int32> CycleIndices;
+#if WITH_MASSENTITY_DEBUG
+	TArray<uint32> ReportedCycleHashes;
+#endif // WITH_MASSENTITY_DEBUG
+
 	TArray<int32> IndicesRemaining;
 	IndicesRemaining.Reserve(AllNodes.Num());
-	for (int32 i = 0; i < AllNodes.Num(); ++i)
+	for (int32 NodeIndex = 0; NodeIndex < AllNodes.Num(); ++NodeIndex)
 	{
 		// skip all the group nodes, all group dependencies have already been converted to individual processor dependencies
-		if (AllNodes[i].IsGroup() == false)
+		if (AllNodes[NodeIndex].IsGroup() == false)
 		{
-			IndicesRemaining.Add(i);
-			AllNodes[i].IncreaseWaitingNodesCount(AllNodes);
+			IndicesRemaining.Add(NodeIndex);
+			if (AllNodes[NodeIndex].IncreaseWaitingNodesCount(AllNodes, TotalDependingNodes, CycleIndices) == false)
+			{
+#if WITH_MASSENTITY_DEBUG
+				// we have a cycle. Report it here
+				UE::Mass::Private::LogCycle(AllNodes, CycleIndices, ReportedCycleHashes);
+#endif // WITH_MASSENTITY_DEBUG
+				CycleIndices.Reset();
+			}
 		}
 	}
 
@@ -630,27 +724,7 @@ void FMassProcessorDependencySolver::Solve(TArray<FMassProcessorOrderInfo>& OutR
 
 		if (bStepSuccessful == false)
 		{
-			bAnyCyclesDetected = true;
-
-			UE_LOG(LogMass, Error, TEXT("Detected processing dependency cycle:"));
-			for (const int32 Index : IndicesRemaining)
-			{
-				UMassProcessor* Processor = AllNodes[Index].Processor;
-				if (Processor)
-				{
-					UE_LOG(LogMass, Log, TEXT("\t%s, group: %s, before: %s, after %s")
-						, *Processor->GetName()
-						, *Processor->GetExecutionOrder().ExecuteInGroup.ToString()
-						, *NameViewToString(Processor->GetExecutionOrder().ExecuteBefore)
-						, *NameViewToString(Processor->GetExecutionOrder().ExecuteAfter));
-				}
-				else
-				{
-					// group
-					UE_LOG(LogMass, Log, TEXT("\tGroup %s"), *AllNodes[Index].Name.ToString());
-				}
-			}
-			UE_LOG(LogMass, Log, TEXT("Cutting the chain at an arbitrary location."));
+			UE_LOG(LogMass, Error, TEXT("Encountered processing dependency cycle - cutting the chain at an arbitrary location."));
 
 			// remove first dependency
 			// note that if we're in a cycle handling scenario every node does have some dependencies left
@@ -696,8 +770,6 @@ void FMassProcessorDependencySolver::ResolveDependencies(TArray<FMassProcessorOr
 		DependencyGraphFileName = InOutOptionalResult->DependencyGraphFileName;
 	}
 
-	bAnyCyclesDetected = false;
-
 	UE_LOG(LogMass, Log, TEXT("Gathering dependencies data:"));
 
 	AllNodes.Reset();
diff --git a/Engine/Source/Runtime/MassEntity/Public/MassProcessorDependencySolver.h b/Engine/Source/Runtime/MassEntity/Public/MassProcessorDependencySolver.h
--- a/Engine/Source/Runtime/MassEntity/Public/MassProcessorDependencySolver.h
+++ b/Engine/Source/Runtime/MassEntity/Public/MassProcessorDependencySolver.h
@@ -83,7 +83,6 @@ struct MASSENTITY_API FMassExecutionRequirements
 
 struct FMassProcessorDependencySolver
 {
-private:
 	struct FNode
 	{
 		FNode(const FName InName, UMassProcessor* InProcessor, const int32 InNodeIndex = INDEX_NONE) 
@@ -91,7 +90,11 @@ private:
 		{}
 
 		bool IsGroup() const { return Processor == nullptr; }
-		void IncreaseWaitingNodesCount(TArrayView<FNode> InAllNodes);
+		/**
+		 * @return `true` when everything's fine, `false` when cycles have been encountered. If that
+		 *	happens `OutCycleIndices` gets filled with the relevant node indices.
+		 */
+		bool IncreaseWaitingNodesCount(TArrayView<FNode> InAllNodes, const int32 IterationsLimit, TArray<int32>& OutCycleIndices);
 
 		FName Name = TEXT("");
 		UMassProcessor* Processor = nullptr;
@@ -112,6 +115,7 @@ private:
 		TArray<FMassArchetypeHandle> ValidArchetypes;
 	};
 
+private:
 	struct FResourceUsage
 	{
 		FResourceUsage(const TArray<FNode>& InAllNodes);
@@ -198,7 +202,7 @@ protected:
 	void GatherSubsystemInformation(const UE::Mass::FTypeManager& TypeManager);
 	
 	TArrayView<UMassProcessor* const> Processors;
-	bool bAnyCyclesDetected = false;
+
 	/**
 	 * indicates whether we're generating processor order to be run in single-threaded or multithreaded environment (usually
 	 * this means Dedicated Server vs Any other configuration). In Single-Threaded mode we can skip a bunch of expensive,
