[MassEntity]
* Fixed non unity
* Moved dll storage to right place

diff --git a/Engine/Source/Runtime/MassEntity/Public/MassRequirementAccessDetector.h b/Engine/Source/Runtime/MassEntity/Public/MassRequirementAccessDetector.h
--- a/Engine/Source/Runtime/MassEntity/Public/MassRequirementAccessDetector.h
+++ b/Engine/Source/Runtime/MassEntity/Public/MassRequirementAccessDetector.h
@@ -10,12 +10,13 @@
 
 
 struct FMassEntityManager;
+struct FMassEntityQuery;
 
-struct MASSENTITY_API FMassRequirementAccessDetector
+struct FMassRequirementAccessDetector
 {
-	void Initialize();
-	void RequireAccess(const FMassEntityQuery& Query);
-	void ReleaseAccess(const FMassEntityQuery& Query);
+	MASSENTITY_API void Initialize();
+	MASSENTITY_API void RequireAccess(const FMassEntityQuery& Query);
+	MASSENTITY_API void ReleaseAccess(const FMassEntityQuery& Query);
 
 private:
 	using FDetectorMethod = bool (FRWAccessDetector::*)() const;
@@ -71,7 +72,7 @@ private:
 	}
 
 	/** @Note the function is not thread-safe and meant to be only called internally on game thread (see FMassRequirementAccessDetector::Initialize) */
-	void AddDetectors(const FStructTracker& StructTracker);
+	MASSENTITY_API void AddDetectors(const FStructTracker& StructTracker);
 
 	TMap<const UStruct*, TSharedRef<FRWAccessDetector>> Detectors;
 };
