[MassEntity]
* Second attempt at fixing non unity compile error

diff --git a/Engine/Source/Runtime/MassEntity/Public/MassRequirementAccessDetector.h b/Engine/Source/Runtime/MassEntity/Public/MassRequirementAccessDetector.h
--- a/Engine/Source/Runtime/MassEntity/Public/MassRequirementAccessDetector.h
+++ b/Engine/Source/Runtime/MassEntity/Public/MassRequirementAccessDetector.h
@@ -10,7 +10,6 @@
 
 
 struct FMassEntityManager;
-struct FMassEntityQuery;
 
 struct FMassRequirementAccessDetector
 {
@@ -77,6 +76,8 @@ private:
 	TMap<const UStruct*, TSharedRef<FRWAccessDetector>> Detectors;
 };
 
+#else
+struct FMassEntityQuery;
 #endif // WITH_MASSENTITY_DEBUG
 
 namespace UE::Mass::Debug
