[Mass] renamed more EntitySystem to EntityManager for instances of FMassEntityManager

#rnx

diff --git a/Engine/Plugins/AI/MassCrowd/Source/MassCrowd/Public/MassCrowdRepresentationActorManagement.h b/Engine/Plugins/AI/MassCrowd/Source/MassCrowd/Public/MassCrowdRepresentationActorManagement.h
--- a/Engine/Plugins/AI/MassCrowd/Source/MassCrowd/Public/MassCrowdRepresentationActorManagement.h
+++ b/Engine/Plugins/AI/MassCrowd/Source/MassCrowd/Public/MassCrowdRepresentationActorManagement.h
@@ -37,7 +37,7 @@ protected:
 	 * @param Priority of this spawn request in comparison with the others, lower value means higher priority
 	 * @return the actor spawned
 	 */
-	virtual AActor* GetOrSpawnActor(UMassRepresentationSubsystem& RepresentationSubsystem, FMassEntityManager& EntitySubsystem
+	virtual AActor* GetOrSpawnActor(UMassRepresentationSubsystem& RepresentationSubsystem, FMassEntityManager& EntityManager
 		, const FMassEntityHandle MassAgent, const FTransform& Transform, const int16 TemplateActorIndex
 		, FMassActorSpawnRequestHandle& SpawnRequestHandle, const float Priority) const override;
 
diff --git a/Engine/Plugins/Runtime/MassGameplay/Source/MassGameplayDebug/Public/MassGameplayDebugTypes.h b/Engine/Plugins/Runtime/MassGameplay/Source/MassGameplayDebug/Public/MassGameplayDebugTypes.h
--- a/Engine/Plugins/Runtime/MassGameplay/Source/MassGameplayDebug/Public/MassGameplayDebugTypes.h
+++ b/Engine/Plugins/Runtime/MassGameplay/Source/MassGameplayDebug/Public/MassGameplayDebugTypes.h
@@ -33,8 +33,8 @@ namespace UE::Mass::Debug
 	 * Fetches entity handles and their locations for entities indicated by index range as set by
 	 * mass.debug.SetDebugEntityRange or mass.debug.DebugEntity console commands.
 	 */
-	MASSGAMEPLAYDEBUG_API extern void GetDebugEntitiesAndLocations(const FMassEntityManager& EntitySubsystem, TArray<FMassEntityHandle>& OutEntities, TArray<FVector>& OutLocations);
-	MASSGAMEPLAYDEBUG_API extern FMassEntityHandle ConvertEntityIndexToHandle(const FMassEntityManager& EntitySubsystem, const int32 EntityIndex);
+	MASSGAMEPLAYDEBUG_API extern void GetDebugEntitiesAndLocations(const FMassEntityManager& EntityManager, TArray<FMassEntityHandle>& OutEntities, TArray<FVector>& OutLocations);
+	MASSGAMEPLAYDEBUG_API extern FMassEntityHandle ConvertEntityIndexToHandle(const FMassEntityManager& EntityManager, const int32 EntityIndex);
 } // namespace UE::Mass::Debug
 #endif // WITH_MASSGAMEPLAY_DEBUG
 
diff --git a/Engine/Source/Editor/MassEntityDebugger/Private/MassDebuggerModel.h b/Engine/Source/Editor/MassEntityDebugger/Private/MassDebuggerModel.h
--- a/Engine/Source/Editor/MassEntityDebugger/Private/MassDebuggerModel.h
+++ b/Engine/Source/Editor/MassEntityDebugger/Private/MassDebuggerModel.h
@@ -212,7 +212,7 @@ struct FMassDebuggerModel
 	TWeakPtr<SMassEntitiesView> ShowEntitiesView(int32 Index);
 
 protected:
-	void StoreArchetypes(const FMassEntityManager& EntitySubsystem, TMap<FMassArchetypeHandle, TSharedPtr<FMassDebuggerArchetypeData>>& OutTransientArchetypesMap);
+	void StoreArchetypes(const FMassEntityManager& EntityManager, TMap<FMassArchetypeHandle, TSharedPtr<FMassDebuggerArchetypeData>>& OutTransientArchetypesMap);
 
 	void ResetSelectedArchetypes();
 	void ResetSelectedProcessors();
diff --git a/Engine/Source/Runtime/MassEntity/Public/MassArchetypeTypes.h b/Engine/Source/Runtime/MassEntity/Public/MassArchetypeTypes.h
--- a/Engine/Source/Runtime/MassEntity/Public/MassArchetypeTypes.h
+++ b/Engine/Source/Runtime/MassEntity/Public/MassArchetypeTypes.h
@@ -199,7 +199,7 @@ struct MASSENTITY_API FMassArchetypeEntityCollectionWithPayload
 	{
 	}
 
-	static void CreateEntityRangesWithPayload(const FMassEntityManager& EntitySubsystem, const TConstArrayView<FMassEntityHandle> Entities
+	static void CreateEntityRangesWithPayload(const FMassEntityManager& EntityManager, const TConstArrayView<FMassEntityHandle> Entities
 		, const FMassArchetypeEntityCollection::EDuplicatesHandling DuplicatesHandling, FMassGenericPayloadView Payload
 		, TArray<FMassArchetypeEntityCollectionWithPayload>& OutEntityCollections);
 
diff --git a/Engine/Source/Runtime/MassEntity/Public/MassDebugger.h b/Engine/Source/Runtime/MassEntity/Public/MassDebugger.h
--- a/Engine/Source/Runtime/MassEntity/Public/MassDebugger.h
+++ b/Engine/Source/Runtime/MassEntity/Public/MassDebugger.h
@@ -172,14 +172,14 @@ struct MASSENTITY_API FMassDebugger
 	static TConstArrayView<FMassEntityQuery*> GetProcessorQueries(const UMassProcessor& Processor);
 	/** fetches all queries registered for given Processor. Note that in order to get up to date information
 	 *  FMassEntityQuery::CacheArchetypes will be called on each query */
-	static TConstArrayView<FMassEntityQuery*> GetUpToDateProcessorQueries(const FMassEntityManager& EntitySubsystem, UMassProcessor& Processor);
+	static TConstArrayView<FMassEntityQuery*> GetUpToDateProcessorQueries(const FMassEntityManager& EntityManager, UMassProcessor& Processor);
 
 	static UE::Mass::Debug::FQueryRequirementsView GetQueryRequirements(const FMassEntityQuery& Query);
 	static void GetQueryExecutionRequirements(const FMassEntityQuery& Query, FMassExecutionRequirements& OutExecutionRequirements);
 	static TArray<FMassEntityHandle> GetEntitiesMatchingQuery(const FMassEntityManager& EntityManager, const FMassEntityQuery& Query);
 
-	static void ForEachArchetype(const FMassEntityManager& EntitySubsystem, const UE::Mass::Debug::FArchetypeFunction& Function);
-	static TArray<FMassArchetypeHandle> GetAllArchetypes(const FMassEntityManager& EntitySubsystem);
+	static void ForEachArchetype(const FMassEntityManager& EntityManager, const UE::Mass::Debug::FArchetypeFunction& Function);
+	static TArray<FMassArchetypeHandle> GetAllArchetypes(const FMassEntityManager& EntityManager);
 	static const FMassArchetypeCompositionDescriptor& GetArchetypeComposition(const FMassArchetypeHandle& ArchetypeHandle);
 
 	static uint64 GetArchetypeTraceID(const FMassArchetypeData& ArchetypeData);
diff --git a/Engine/Source/Runtime/MassEntity/Public/MassEntityUtils.h b/Engine/Source/Runtime/MassEntity/Public/MassEntityUtils.h
--- a/Engine/Source/Runtime/MassEntity/Public/MassEntityUtils.h
+++ b/Engine/Source/Runtime/MassEntity/Public/MassEntityUtils.h
@@ -25,7 +25,7 @@ uint8 DetermineProcessorSupportedTickTypes(const UWorld* World);
  * Fills OutEntityCollections with per-archetype FMassArchetypeEntityCollection instances. 
  * @param DuplicatesHandling used to inform the function whether to expect duplicates.
  */
-MASSENTITY_API extern void CreateEntityCollections(const FMassEntityManager& EntitySubsystem, const TConstArrayView<FMassEntityHandle> Entities
+MASSENTITY_API extern void CreateEntityCollections(const FMassEntityManager& EntityManager, const TConstArrayView<FMassEntityHandle> Entities
 	, const FMassArchetypeEntityCollection::EDuplicatesHandling DuplicatesHandling, TArray<FMassArchetypeEntityCollection>& OutEntityCollections);
 
 /**
