[Mass] Massaged code around "Creation Context" to lay foundations to the following change that will fix observes issues when modifying entities, that are not being created right now, while the creation context is active.

The main change here is renaming FMassEntityManager::FEntityCreationContext to FMassEntityManager::FObserverLock, plus some deprecation

diff --git a/Engine/Source/Runtime/MassEntity/Private/MassEntityManager.cpp b/Engine/Source/Runtime/MassEntity/Private/MassEntityManager.cpp
--- a/Engine/Source/Runtime/MassEntity/Private/MassEntityManager.cpp
+++ b/Engine/Source/Runtime/MassEntity/Private/MassEntityManager.cpp
@@ -46,32 +46,33 @@ namespace UE::Mass::Private
 }
 
 //-----------------------------------------------------------------------------
-// FMassEntityManager::FEntityCreationContext
+// FMassEntityManager::FObserverLock
 //-----------------------------------------------------------------------------
-FMassEntityManager::FEntityCreationContext::FEntityCreationContext()
+FMassEntityManager::FObserverLock::FObserverLock()
 	: OwnerThreadId(FPlatformTLS::GetCurrentThreadId())
 {	
 }
 
-FMassEntityManager::FEntityCreationContext::FEntityCreationContext(FMassEntityManager& InManager, const TConstArrayView<FMassEntityHandle> InCreatedEntities)
-	: FEntityCreationContext()
+FMassEntityManager::FObserverLock::FObserverLock(FMassEntityManager& InManager, const TConstArrayView<FMassEntityHandle> InCreatedEntities)
+	: FObserverLock()
 {
 	CreatedEntities = InCreatedEntities;
 	Manager = InManager.AsShared();
+	bDuringCreation = (InCreatedEntities.Num() > 0);
 }
 
-FMassEntityManager::FEntityCreationContext::FEntityCreationContext(FMassEntityManager& InManager, const TConstArrayView<FMassEntityHandle> InCreatedEntities
+FMassEntityManager::FObserverLock::FObserverLock(FMassEntityManager& InManager, const TConstArrayView<FMassEntityHandle> InCreatedEntities
 	, FMassArchetypeEntityCollection&& InEntityCollection)
-	: FEntityCreationContext(InManager, InCreatedEntities)
+	: FObserverLock(InManager, InCreatedEntities)
 {
-	checkf(InCreatedEntities.IsEmpty() || InEntityCollection.IsEmpty() == false, TEXT("Trying to create FEntityCreationContext instance with no entities but non-empty entity collection. This is not supported."))
+	checkf(InCreatedEntities.IsEmpty() || InEntityCollection.IsEmpty() == false, TEXT("Trying to create FObserverLock instance with no entities but non-empty entity collection. This is not supported."))
 	if (InCreatedEntities.IsEmpty() == false)
 	{
 		EntityCollections.Add(MoveTemp(InEntityCollection));
 	}
 }
 
-FMassEntityManager::FEntityCreationContext::~FEntityCreationContext()
+FMassEntityManager::FObserverLock::~FObserverLock()
 {
 	if (((EntityCollections.IsEmpty() == false) || (CreatedEntities.IsEmpty() == false)) && ensure(Manager))
 	{
@@ -79,7 +80,7 @@ FMassEntityManager::FEntityCreationContext::~FEntityCreationContext()
 	}
 }
 
-TConstArrayView<FMassArchetypeEntityCollection> FMassEntityManager::FEntityCreationContext::GetEntityCollections() const
+TConstArrayView<FMassArchetypeEntityCollection> FMassEntityManager::FObserverLock::GetEntityCollections() const
 {
 	// the EntityCollection has been dirtied, we need to rebuild it
 	if (IsDirty() && ensure(Manager))
@@ -90,16 +91,16 @@ TConstArrayView<FMassArchetypeEntityCollection> FMassEntityManager::FEntityCreat
 	return EntityCollections;
 }
 
-void FMassEntityManager::FEntityCreationContext::MarkDirty()
+void FMassEntityManager::FObserverLock::MarkDirty()
 {
-	checkf(OwnerThreadId == FPlatformTLS::GetCurrentThreadId(), TEXT("%hs: all FEntityCreationContext operations ere expected to be run in a single thread"), __FUNCTION__);
+	checkf(OwnerThreadId == FPlatformTLS::GetCurrentThreadId(), TEXT("%hs: all FObserverLock operations ere expected to be run in a single thread"), __FUNCTION__);
 
 	EntityCollections.Reset();
 }
 
-void FMassEntityManager::FEntityCreationContext::AppendEntities(const TConstArrayView<FMassEntityHandle> InEntitiesToAppend)
+void FMassEntityManager::FObserverLock::AppendCreatedEntities(const TConstArrayView<FMassEntityHandle> InEntitiesToAppend)
 {
-	checkf(OwnerThreadId == FPlatformTLS::GetCurrentThreadId(), TEXT("%hs: all FEntityCreationContext operations ere expected to be run in a single thread"), __FUNCTION__);
+	checkf(OwnerThreadId == FPlatformTLS::GetCurrentThreadId(), TEXT("%hs: all FObserverLock operations ere expected to be run in a single thread"), __FUNCTION__);
 
 	if (InEntitiesToAppend.Num())
 	{
@@ -119,9 +120,9 @@ void FMassEntityManager::FEntityCreationContext::AppendEntities(const TConstArra
 	}
 }
 
-void FMassEntityManager::FEntityCreationContext::AppendEntities(const TConstArrayView<FMassEntityHandle> InEntitiesToAppend, FMassArchetypeEntityCollection&& InEntityCollection)
+void FMassEntityManager::FObserverLock::AppendCreatedEntities(const TConstArrayView<FMassEntityHandle> InEntitiesToAppend, FMassArchetypeEntityCollection&& InEntityCollection)
 {
-	checkf(OwnerThreadId == FPlatformTLS::GetCurrentThreadId(), TEXT("%hs: all FEntityCreationContext operations are expected to be run in a single thread"), __FUNCTION__);
+	checkf(OwnerThreadId == FPlatformTLS::GetCurrentThreadId(), TEXT("%hs: all FObserverLock operations are expected to be run in a single thread"), __FUNCTION__);
 
 	if (InEntitiesToAppend.Num() == 0)
 	{
@@ -130,9 +131,9 @@ void FMassEntityManager::FEntityCreationContext::AppendEntities(const TConstArra
 
 	LLM_SCOPE_BYNAME(TEXT("Mass/EntityManager"));
 
-	AppendEntities(InEntitiesToAppend);
+	AppendCreatedEntities(InEntitiesToAppend);
 
-	// this condition boils down to checking if this FEntityCreationContext instance only connects the just added EntitiesToAppend
+	// this condition boils down to checking if this FObserverLock instance only connects the just added EntitiesToAppend
 	if (CreatedEntities.Num() == InEntitiesToAppend.Num())
 	{
 		checkf(EntityCollections.Num() == 0, TEXT("We never expect EntityCollections to be non-empty while there are no entities in CreatedEntities."));
@@ -140,7 +141,7 @@ void FMassEntityManager::FEntityCreationContext::AppendEntities(const TConstArra
 	}
 }
 
-void FMassEntityManager::FEntityCreationContext::ForceUpdateCurrentThreadID()
+void FMassEntityManager::FObserverLock::ForceUpdateCurrentThreadID()
 {
 	OwnerThreadId = FPlatformTLS::GetCurrentThreadId();
 }
@@ -380,7 +381,7 @@ void FMassEntityManager::OnPostFork(EForkProcessRole Role)
 			}
 		}
 
-		if (TSharedPtr<FEntityCreationContext> ActiveContext = ActiveCreationContext.Pin())
+		if (TSharedPtr<FObserverLock> ActiveContext = ActiveObserverLock.Pin())
 		{
 			ActiveContext->ForceUpdateCurrentThreadID();
 		}
@@ -683,8 +684,8 @@ FMassEntityHandle FMassEntityManager::CreateEntity(TConstArrayView<FInstancedStr
 	const FMassEntityHandle Entity = ReserveEntity();
 
 	// Using a creation context to prevent InternalBuildEntity from notifying observers before we set fragments data
-	const TSharedRef<FEntityCreationContext> CreationContext = GetOrMakeCreationContext();
-	CreationContext->AppendEntities({Entity});
+	const TSharedRef<FObserverLock> ObserverLock = GetOrMakeObserverLock();
+	ObserverLock->AppendCreatedEntities({Entity});
 
 	InternalBuildEntity(Entity, ArchetypeHandle, SharedFragmentValues);
 
@@ -739,8 +740,8 @@ void FMassEntityManager::BuildEntity(FMassEntityHandle Entity, TConstArrayView<F
 	check(ArchetypeHandle.IsValid());
 
 	// Using a creation context to prevent InternalBuildEntity from notifying observers before we set fragments data
-	const TSharedRef<FEntityCreationContext> CreationContext = GetOrMakeCreationContext();
-	CreationContext->AppendEntities({Entity});
+	const TSharedRef<FObserverLock> ObserverLock = GetOrMakeObserverLock();
+	ObserverLock->AppendCreatedEntities({Entity});
 
 	InternalBuildEntity(Entity, ArchetypeHandle, SharedFragmentValues);
 
@@ -763,7 +764,7 @@ int32 FMassEntityManager::BatchReserveEntities(TArrayView<FMassEntityHandle> InO
 	return GetEntityStorageInterface().Acquire(InOutEntities);
 }
 
-TSharedRef<FMassEntityManager::FEntityCreationContext> FMassEntityManager::BatchBuildEntities(const FMassArchetypeEntityCollectionWithPayload& EncodedEntitiesWithPayload
+TSharedRef<FMassEntityManager::FObserverLock> FMassEntityManager::BatchBuildEntities(const FMassArchetypeEntityCollectionWithPayload& EncodedEntitiesWithPayload
 	, const FMassFragmentBitSet& FragmentsAffected, const FMassArchetypeSharedFragmentValues& SharedFragmentValues, const FMassArchetypeCreationParams& CreationParams)
 {
 	CHECK_SYNC_API();
@@ -782,7 +783,7 @@ TSharedRef<FMassEntityManager::FEntityCreationContext> FMassEntityManager::Batch
 	return BatchBuildEntities(EncodedEntitiesWithPayload, MoveTemp(Composition), SharedFragmentValues, CreationParams);
 }
 
-TSharedRef<FMassEntityManager::FEntityCreationContext> FMassEntityManager::BatchBuildEntities(const FMassArchetypeEntityCollectionWithPayload& EncodedEntitiesWithPayload
+TSharedRef<FMassEntityManager::FObserverLock> FMassEntityManager::BatchBuildEntities(const FMassArchetypeEntityCollectionWithPayload& EncodedEntitiesWithPayload
 	, FMassArchetypeCompositionDescriptor&& Composition
 	, const FMassArchetypeSharedFragmentValues& SharedFragmentValues, const FMassArchetypeCreationParams& CreationParams)
 {
@@ -836,11 +837,11 @@ TSharedRef<FMassEntityManager::FEntityCreationContext> FMassEntityManager::Batch
 	// and as such won't cause observers triggering (which usually is prevented by context's existence), and that we 
 	// strongly assume the all entity creation/building (not to be mistaken with "reserving") takes place in a single thread
 	// @todo add checks/ensures enforcing the assumption mentioned above.
-	return GetOrMakeCreationContext(EntityHandles
+	return GetOrMakeObserverLock(EntityHandles
 		, FMassArchetypeEntityCollection(TargetArchetypeHandle, MoveTemp(TargetArchetypeEntityRanges)));
 }
 
-TSharedRef<FMassEntityManager::FEntityCreationContext> FMassEntityManager::BatchCreateReservedEntities(const FMassArchetypeHandle& ArchetypeHandle
+TSharedRef<FMassEntityManager::FObserverLock> FMassEntityManager::BatchCreateReservedEntities(const FMassArchetypeHandle& ArchetypeHandle
 	, const FMassArchetypeSharedFragmentValues& SharedFragmentValues, TConstArrayView<FMassEntityHandle> ReservedEntities)
 {
 	CHECK_SYNC_API();
@@ -853,11 +854,11 @@ TSharedRef<FMassEntityManager::FEntityCreationContext> FMassEntityManager::Batch
 		, SharedFragmentValues, ReservedEntities);
 }
 
-TSharedRef<FMassEntityManager::FEntityCreationContext> FMassEntityManager::BatchCreateEntities(const FMassArchetypeHandle& ArchetypeHandle
+TSharedRef<FMassEntityManager::FObserverLock> FMassEntityManager::BatchCreateEntities(const FMassArchetypeHandle& ArchetypeHandle
 	, const FMassArchetypeSharedFragmentValues& SharedFragmentValues, const int32 Count, TArray<FMassEntityHandle>& InOutEntities)
 {
 	CHECK_SYNC_API();
-	testableCheckfReturn(ArchetypeHandle.IsValid(), return GetOrMakeCreationContext()
+	testableCheckfReturn(ArchetypeHandle.IsValid(), return GetOrMakeObserverLock()
 		, TEXT("%hs expecting a valid ArchetypeHandle"), __FUNCTION__);
 
 	TRACE_CPUPROFILER_EVENT_SCOPE(Mass_BatchCreateEntities);
@@ -869,7 +870,7 @@ TSharedRef<FMassEntityManager::FEntityCreationContext> FMassEntityManager::Batch
 		, SharedFragmentValues, ReservedEntities);
 }
 
-TSharedRef<FMassEntityManager::FEntityCreationContext> FMassEntityManager::InternalBatchCreateReservedEntities(const FMassArchetypeHandle& ArchetypeHandle
+TSharedRef<FMassEntityManager::FObserverLock> FMassEntityManager::InternalBatchCreateReservedEntities(const FMassArchetypeHandle& ArchetypeHandle
 	, const FMassArchetypeSharedFragmentValues& SharedFragmentValues, TConstArrayView<FMassEntityHandle> ReservedEntities)
 {
 	// Functions calling into this one are required to verify that the archetype handle is valid
@@ -890,7 +891,7 @@ TSharedRef<FMassEntityManager::FEntityCreationContext> FMassEntityManager::Inter
 
 	UE_TRACE_MASS_ENTITIES_CREATED(ReservedEntities, *ArchetypeData)
 
-	return GetOrMakeCreationContext(ReservedEntities, FMassArchetypeEntityCollection(ArchetypeHandle, MoveTemp(TargetArchetypeEntityRanges)));
+	return GetOrMakeObserverLock(ReservedEntities, FMassArchetypeEntityCollection(ArchetypeHandle, MoveTemp(TargetArchetypeEntityRanges)));
 }
 
 void FMassEntityManager::DestroyEntity(FMassEntityHandle Entity)
@@ -1644,7 +1645,7 @@ void FMassEntityManager::BatchChangeTagsForEntities(TConstArrayView<FMassArchety
 
 	TRACE_CPUPROFILER_EVENT_SCOPE(Mass_BatchChangeTagsForEntities);
 
-	const FScopedCreationContextOperations CreationContextOperations(*this);
+	const FScopedObserverLockOperations ObserverLockOperations(*this);
 	
 	for (const FMassArchetypeEntityCollection& Collection : EntityCollections)
 	{
@@ -1658,7 +1659,7 @@ void FMassEntityManager::BatchChangeTagsForEntities(TConstArrayView<FMassArchety
 			FMassTagBitSet TagsAdded = TagsToAdd - CurrentArchetype->GetTagBitSet();
 			FMassTagBitSet TagsRemoved = TagsToRemove.GetOverlap(CurrentArchetype->GetTagBitSet());
 
-			if (CreationContextOperations.IsAllowedToTriggerObservers()
+			if (ObserverLockOperations.IsAllowedToTriggerObservers()
 				&& ObserverManager.HasObserversForBitSet(TagsRemoved, EMassObservedOperation::Remove))
 			{
 				// @todo should use OnPreCompositionRemoved here instead, but we're missing a FMassArchetypeEntityCollection version
@@ -1682,7 +1683,7 @@ void FMassEntityManager::BatchChangeTagsForEntities(TConstArrayView<FMassArchety
 				GetEntityStorageInterface().SetArchetypeFromShared(Entity.Index, NewArchetypeHandle.DataPtr);
 			}
 
-			if (bTagsAddedAreObserved && CreationContextOperations.IsAllowedToTriggerObservers())
+			if (bTagsAddedAreObserved && ObserverLockOperations.IsAllowedToTriggerObservers())
 			{
 				// @todo should use OnPostCompositionAdded here instead, but we're missing a FMassArchetypeEntityCollection version
 				ObserverManager.OnCompositionChanged(
@@ -1700,7 +1701,7 @@ void FMassEntityManager::BatchChangeFragmentCompositionForEntities(TConstArrayVi
 
 	TRACE_CPUPROFILER_EVENT_SCOPE(Mass_BatchChangeFragmentCompositionForEntities);
 
-	FScopedCreationContextOperations CreationContextOperations(*this);
+	FScopedObserverLockOperations ObserverLockOperations(*this);
 
 	for (const FMassArchetypeEntityCollection& Collection : EntityCollections)
 	{
@@ -1717,7 +1718,7 @@ void FMassEntityManager::BatchChangeFragmentCompositionForEntities(TConstArrayVi
 				const bool bFragmentsAddedAreObserved = ObserverManager.HasObserversForBitSet(FragmentsAdded, EMassObservedOperation::Add);
 				FMassFragmentBitSet FragmentsRemoved = FragmentsToRemove.GetOverlap(CurrentArchetype->GetFragmentBitSet());
 				
-				if (CreationContextOperations.IsAllowedToTriggerObservers()
+				if (ObserverLockOperations.IsAllowedToTriggerObservers()
 					&& ObserverManager.HasObserversForBitSet(FragmentsRemoved, EMassObservedOperation::Remove))
 				{
 					ObserverManager.OnCompositionChanged(Collection, FMassArchetypeCompositionDescriptor(MoveTemp(FragmentsRemoved)), EMassObservedOperation::Remove);
@@ -1739,7 +1740,7 @@ void FMassEntityManager::BatchChangeFragmentCompositionForEntities(TConstArrayVi
 					GetEntityStorageInterface().SetArchetypeFromShared(Entity.Index, NewArchetypeHandle.DataPtr);
 				}
 
-				if (bFragmentsAddedAreObserved && CreationContextOperations.IsAllowedToTriggerObservers())
+				if (bFragmentsAddedAreObserved && ObserverLockOperations.IsAllowedToTriggerObservers())
 				{
 					ObserverManager.OnCompositionChanged(
 						FMassArchetypeEntityCollection(NewArchetypeHandle, MoveTemp(NewArchetypeEntityRanges))
@@ -1772,7 +1773,7 @@ void FMassEntityManager::BatchAddFragmentInstancesForEntities(TConstArrayView<FM
 	//		archetype. This could be utilized as a basis for optimization. To be investigated.
 	// 
 
-	FScopedCreationContextOperations CreationContextOperations(*this);
+	FScopedObserverLockOperations ObserverLockOperations(*this);
 
 	for (const FMassArchetypeEntityCollectionWithPayload& EntityRangesWithPayload : EntityCollections)
 	{
@@ -1820,7 +1821,7 @@ void FMassEntityManager::BatchAddFragmentInstancesForEntities(TConstArrayView<FM
 			// corresponds to the order in FMassArchetypeEntityCollectionWithPayload's payload
 			TargetArchetypeHandle.DataPtr->BatchSetFragmentValues(TargetArchetypeEntityRanges, EntityRangesWithPayload.GetPayload());
 			
-			if (bFragmentsAddedAreObserved && CreationContextOperations.IsAllowedToTriggerObservers())
+			if (bFragmentsAddedAreObserved && ObserverLockOperations.IsAllowedToTriggerObservers())
 			{
 				ObserverManager.OnCompositionChanged(
 					FMassArchetypeEntityCollection(TargetArchetypeHandle, MoveTemp(TargetArchetypeEntityRanges))
@@ -2105,44 +2106,51 @@ void FMassEntityManager::AppendCommands(TSharedPtr<FMassCommandBuffer>& InOutCom
 	Defer().MoveAppend(*InOutCommandBuffer.Get());
 }
 
-TSharedRef<FMassEntityManager::FEntityCreationContext> FMassEntityManager::GetOrMakeCreationContext()
+TSharedRef<FMassEntityManager::FObserverLock> FMassEntityManager::GetOrMakeObserverLock()
 {
-	if (ActiveCreationContext.IsValid())
+	if (ActiveObserverLock.IsValid())
 	{
-		return ActiveCreationContext.Pin().ToSharedRef();
+		return ActiveObserverLock.Pin().ToSharedRef();
 	}
 	else
 	{
-		FEntityCreationContext* CreationContext = new FEntityCreationContext(*this);
-		TSharedRef<FEntityCreationContext> SharedContext = MakeShareable(CreationContext);
-		ActiveCreationContext = SharedContext;
+		FObserverLock* ObserverLock = new FObserverLock(*this);
+		TSharedRef<FObserverLock> SharedContext = MakeShareable(ObserverLock);
+		ActiveObserverLock = SharedContext;
 		return SharedContext;
 	}
 }
 
-TSharedRef<FMassEntityManager::FEntityCreationContext> FMassEntityManager::GetOrMakeCreationContext(TConstArrayView<FMassEntityHandle> ReservedEntities
+TSharedRef<FMassEntityManager::FEntityCreationContext> FMassEntityManager::GetOrMakeCreationContext()
+{
+	TSharedRef<FObserverLock> ObserverLock = GetOrMakeObserverLock();
+	ObserverLock->bDuringCreation = true;
+	return ObserverLock;
+}
+
+TSharedRef<FMassEntityManager::FObserverLock> FMassEntityManager::GetOrMakeObserverLock(TConstArrayView<FMassEntityHandle> ReservedEntities
 	, FMassArchetypeEntityCollection&& EntityCollection)
 {
 	LLM_SCOPE_BYNAME(TEXT("Mass/EntityManager"));
-	if (ActiveCreationContext.IsValid())
+	if (ActiveObserverLock.IsValid())
 	{
-		TSharedPtr<FEntityCreationContext> SharedContext = ActiveCreationContext.Pin();
+		TSharedPtr<FObserverLock> SharedContext = ActiveObserverLock.Pin();
 		CA_ASSUME(SharedContext);
-		SharedContext->AppendEntities(ReservedEntities, MoveTemp(EntityCollection));
+		SharedContext->AppendCreatedEntities(ReservedEntities, MoveTemp(EntityCollection));
 		return SharedContext.ToSharedRef();
 	}
 	else
 	{
-		FEntityCreationContext* CreationContext = new FEntityCreationContext(*this, ReservedEntities, MoveTemp(EntityCollection));
-		TSharedRef<FEntityCreationContext> SharedContext = MakeShareable(CreationContext);
-		ActiveCreationContext = SharedContext;
+		FObserverLock* ObserverLock = new FObserverLock(*this, ReservedEntities, MoveTemp(EntityCollection));
+		TSharedRef<FObserverLock> SharedContext = MakeShareable(ObserverLock);
+		ActiveObserverLock = SharedContext;
 		return SharedContext;
 	}
 }
 
-bool FMassEntityManager::DirtyCreationContext()
+bool FMassEntityManager::DirtyObserverLock()
 {
-	if (TSharedPtr<FEntityCreationContext> AsSharedPtr = ActiveCreationContext.Pin())
+	if (TSharedPtr<FObserverLock> AsSharedPtr = ActiveObserverLock.Pin())
 	{
 		AsSharedPtr->MarkDirty();
 		return true;
@@ -2157,7 +2165,7 @@ UE::Mass::FEntityBuilder FMassEntityManager::MakeEntityBuilder()
 
 bool FMassEntityManager::DebugDoCollectionsOverlapCreationContext(TConstArrayView<FMassArchetypeEntityCollection> EntityCollections) const
 {
-	if (TSharedPtr<FEntityCreationContext> AsSharedPtr = ActiveCreationContext.Pin())
+	if (TSharedPtr<FObserverLock> AsSharedPtr = ActiveObserverLock.Pin())
 	{
 		TConstArrayView<FMassArchetypeEntityCollection> CreationCollections = AsSharedPtr->EntityCollections;
 		return CreationCollections.GetData() <= EntityCollections.GetData()
@@ -2337,7 +2345,7 @@ const UE::Mass::FStorageType& FMassEntityManager::DebugGetEntityStorageInterface
 // DEPRECATED
 //-----------------------------------------------------------------------------
 
-const FMassArchetypeEntityCollection& FMassEntityManager::FEntityCreationContext::GetEntityCollection() const
+const FMassArchetypeEntityCollection& FMassEntityManager::FObserverLock::GetEntityCollection() const
 {
 	static FMassArchetypeEntityCollection EmptyCollection;
 	return EntityCollections.Num() ? EntityCollections[0] : EmptyCollection;
diff --git a/Engine/Source/Runtime/MassEntity/Private/MassObserverManager.cpp b/Engine/Source/Runtime/MassEntity/Private/MassObserverManager.cpp
--- a/Engine/Source/Runtime/MassEntity/Private/MassObserverManager.cpp
+++ b/Engine/Source/Runtime/MassEntity/Private/MassObserverManager.cpp
@@ -162,7 +162,7 @@ bool FMassObserverManager::OnPostEntitiesCreated(FMassProcessingContext& Process
 
 bool FMassObserverManager::OnPostEntitiesCreated(FMassProcessingContext& ProcessingContext, TConstArrayView<FMassArchetypeEntityCollection> EntityCollections)
 {
-	if (!ensureMsgf(EntityManager.IsDuringEntityCreation() == false, TEXT("Entity creation not done yet.")))
+	if (!ensureMsgf(EntityManager.IsObserverExecutionLocked() == false, TEXT("Observes execution is locked right now.")))
 	{
 		return false;
 	}
diff --git a/Engine/Source/Runtime/MassEntity/Public/MassEntityManager.h b/Engine/Source/Runtime/MassEntity/Public/MassEntityManager.h
--- a/Engine/Source/Runtime/MassEntity/Public/MassEntityManager.h
+++ b/Engine/Source/Runtime/MassEntity/Public/MassEntityManager.h
@@ -240,31 +240,36 @@ public:
 	 * A dedicated structure for ensuring the "on entities creation" observers get notified only once all other 
 	 * initialization operations are done and this creation context instance gets released. 
 	 */
-	struct MASSENTITY_API FEntityCreationContext
+	struct MASSENTITY_API FObserverLock
 	{
 	private:
-		FEntityCreationContext();
-		explicit FEntityCreationContext(FMassEntityManager& InManager, const TConstArrayView<FMassEntityHandle> InCreatedEntities = {});
-		FEntityCreationContext(FMassEntityManager& InManager, const TConstArrayView<FMassEntityHandle> InCreatedEntities, FMassArchetypeEntityCollection&& InEntityCollection);
+		FObserverLock();
+		explicit FObserverLock(FMassEntityManager& InManager, const TConstArrayView<FMassEntityHandle> InCreatedEntities = {});
+		FObserverLock(FMassEntityManager& InManager, const TConstArrayView<FMassEntityHandle> InCreatedEntities, FMassArchetypeEntityCollection&& InEntityCollection);
 		
 	public:
-		~FEntityCreationContext();
+		~FObserverLock();
 
 		/** Returns EntityCollections, reconstructing them if needed (empty or dirtied). */
 		TConstArrayView<FMassArchetypeEntityCollection> GetEntityCollections() const;
-		int32 GetSpawnedNum() const { return CreatedEntities.Num(); }
+		int32 GetSpawnedNum() const;
 		void MarkDirty();
-		FORCEINLINE bool IsDirty() const { return EntityCollections.IsEmpty() && (CreatedEntities.IsEmpty() == false); }
-		void AppendEntities(const TConstArrayView<FMassEntityHandle> InEntitiesToAppend);
-		void AppendEntities(const TConstArrayView<FMassEntityHandle> InEntitiesToAppend, FMassArchetypeEntityCollection&& InEntityCollection);
+		bool IsDirty() const;
+		bool IsDuringCreation() const;
+		void AppendCreatedEntities(const TConstArrayView<FMassEntityHandle> InEntitiesToAppend);
+		void AppendCreatedEntities(const TConstArrayView<FMassEntityHandle> InEntitiesToAppend, FMassArchetypeEntityCollection&& InEntityCollection);
 
 		/** Function for debugging/testing purposes. We don't expect users to ever call it, always get collections via GetEntityCollections */
-		bool DebugAreEntityCollectionsUpToDate() const { return EntityCollections.IsEmpty() == CreatedEntities.IsEmpty(); }
+		bool DebugAreEntityCollectionsUpToDate() const;
 
 		UE_DEPRECATED(5.5, "This constructor is now deprecated and defunct. Use one of the others instead.")
-		explicit FEntityCreationContext(const int32 InNumSpawned) : FEntityCreationContext() {}
-		UE_DEPRECATED(5.5, "This function is now deprecated since FEntityCreationContext can contain more than a single collection now. Use GetEntityCollections instead.")
+		explicit FObserverLock(const int32) : FObserverLock() {}
+		UE_DEPRECATED(5.5, "This function is now deprecated since FObserverLock can contain more than a single collection now. Use GetEntityCollections instead.")
 		const FMassArchetypeEntityCollection& GetEntityCollection() const;
+		UE_DEPRECATED(5.6, "AppendEntities is deprecated. Call AppendCreatedEntities instead.")
+		void AppendEntities(const TConstArrayView<FMassEntityHandle> InEntitiesToAppend);
+		UE_DEPRECATED(5.6, "AppendEntities is deprecated. Call AppendCreatedEntities instead.")
+		void AppendEntities(const TConstArrayView<FMassEntityHandle> InEntitiesToAppend, FMassArchetypeEntityCollection&& InEntityCollection);
 
 	private:
 		friend FMassEntityManager;
@@ -272,7 +277,7 @@ public:
 		void ForceUpdateCurrentThreadID();
 
 		/**
-		 * Identifies the thread where given FEntityCreationContext instance was created. All subsequent operations are 
+		 * Identifies the thread where given FObserverLock instance was created. All subsequent operations are 
 		 * expected to be run in the same thread.
 		 */
 		uint32 OwnerThreadId;
@@ -280,13 +285,22 @@ public:
 		TArray<FMassEntityHandle> CreatedEntities;
 		FMassArchetypeEntityCollection::EDuplicatesHandling CollectionCreationDuplicatesHandling = FMassArchetypeEntityCollection::EDuplicatesHandling::NoDuplicates;
 		TSharedPtr<FMassEntityManager> Manager;
+		bool bDuringCreation = false;
 	};
+	using FEntityCreationContext = FObserverLock;
+
 	/**
-	 * The main use-case for this function is to create a blank FEntityCreationContext and hold on to it while creating 
+	 * The main use-case for this function is to create a blank FObserverLock and hold on to it while creating 
 	 * a bunch of entities (with multiple calls to BatchCreate* and/or BatchBuild*) and modifying them (with mutating batched API)
-	 * while not causing multiple Observers to trigger. All the observers will be triggered at one go, once the FEntityCreationContext 
+	 * while not causing multiple Observers to trigger. All the observers will be triggered at one go, once the FObserverLock 
 	 * instance gets destroyed. 
-	 * @return ActiveCreationContext. If it's valid FEntityCreationContext instance will be created and assigned to ActiveCreationContext first.
+	 * @return ActiveObserverLock. If it's valid FObserverLock instance will be created and assigned to ActiveObserverLock first.
+	 */
+	TSharedRef<FObserverLock> GetOrMakeObserverLock();
+
+	/**
+	 * creates FObserverLock and marks it up to properly to handle newly-created entities
+	 * @see GetOrMakeObserverLock
 	 */
 	TSharedRef<FEntityCreationContext> GetOrMakeCreationContext();
 
@@ -1026,6 +1040,7 @@ PRAGMA_ENABLE_DEPRECATION_WARNINGS
 	UObject* GetOwner() const { return Owner.Get(); }
 
 	bool IsDuringEntityCreation() const;
+	bool IsObserverExecutionLocked() const;
 
 	void SetDebugName(const FString& NewDebugGame);
 #if WITH_MASSENTITY_DEBUG
@@ -1136,47 +1151,47 @@ private:
 	}
 
 	/**
-	 * If ActiveCreationContext is not valid the function creates a new shared FEntityCreationContext instance and returns that.
-	 * Otherwise ActiveCreationContext will get extended with ReservedEntities and EntityCollection, and returned by the function.
+	 * If ActiveObserverLock is not valid the function creates a new shared FObserverLock instance and returns that.
+	 * Otherwise, ActiveObserverLock will get extended with ReservedEntities and EntityCollection, and returned by the function.
 	 */
-	TSharedRef<FEntityCreationContext> GetOrMakeCreationContext(TConstArrayView<FMassEntityHandle> ReservedEntities, FMassArchetypeEntityCollection&& EntityCollection);
+	TSharedRef<FObserverLock> GetOrMakeObserverLock(TConstArrayView<FMassEntityHandle> ReservedEntities, FMassArchetypeEntityCollection&& EntityCollection);
 
 	/** 
-	 * This type is used in entity mutating batched API to ensure the active FEntityCreationContext gets dirtied 
-	 * upon function's end (since the mutating operations render FEntityCreationContext.EntityCollections invalid).
-	 * It also serves as a cached IsDuringEntityCreation value.
+	 * This type is used in entity mutating batched API to ensure the active FObserverLock gets dirtied 
+	 * upon function's end (since the mutating operations render invalid the entity collections stored by FObserverLock).
+	 * It also serves as a cached IsObserverExecutionLocked() value.
 	 */
-	struct FScopedCreationContextOperations
+	struct FScopedObserverLockOperations
 	{
-		FScopedCreationContextOperations(FMassEntityManager& InManager)
-			: bIsDuringEntityCreation(InManager.IsDuringEntityCreation())
+		FScopedObserverLockOperations(FMassEntityManager& InManager)
+			: bIsObserverExecutionLocked(InManager.IsObserverExecutionLocked())
 			, Manager(InManager)
 		{}
-		~FScopedCreationContextOperations()
+		~FScopedObserverLockOperations()
 		{
-			if (bIsDuringEntityCreation)
+			if (bIsObserverExecutionLocked)
 			{
-				Manager.DirtyCreationContext();
+				Manager.DirtyObserverLock();
 			}
 			// else, there's nothing to do, there's no creation context to call functions for
 		}
 
-		bool IsAllowedToTriggerObservers() const { return (bIsDuringEntityCreation == false); }
+		bool IsAllowedToTriggerObservers() const { return (bIsObserverExecutionLocked == false); }
 
 	private:
-		const bool bIsDuringEntityCreation = false;
+		const bool bIsObserverExecutionLocked = false;
 		FMassEntityManager& Manager;
 	};
-	friend FScopedCreationContextOperations;
+	friend FScopedObserverLockOperations;
 
 	/** @return whether actual context dirtying took place which is equivalent to IsDuringEntityCreation */
-	bool DirtyCreationContext();
+	bool DirtyObserverLock();
 
 	/** 
 	 * @return whether it's allowed for observers to get triggered. If not then the active creation context will be dirtied
 	 *	to cause observers triggering upon context's destruction
 	 */
-	bool IsAllowedToTriggerObservers() { return DirtyCreationContext() == false; }
+	bool IsAllowedToTriggerObservers();
 	bool DebugDoCollectionsOverlapCreationContext(TConstArrayView<FMassArchetypeEntityCollection> EntityCollections) const;
 		
 private:
@@ -1191,8 +1206,8 @@ private:
 	UE::Mass::FStorageType* DebugEntityStoragePtr = nullptr;
 #endif // WITH_MASSENTITY_DEBUG
 
-	/** Never access directly, use GetOrMakeCreationContext instead. */
-	TWeakPtr<FEntityCreationContext> ActiveCreationContext;
+	/** Never access directly, use GetOrMakeObserverLock or GetOrMakeCreationContext instead. */
+	TWeakPtr<FObserverLock> ActiveObserverLock;
 
 	std::atomic<bool> bCommandBufferFlushingInProgress = false;
 
@@ -1335,6 +1350,10 @@ public:
 	void ForEachSharedFragmentConditional(TFunctionRef< bool(T& /*SharedFragment*/) > ConditionFunction, TFunctionRef< void(T& /*SharedFragment*/) > ExecuteFunction)
 	{
 	}
+
+	// @todo consider removing DirtyObserverLock altoghether
+	UE_DEPRECATED(5.6, "DirtyCreationContext is deprecated. Use DirtyObserverLock instead.")
+	bool DirtyCreationContext();
 };
 
 //-----------------------------------------------------------------------------
@@ -1413,5 +1432,53 @@ inline UE::Mass::FTypeManager& FMassEntityManager::GetTypeManager()
 
 inline bool FMassEntityManager::IsDuringEntityCreation() const
 {
-	return ActiveCreationContext.IsValid();
+	return ActiveObserverLock.IsValid() && ActiveObserverLock.Pin()->IsDuringCreation();
+}
+
+inline bool FMassEntityManager::IsObserverExecutionLocked() const
+{
+	return ActiveObserverLock.IsValid();
+}
+
+inline bool FMassEntityManager::IsAllowedToTriggerObservers() 
+{ 
+	return DirtyObserverLock() == false; 
+}
+
+inline int32 FMassEntityManager::FObserverLock::GetSpawnedNum() const 
+{ 
+	return CreatedEntities.Num(); 
+}
+
+inline bool FMassEntityManager::FObserverLock::IsDirty() const 
+{ 
+	return EntityCollections.IsEmpty() && (CreatedEntities.IsEmpty() == false); 
+}
+
+inline bool FMassEntityManager::FObserverLock::IsDuringCreation() const 
+{ 
+	return bDuringCreation; 
+}
+
+inline bool FMassEntityManager::FObserverLock::DebugAreEntityCollectionsUpToDate() const 
+{ 
+	return EntityCollections.IsEmpty() == CreatedEntities.IsEmpty(); 
+}
+
+//-----------------------------------------------------------------------------
+// DEPRECATED
+//-----------------------------------------------------------------------------
+inline bool FMassEntityManager::DirtyCreationContext()
+{
+	return DirtyObserverLock();
+}
+
+inline void FMassEntityManager::FObserverLock::AppendEntities(const TConstArrayView<FMassEntityHandle> InEntitiesToAppend)
+{
+	AppendCreatedEntities(InEntitiesToAppend);
 }
+
+inline void FMassEntityManager::FObserverLock::AppendEntities(const TConstArrayView<FMassEntityHandle> InEntitiesToAppend, FMassArchetypeEntityCollection&& InEntityCollection)
+{
+	AppendCreatedEntities(InEntitiesToAppend, Forward<FMassArchetypeEntityCollection>(InEntityCollection));
+}
\ No newline at end of file
