[Backout] - CL40421411
#rnx
Original CL Desc
-----------------------------------------------------------------
[Mass] Massaged code around "Creation Context" to lay foundations to the following change that will fix observes issues when modifying entities, that are not being created right now, while the creation context is active.

The main change here is renaming FMassEntityManager::FEntityCreationContext to FMassEntityManager::FObserverLock, plus some deprecation

diff --git a/Engine/Source/Runtime/MassEntity/Private/MassEntityManager.cpp b/Engine/Source/Runtime/MassEntity/Private/MassEntityManager.cpp
--- a/Engine/Source/Runtime/MassEntity/Private/MassEntityManager.cpp
+++ b/Engine/Source/Runtime/MassEntity/Private/MassEntityManager.cpp
@@ -46,33 +46,32 @@ namespace UE::Mass::Private
 }
 
 //-----------------------------------------------------------------------------
-// FMassEntityManager::FObserverLock
+// FMassEntityManager::FEntityCreationContext
 //-----------------------------------------------------------------------------
-FMassEntityManager::FObserverLock::FObserverLock()
+FMassEntityManager::FEntityCreationContext::FEntityCreationContext()
 	: OwnerThreadId(FPlatformTLS::GetCurrentThreadId())
 {	
 }
 
-FMassEntityManager::FObserverLock::FObserverLock(FMassEntityManager& InManager, const TConstArrayView<FMassEntityHandle> InCreatedEntities)
-	: FObserverLock()
+FMassEntityManager::FEntityCreationContext::FEntityCreationContext(FMassEntityManager& InManager, const TConstArrayView<FMassEntityHandle> InCreatedEntities)
+	: FEntityCreationContext()
 {
 	CreatedEntities = InCreatedEntities;
 	Manager = InManager.AsShared();
-	bDuringCreation = (InCreatedEntities.Num() > 0);
 }
 
-FMassEntityManager::FObserverLock::FObserverLock(FMassEntityManager& InManager, const TConstArrayView<FMassEntityHandle> InCreatedEntities
+FMassEntityManager::FEntityCreationContext::FEntityCreationContext(FMassEntityManager& InManager, const TConstArrayView<FMassEntityHandle> InCreatedEntities
 	, FMassArchetypeEntityCollection&& InEntityCollection)
-	: FObserverLock(InManager, InCreatedEntities)
+	: FEntityCreationContext(InManager, InCreatedEntities)
 {
-	checkf(InCreatedEntities.IsEmpty() || InEntityCollection.IsEmpty() == false, TEXT("Trying to create FObserverLock instance with no entities but non-empty entity collection. This is not supported."))
+	checkf(InCreatedEntities.IsEmpty() || InEntityCollection.IsEmpty() == false, TEXT("Trying to create FEntityCreationContext instance with no entities but non-empty entity collection. This is not supported."))
 	if (InCreatedEntities.IsEmpty() == false)
 	{
 		EntityCollections.Add(MoveTemp(InEntityCollection));
 	}
 }
 
-FMassEntityManager::FObserverLock::~FObserverLock()
+FMassEntityManager::FEntityCreationContext::~FEntityCreationContext()
 {
 	if (((EntityCollections.IsEmpty() == false) || (CreatedEntities.IsEmpty() == false)) && ensure(Manager))
 	{
@@ -80,7 +79,7 @@ FMassEntityManager::FObserverLock::~FObserverLock()
 	}
 }
 
-TConstArrayView<FMassArchetypeEntityCollection> FMassEntityManager::FObserverLock::GetEntityCollections() const
+TConstArrayView<FMassArchetypeEntityCollection> FMassEntityManager::FEntityCreationContext::GetEntityCollections() const
 {
 	// the EntityCollection has been dirtied, we need to rebuild it
 	if (IsDirty() && ensure(Manager))
@@ -91,16 +90,16 @@ TConstArrayView<FMassArchetypeEntityCollection> FMassEntityManager::FObserverLoc
 	return EntityCollections;
 }
 
-void FMassEntityManager::FObserverLock::MarkDirty()
+void FMassEntityManager::FEntityCreationContext::MarkDirty()
 {
-	checkf(OwnerThreadId == FPlatformTLS::GetCurrentThreadId(), TEXT("%hs: all FObserverLock operations ere expected to be run in a single thread"), __FUNCTION__);
+	checkf(OwnerThreadId == FPlatformTLS::GetCurrentThreadId(), TEXT("%hs: all FEntityCreationContext operations ere expected to be run in a single thread"), __FUNCTION__);
 
 	EntityCollections.Reset();
 }
 
-void FMassEntityManager::FObserverLock::AppendCreatedEntities(const TConstArrayView<FMassEntityHandle> InEntitiesToAppend)
+void FMassEntityManager::FEntityCreationContext::AppendEntities(const TConstArrayView<FMassEntityHandle> InEntitiesToAppend)
 {
-	checkf(OwnerThreadId == FPlatformTLS::GetCurrentThreadId(), TEXT("%hs: all FObserverLock operations ere expected to be run in a single thread"), __FUNCTION__);
+	checkf(OwnerThreadId == FPlatformTLS::GetCurrentThreadId(), TEXT("%hs: all FEntityCreationContext operations ere expected to be run in a single thread"), __FUNCTION__);
 
 	if (InEntitiesToAppend.Num())
 	{
@@ -120,9 +119,9 @@ void FMassEntityManager::FObserverLock::AppendCreatedEntities(const TConstArrayV
 	}
 }
 
-void FMassEntityManager::FObserverLock::AppendCreatedEntities(const TConstArrayView<FMassEntityHandle> InEntitiesToAppend, FMassArchetypeEntityCollection&& InEntityCollection)
+void FMassEntityManager::FEntityCreationContext::AppendEntities(const TConstArrayView<FMassEntityHandle> InEntitiesToAppend, FMassArchetypeEntityCollection&& InEntityCollection)
 {
-	checkf(OwnerThreadId == FPlatformTLS::GetCurrentThreadId(), TEXT("%hs: all FObserverLock operations are expected to be run in a single thread"), __FUNCTION__);
+	checkf(OwnerThreadId == FPlatformTLS::GetCurrentThreadId(), TEXT("%hs: all FEntityCreationContext operations are expected to be run in a single thread"), __FUNCTION__);
 
 	if (InEntitiesToAppend.Num() == 0)
 	{
@@ -131,9 +130,9 @@ void FMassEntityManager::FObserverLock::AppendCreatedEntities(const TConstArrayV
 
 	LLM_SCOPE_BYNAME(TEXT("Mass/EntityManager"));
 
-	AppendCreatedEntities(InEntitiesToAppend);
+	AppendEntities(InEntitiesToAppend);
 
-	// this condition boils down to checking if this FObserverLock instance only connects the just added EntitiesToAppend
+	// this condition boils down to checking if this FEntityCreationContext instance only connects the just added EntitiesToAppend
 	if (CreatedEntities.Num() == InEntitiesToAppend.Num())
 	{
 		checkf(EntityCollections.Num() == 0, TEXT("We never expect EntityCollections to be non-empty while there are no entities in CreatedEntities."));
@@ -141,7 +140,7 @@ void FMassEntityManager::FObserverLock::AppendCreatedEntities(const TConstArrayV
 	}
 }
 
-void FMassEntityManager::FObserverLock::ForceUpdateCurrentThreadID()
+void FMassEntityManager::FEntityCreationContext::ForceUpdateCurrentThreadID()
 {
 	OwnerThreadId = FPlatformTLS::GetCurrentThreadId();
 }
@@ -381,7 +380,7 @@ void FMassEntityManager::OnPostFork(EForkProcessRole Role)
 			}
 		}
 
-		if (TSharedPtr<FObserverLock> ActiveContext = ActiveObserverLock.Pin())
+		if (TSharedPtr<FEntityCreationContext> ActiveContext = ActiveCreationContext.Pin())
 		{
 			ActiveContext->ForceUpdateCurrentThreadID();
 		}
@@ -684,8 +683,8 @@ FMassEntityHandle FMassEntityManager::CreateEntity(TConstArrayView<FInstancedStr
 	const FMassEntityHandle Entity = ReserveEntity();
 
 	// Using a creation context to prevent InternalBuildEntity from notifying observers before we set fragments data
-	const TSharedRef<FObserverLock> ObserverLock = GetOrMakeObserverLock();
-	ObserverLock->AppendCreatedEntities({Entity});
+	const TSharedRef<FEntityCreationContext> CreationContext = GetOrMakeCreationContext();
+	CreationContext->AppendEntities({Entity});
 
 	InternalBuildEntity(Entity, ArchetypeHandle, SharedFragmentValues);
 
@@ -740,8 +739,8 @@ void FMassEntityManager::BuildEntity(FMassEntityHandle Entity, TConstArrayView<F
 	check(ArchetypeHandle.IsValid());
 
 	// Using a creation context to prevent InternalBuildEntity from notifying observers before we set fragments data
-	const TSharedRef<FObserverLock> ObserverLock = GetOrMakeObserverLock();
-	ObserverLock->AppendCreatedEntities({Entity});
+	const TSharedRef<FEntityCreationContext> CreationContext = GetOrMakeCreationContext();
+	CreationContext->AppendEntities({Entity});
 
 	InternalBuildEntity(Entity, ArchetypeHandle, SharedFragmentValues);
 
@@ -764,7 +763,7 @@ int32 FMassEntityManager::BatchReserveEntities(TArrayView<FMassEntityHandle> InO
 	return GetEntityStorageInterface().Acquire(InOutEntities);
 }
 
-TSharedRef<FMassEntityManager::FObserverLock> FMassEntityManager::BatchBuildEntities(const FMassArchetypeEntityCollectionWithPayload& EncodedEntitiesWithPayload
+TSharedRef<FMassEntityManager::FEntityCreationContext> FMassEntityManager::BatchBuildEntities(const FMassArchetypeEntityCollectionWithPayload& EncodedEntitiesWithPayload
 	, const FMassFragmentBitSet& FragmentsAffected, const FMassArchetypeSharedFragmentValues& SharedFragmentValues, const FMassArchetypeCreationParams& CreationParams)
 {
 	CHECK_SYNC_API();
@@ -783,7 +782,7 @@ TSharedRef<FMassEntityManager::FObserverLock> FMassEntityManager::BatchBuildEnti
 	return BatchBuildEntities(EncodedEntitiesWithPayload, MoveTemp(Composition), SharedFragmentValues, CreationParams);
 }
 
-TSharedRef<FMassEntityManager::FObserverLock> FMassEntityManager::BatchBuildEntities(const FMassArchetypeEntityCollectionWithPayload& EncodedEntitiesWithPayload
+TSharedRef<FMassEntityManager::FEntityCreationContext> FMassEntityManager::BatchBuildEntities(const FMassArchetypeEntityCollectionWithPayload& EncodedEntitiesWithPayload
 	, FMassArchetypeCompositionDescriptor&& Composition
 	, const FMassArchetypeSharedFragmentValues& SharedFragmentValues, const FMassArchetypeCreationParams& CreationParams)
 {
@@ -837,11 +836,11 @@ TSharedRef<FMassEntityManager::FObserverLock> FMassEntityManager::BatchBuildEnti
 	// and as such won't cause observers triggering (which usually is prevented by context's existence), and that we 
 	// strongly assume the all entity creation/building (not to be mistaken with "reserving") takes place in a single thread
 	// @todo add checks/ensures enforcing the assumption mentioned above.
-	return GetOrMakeObserverLock(EntityHandles
+	return GetOrMakeCreationContext(EntityHandles
 		, FMassArchetypeEntityCollection(TargetArchetypeHandle, MoveTemp(TargetArchetypeEntityRanges)));
 }
 
-TSharedRef<FMassEntityManager::FObserverLock> FMassEntityManager::BatchCreateReservedEntities(const FMassArchetypeHandle& ArchetypeHandle
+TSharedRef<FMassEntityManager::FEntityCreationContext> FMassEntityManager::BatchCreateReservedEntities(const FMassArchetypeHandle& ArchetypeHandle
 	, const FMassArchetypeSharedFragmentValues& SharedFragmentValues, TConstArrayView<FMassEntityHandle> ReservedEntities)
 {
 	CHECK_SYNC_API();
@@ -854,11 +853,11 @@ TSharedRef<FMassEntityManager::FObserverLock> FMassEntityManager::BatchCreateRes
 		, SharedFragmentValues, ReservedEntities);
 }
 
-TSharedRef<FMassEntityManager::FObserverLock> FMassEntityManager::BatchCreateEntities(const FMassArchetypeHandle& ArchetypeHandle
+TSharedRef<FMassEntityManager::FEntityCreationContext> FMassEntityManager::BatchCreateEntities(const FMassArchetypeHandle& ArchetypeHandle
 	, const FMassArchetypeSharedFragmentValues& SharedFragmentValues, const int32 Count, TArray<FMassEntityHandle>& InOutEntities)
 {
 	CHECK_SYNC_API();
-	testableCheckfReturn(ArchetypeHandle.IsValid(), return GetOrMakeObserverLock()
+	testableCheckfReturn(ArchetypeHandle.IsValid(), return GetOrMakeCreationContext()
 		, TEXT("%hs expecting a valid ArchetypeHandle"), __FUNCTION__);
 
 	TRACE_CPUPROFILER_EVENT_SCOPE(Mass_BatchCreateEntities);
@@ -870,7 +869,7 @@ TSharedRef<FMassEntityManager::FObserverLock> FMassEntityManager::BatchCreateEnt
 		, SharedFragmentValues, ReservedEntities);
 }
 
-TSharedRef<FMassEntityManager::FObserverLock> FMassEntityManager::InternalBatchCreateReservedEntities(const FMassArchetypeHandle& ArchetypeHandle
+TSharedRef<FMassEntityManager::FEntityCreationContext> FMassEntityManager::InternalBatchCreateReservedEntities(const FMassArchetypeHandle& ArchetypeHandle
 	, const FMassArchetypeSharedFragmentValues& SharedFragmentValues, TConstArrayView<FMassEntityHandle> ReservedEntities)
 {
 	// Functions calling into this one are required to verify that the archetype handle is valid
@@ -891,7 +890,7 @@ TSharedRef<FMassEntityManager::FObserverLock> FMassEntityManager::InternalBatchC
 
 	UE_TRACE_MASS_ENTITIES_CREATED(ReservedEntities, *ArchetypeData)
 
-	return GetOrMakeObserverLock(ReservedEntities, FMassArchetypeEntityCollection(ArchetypeHandle, MoveTemp(TargetArchetypeEntityRanges)));
+	return GetOrMakeCreationContext(ReservedEntities, FMassArchetypeEntityCollection(ArchetypeHandle, MoveTemp(TargetArchetypeEntityRanges)));
 }
 
 void FMassEntityManager::DestroyEntity(FMassEntityHandle Entity)
@@ -1645,7 +1644,7 @@ void FMassEntityManager::BatchChangeTagsForEntities(TConstArrayView<FMassArchety
 
 	TRACE_CPUPROFILER_EVENT_SCOPE(Mass_BatchChangeTagsForEntities);
 
-	const FScopedObserverLockOperations ObserverLockOperations(*this);
+	const FScopedCreationContextOperations CreationContextOperations(*this);
 	
 	for (const FMassArchetypeEntityCollection& Collection : EntityCollections)
 	{
@@ -1659,7 +1658,7 @@ void FMassEntityManager::BatchChangeTagsForEntities(TConstArrayView<FMassArchety
 			FMassTagBitSet TagsAdded = TagsToAdd - CurrentArchetype->GetTagBitSet();
 			FMassTagBitSet TagsRemoved = TagsToRemove.GetOverlap(CurrentArchetype->GetTagBitSet());
 
-			if (ObserverLockOperations.IsAllowedToTriggerObservers()
+			if (CreationContextOperations.IsAllowedToTriggerObservers()
 				&& ObserverManager.HasObserversForBitSet(TagsRemoved, EMassObservedOperation::Remove))
 			{
 				// @todo should use OnPreCompositionRemoved here instead, but we're missing a FMassArchetypeEntityCollection version
@@ -1683,7 +1682,7 @@ void FMassEntityManager::BatchChangeTagsForEntities(TConstArrayView<FMassArchety
 				GetEntityStorageInterface().SetArchetypeFromShared(Entity.Index, NewArchetypeHandle.DataPtr);
 			}
 
-			if (bTagsAddedAreObserved && ObserverLockOperations.IsAllowedToTriggerObservers())
+			if (bTagsAddedAreObserved && CreationContextOperations.IsAllowedToTriggerObservers())
 			{
 				// @todo should use OnPostCompositionAdded here instead, but we're missing a FMassArchetypeEntityCollection version
 				ObserverManager.OnCompositionChanged(
@@ -1701,7 +1700,7 @@ void FMassEntityManager::BatchChangeFragmentCompositionForEntities(TConstArrayVi
 
 	TRACE_CPUPROFILER_EVENT_SCOPE(Mass_BatchChangeFragmentCompositionForEntities);
 
-	FScopedObserverLockOperations ObserverLockOperations(*this);
+	FScopedCreationContextOperations CreationContextOperations(*this);
 
 	for (const FMassArchetypeEntityCollection& Collection : EntityCollections)
 	{
@@ -1718,7 +1717,7 @@ void FMassEntityManager::BatchChangeFragmentCompositionForEntities(TConstArrayVi
 				const bool bFragmentsAddedAreObserved = ObserverManager.HasObserversForBitSet(FragmentsAdded, EMassObservedOperation::Add);
 				FMassFragmentBitSet FragmentsRemoved = FragmentsToRemove.GetOverlap(CurrentArchetype->GetFragmentBitSet());
 				
-				if (ObserverLockOperations.IsAllowedToTriggerObservers()
+				if (CreationContextOperations.IsAllowedToTriggerObservers()
 					&& ObserverManager.HasObserversForBitSet(FragmentsRemoved, EMassObservedOperation::Remove))
 				{
 					ObserverManager.OnCompositionChanged(Collection, FMassArchetypeCompositionDescriptor(MoveTemp(FragmentsRemoved)), EMassObservedOperation::Remove);
@@ -1740,7 +1739,7 @@ void FMassEntityManager::BatchChangeFragmentCompositionForEntities(TConstArrayVi
 					GetEntityStorageInterface().SetArchetypeFromShared(Entity.Index, NewArchetypeHandle.DataPtr);
 				}
 
-				if (bFragmentsAddedAreObserved && ObserverLockOperations.IsAllowedToTriggerObservers())
+				if (bFragmentsAddedAreObserved && CreationContextOperations.IsAllowedToTriggerObservers())
 				{
 					ObserverManager.OnCompositionChanged(
 						FMassArchetypeEntityCollection(NewArchetypeHandle, MoveTemp(NewArchetypeEntityRanges))
@@ -1773,7 +1772,7 @@ void FMassEntityManager::BatchAddFragmentInstancesForEntities(TConstArrayView<FM
 	//		archetype. This could be utilized as a basis for optimization. To be investigated.
 	// 
 
-	FScopedObserverLockOperations ObserverLockOperations(*this);
+	FScopedCreationContextOperations CreationContextOperations(*this);
 
 	for (const FMassArchetypeEntityCollectionWithPayload& EntityRangesWithPayload : EntityCollections)
 	{
@@ -1821,7 +1820,7 @@ void FMassEntityManager::BatchAddFragmentInstancesForEntities(TConstArrayView<FM
 			// corresponds to the order in FMassArchetypeEntityCollectionWithPayload's payload
 			TargetArchetypeHandle.DataPtr->BatchSetFragmentValues(TargetArchetypeEntityRanges, EntityRangesWithPayload.GetPayload());
 			
-			if (bFragmentsAddedAreObserved && ObserverLockOperations.IsAllowedToTriggerObservers())
+			if (bFragmentsAddedAreObserved && CreationContextOperations.IsAllowedToTriggerObservers())
 			{
 				ObserverManager.OnCompositionChanged(
 					FMassArchetypeEntityCollection(TargetArchetypeHandle, MoveTemp(TargetArchetypeEntityRanges))
@@ -2106,51 +2105,44 @@ void FMassEntityManager::AppendCommands(TSharedPtr<FMassCommandBuffer>& InOutCom
 	Defer().MoveAppend(*InOutCommandBuffer.Get());
 }
 
-TSharedRef<FMassEntityManager::FObserverLock> FMassEntityManager::GetOrMakeObserverLock()
+TSharedRef<FMassEntityManager::FEntityCreationContext> FMassEntityManager::GetOrMakeCreationContext()
 {
-	if (ActiveObserverLock.IsValid())
+	if (ActiveCreationContext.IsValid())
 	{
-		return ActiveObserverLock.Pin().ToSharedRef();
+		return ActiveCreationContext.Pin().ToSharedRef();
 	}
 	else
 	{
-		FObserverLock* ObserverLock = new FObserverLock(*this);
-		TSharedRef<FObserverLock> SharedContext = MakeShareable(ObserverLock);
-		ActiveObserverLock = SharedContext;
+		FEntityCreationContext* CreationContext = new FEntityCreationContext(*this);
+		TSharedRef<FEntityCreationContext> SharedContext = MakeShareable(CreationContext);
+		ActiveCreationContext = SharedContext;
 		return SharedContext;
 	}
 }
 
-TSharedRef<FMassEntityManager::FEntityCreationContext> FMassEntityManager::GetOrMakeCreationContext()
-{
-	TSharedRef<FObserverLock> ObserverLock = GetOrMakeObserverLock();
-	ObserverLock->bDuringCreation = true;
-	return ObserverLock;
-}
-
-TSharedRef<FMassEntityManager::FObserverLock> FMassEntityManager::GetOrMakeObserverLock(TConstArrayView<FMassEntityHandle> ReservedEntities
+TSharedRef<FMassEntityManager::FEntityCreationContext> FMassEntityManager::GetOrMakeCreationContext(TConstArrayView<FMassEntityHandle> ReservedEntities
 	, FMassArchetypeEntityCollection&& EntityCollection)
 {
 	LLM_SCOPE_BYNAME(TEXT("Mass/EntityManager"));
-	if (ActiveObserverLock.IsValid())
+	if (ActiveCreationContext.IsValid())
 	{
-		TSharedPtr<FObserverLock> SharedContext = ActiveObserverLock.Pin();
+		TSharedPtr<FEntityCreationContext> SharedContext = ActiveCreationContext.Pin();
 		CA_ASSUME(SharedContext);
-		SharedContext->AppendCreatedEntities(ReservedEntities, MoveTemp(EntityCollection));
+		SharedContext->AppendEntities(ReservedEntities, MoveTemp(EntityCollection));
 		return SharedContext.ToSharedRef();
 	}
 	else
 	{
-		FObserverLock* ObserverLock = new FObserverLock(*this, ReservedEntities, MoveTemp(EntityCollection));
-		TSharedRef<FObserverLock> SharedContext = MakeShareable(ObserverLock);
-		ActiveObserverLock = SharedContext;
+		FEntityCreationContext* CreationContext = new FEntityCreationContext(*this, ReservedEntities, MoveTemp(EntityCollection));
+		TSharedRef<FEntityCreationContext> SharedContext = MakeShareable(CreationContext);
+		ActiveCreationContext = SharedContext;
 		return SharedContext;
 	}
 }
 
-bool FMassEntityManager::DirtyObserverLock()
+bool FMassEntityManager::DirtyCreationContext()
 {
-	if (TSharedPtr<FObserverLock> AsSharedPtr = ActiveObserverLock.Pin())
+	if (TSharedPtr<FEntityCreationContext> AsSharedPtr = ActiveCreationContext.Pin())
 	{
 		AsSharedPtr->MarkDirty();
 		return true;
@@ -2165,7 +2157,7 @@ UE::Mass::FEntityBuilder FMassEntityManager::MakeEntityBuilder()
 
 bool FMassEntityManager::DebugDoCollectionsOverlapCreationContext(TConstArrayView<FMassArchetypeEntityCollection> EntityCollections) const
 {
-	if (TSharedPtr<FObserverLock> AsSharedPtr = ActiveObserverLock.Pin())
+	if (TSharedPtr<FEntityCreationContext> AsSharedPtr = ActiveCreationContext.Pin())
 	{
 		TConstArrayView<FMassArchetypeEntityCollection> CreationCollections = AsSharedPtr->EntityCollections;
 		return CreationCollections.GetData() <= EntityCollections.GetData()
@@ -2345,7 +2337,7 @@ const UE::Mass::FStorageType& FMassEntityManager::DebugGetEntityStorageInterface
 // DEPRECATED
 //-----------------------------------------------------------------------------
 
-const FMassArchetypeEntityCollection& FMassEntityManager::FObserverLock::GetEntityCollection() const
+const FMassArchetypeEntityCollection& FMassEntityManager::FEntityCreationContext::GetEntityCollection() const
 {
 	static FMassArchetypeEntityCollection EmptyCollection;
 	return EntityCollections.Num() ? EntityCollections[0] : EmptyCollection;
diff --git a/Engine/Source/Runtime/MassEntity/Private/MassObserverManager.cpp b/Engine/Source/Runtime/MassEntity/Private/MassObserverManager.cpp
--- a/Engine/Source/Runtime/MassEntity/Private/MassObserverManager.cpp
+++ b/Engine/Source/Runtime/MassEntity/Private/MassObserverManager.cpp
@@ -162,7 +162,7 @@ bool FMassObserverManager::OnPostEntitiesCreated(FMassProcessingContext& Process
 
 bool FMassObserverManager::OnPostEntitiesCreated(FMassProcessingContext& ProcessingContext, TConstArrayView<FMassArchetypeEntityCollection> EntityCollections)
 {
-	if (!ensureMsgf(EntityManager.IsObserverExecutionLocked() == false, TEXT("Observes execution is locked right now.")))
+	if (!ensureMsgf(EntityManager.IsDuringEntityCreation() == false, TEXT("Entity creation not done yet.")))
 	{
 		return false;
 	}
diff --git a/Engine/Source/Runtime/MassEntity/Public/MassEntityManager.h b/Engine/Source/Runtime/MassEntity/Public/MassEntityManager.h
--- a/Engine/Source/Runtime/MassEntity/Public/MassEntityManager.h
+++ b/Engine/Source/Runtime/MassEntity/Public/MassEntityManager.h
@@ -240,36 +240,31 @@ public:
 	 * A dedicated structure for ensuring the "on entities creation" observers get notified only once all other 
 	 * initialization operations are done and this creation context instance gets released. 
 	 */
-	struct MASSENTITY_API FObserverLock
+	struct MASSENTITY_API FEntityCreationContext
 	{
 	private:
-		FObserverLock();
-		explicit FObserverLock(FMassEntityManager& InManager, const TConstArrayView<FMassEntityHandle> InCreatedEntities = {});
-		FObserverLock(FMassEntityManager& InManager, const TConstArrayView<FMassEntityHandle> InCreatedEntities, FMassArchetypeEntityCollection&& InEntityCollection);
+		FEntityCreationContext();
+		explicit FEntityCreationContext(FMassEntityManager& InManager, const TConstArrayView<FMassEntityHandle> InCreatedEntities = {});
+		FEntityCreationContext(FMassEntityManager& InManager, const TConstArrayView<FMassEntityHandle> InCreatedEntities, FMassArchetypeEntityCollection&& InEntityCollection);
 		
 	public:
-		~FObserverLock();
+		~FEntityCreationContext();
 
 		/** Returns EntityCollections, reconstructing them if needed (empty or dirtied). */
 		TConstArrayView<FMassArchetypeEntityCollection> GetEntityCollections() const;
-		int32 GetSpawnedNum() const;
+		int32 GetSpawnedNum() const { return CreatedEntities.Num(); }
 		void MarkDirty();
-		bool IsDirty() const;
-		bool IsDuringCreation() const;
-		void AppendCreatedEntities(const TConstArrayView<FMassEntityHandle> InEntitiesToAppend);
-		void AppendCreatedEntities(const TConstArrayView<FMassEntityHandle> InEntitiesToAppend, FMassArchetypeEntityCollection&& InEntityCollection);
+		FORCEINLINE bool IsDirty() const { return EntityCollections.IsEmpty() && (CreatedEntities.IsEmpty() == false); }
+		void AppendEntities(const TConstArrayView<FMassEntityHandle> InEntitiesToAppend);
+		void AppendEntities(const TConstArrayView<FMassEntityHandle> InEntitiesToAppend, FMassArchetypeEntityCollection&& InEntityCollection);
 
 		/** Function for debugging/testing purposes. We don't expect users to ever call it, always get collections via GetEntityCollections */
-		bool DebugAreEntityCollectionsUpToDate() const;
+		bool DebugAreEntityCollectionsUpToDate() const { return EntityCollections.IsEmpty() == CreatedEntities.IsEmpty(); }
 
 		UE_DEPRECATED(5.5, "This constructor is now deprecated and defunct. Use one of the others instead.")
-		explicit FObserverLock(const int32) : FObserverLock() {}
-		UE_DEPRECATED(5.5, "This function is now deprecated since FObserverLock can contain more than a single collection now. Use GetEntityCollections instead.")
+		explicit FEntityCreationContext(const int32 InNumSpawned) : FEntityCreationContext() {}
+		UE_DEPRECATED(5.5, "This function is now deprecated since FEntityCreationContext can contain more than a single collection now. Use GetEntityCollections instead.")
 		const FMassArchetypeEntityCollection& GetEntityCollection() const;
-		UE_DEPRECATED(5.6, "AppendEntities is deprecated. Call AppendCreatedEntities instead.")
-		void AppendEntities(const TConstArrayView<FMassEntityHandle> InEntitiesToAppend);
-		UE_DEPRECATED(5.6, "AppendEntities is deprecated. Call AppendCreatedEntities instead.")
-		void AppendEntities(const TConstArrayView<FMassEntityHandle> InEntitiesToAppend, FMassArchetypeEntityCollection&& InEntityCollection);
 
 	private:
 		friend FMassEntityManager;
@@ -277,7 +272,7 @@ public:
 		void ForceUpdateCurrentThreadID();
 
 		/**
-		 * Identifies the thread where given FObserverLock instance was created. All subsequent operations are 
+		 * Identifies the thread where given FEntityCreationContext instance was created. All subsequent operations are 
 		 * expected to be run in the same thread.
 		 */
 		uint32 OwnerThreadId;
@@ -285,22 +280,13 @@ public:
 		TArray<FMassEntityHandle> CreatedEntities;
 		FMassArchetypeEntityCollection::EDuplicatesHandling CollectionCreationDuplicatesHandling = FMassArchetypeEntityCollection::EDuplicatesHandling::NoDuplicates;
 		TSharedPtr<FMassEntityManager> Manager;
-		bool bDuringCreation = false;
 	};
-	using FEntityCreationContext = FObserverLock;
-
 	/**
-	 * The main use-case for this function is to create a blank FObserverLock and hold on to it while creating 
+	 * The main use-case for this function is to create a blank FEntityCreationContext and hold on to it while creating 
 	 * a bunch of entities (with multiple calls to BatchCreate* and/or BatchBuild*) and modifying them (with mutating batched API)
-	 * while not causing multiple Observers to trigger. All the observers will be triggered at one go, once the FObserverLock 
+	 * while not causing multiple Observers to trigger. All the observers will be triggered at one go, once the FEntityCreationContext 
 	 * instance gets destroyed. 
-	 * @return ActiveObserverLock. If it's valid FObserverLock instance will be created and assigned to ActiveObserverLock first.
-	 */
-	TSharedRef<FObserverLock> GetOrMakeObserverLock();
-
-	/**
-	 * creates FObserverLock and marks it up to properly to handle newly-created entities
-	 * @see GetOrMakeObserverLock
+	 * @return ActiveCreationContext. If it's valid FEntityCreationContext instance will be created and assigned to ActiveCreationContext first.
 	 */
 	TSharedRef<FEntityCreationContext> GetOrMakeCreationContext();
 
@@ -1040,7 +1026,6 @@ PRAGMA_ENABLE_DEPRECATION_WARNINGS
 	UObject* GetOwner() const { return Owner.Get(); }
 
 	bool IsDuringEntityCreation() const;
-	bool IsObserverExecutionLocked() const;
 
 	void SetDebugName(const FString& NewDebugGame);
 #if WITH_MASSENTITY_DEBUG
@@ -1151,47 +1136,47 @@ private:
 	}
 
 	/**
-	 * If ActiveObserverLock is not valid the function creates a new shared FObserverLock instance and returns that.
-	 * Otherwise, ActiveObserverLock will get extended with ReservedEntities and EntityCollection, and returned by the function.
+	 * If ActiveCreationContext is not valid the function creates a new shared FEntityCreationContext instance and returns that.
+	 * Otherwise ActiveCreationContext will get extended with ReservedEntities and EntityCollection, and returned by the function.
 	 */
-	TSharedRef<FObserverLock> GetOrMakeObserverLock(TConstArrayView<FMassEntityHandle> ReservedEntities, FMassArchetypeEntityCollection&& EntityCollection);
+	TSharedRef<FEntityCreationContext> GetOrMakeCreationContext(TConstArrayView<FMassEntityHandle> ReservedEntities, FMassArchetypeEntityCollection&& EntityCollection);
 
 	/** 
-	 * This type is used in entity mutating batched API to ensure the active FObserverLock gets dirtied 
-	 * upon function's end (since the mutating operations render invalid the entity collections stored by FObserverLock).
-	 * It also serves as a cached IsObserverExecutionLocked() value.
+	 * This type is used in entity mutating batched API to ensure the active FEntityCreationContext gets dirtied 
+	 * upon function's end (since the mutating operations render FEntityCreationContext.EntityCollections invalid).
+	 * It also serves as a cached IsDuringEntityCreation value.
 	 */
-	struct FScopedObserverLockOperations
+	struct FScopedCreationContextOperations
 	{
-		FScopedObserverLockOperations(FMassEntityManager& InManager)
-			: bIsObserverExecutionLocked(InManager.IsObserverExecutionLocked())
+		FScopedCreationContextOperations(FMassEntityManager& InManager)
+			: bIsDuringEntityCreation(InManager.IsDuringEntityCreation())
 			, Manager(InManager)
 		{}
-		~FScopedObserverLockOperations()
+		~FScopedCreationContextOperations()
 		{
-			if (bIsObserverExecutionLocked)
+			if (bIsDuringEntityCreation)
 			{
-				Manager.DirtyObserverLock();
+				Manager.DirtyCreationContext();
 			}
 			// else, there's nothing to do, there's no creation context to call functions for
 		}
 
-		bool IsAllowedToTriggerObservers() const { return (bIsObserverExecutionLocked == false); }
+		bool IsAllowedToTriggerObservers() const { return (bIsDuringEntityCreation == false); }
 
 	private:
-		const bool bIsObserverExecutionLocked = false;
+		const bool bIsDuringEntityCreation = false;
 		FMassEntityManager& Manager;
 	};
-	friend FScopedObserverLockOperations;
+	friend FScopedCreationContextOperations;
 
 	/** @return whether actual context dirtying took place which is equivalent to IsDuringEntityCreation */
-	bool DirtyObserverLock();
+	bool DirtyCreationContext();
 
 	/** 
 	 * @return whether it's allowed for observers to get triggered. If not then the active creation context will be dirtied
 	 *	to cause observers triggering upon context's destruction
 	 */
-	bool IsAllowedToTriggerObservers();
+	bool IsAllowedToTriggerObservers() { return DirtyCreationContext() == false; }
 	bool DebugDoCollectionsOverlapCreationContext(TConstArrayView<FMassArchetypeEntityCollection> EntityCollections) const;
 		
 private:
@@ -1206,8 +1191,8 @@ private:
 	UE::Mass::FStorageType* DebugEntityStoragePtr = nullptr;
 #endif // WITH_MASSENTITY_DEBUG
 
-	/** Never access directly, use GetOrMakeObserverLock or GetOrMakeCreationContext instead. */
-	TWeakPtr<FObserverLock> ActiveObserverLock;
+	/** Never access directly, use GetOrMakeCreationContext instead. */
+	TWeakPtr<FEntityCreationContext> ActiveCreationContext;
 
 	std::atomic<bool> bCommandBufferFlushingInProgress = false;
 
@@ -1350,10 +1335,6 @@ public:
 	void ForEachSharedFragmentConditional(TFunctionRef< bool(T& /*SharedFragment*/) > ConditionFunction, TFunctionRef< void(T& /*SharedFragment*/) > ExecuteFunction)
 	{
 	}
-
-	// @todo consider removing DirtyObserverLock altoghether
-	UE_DEPRECATED(5.6, "DirtyCreationContext is deprecated. Use DirtyObserverLock instead.")
-	bool DirtyCreationContext();
 };
 
 //-----------------------------------------------------------------------------
@@ -1432,53 +1413,5 @@ inline UE::Mass::FTypeManager& FMassEntityManager::GetTypeManager()
 
 inline bool FMassEntityManager::IsDuringEntityCreation() const
 {
-	return ActiveObserverLock.IsValid() && ActiveObserverLock.Pin()->IsDuringCreation();
-}
-
-inline bool FMassEntityManager::IsObserverExecutionLocked() const
-{
-	return ActiveObserverLock.IsValid();
-}
-
-inline bool FMassEntityManager::IsAllowedToTriggerObservers() 
-{ 
-	return DirtyObserverLock() == false; 
-}
-
-inline int32 FMassEntityManager::FObserverLock::GetSpawnedNum() const 
-{ 
-	return CreatedEntities.Num(); 
-}
-
-inline bool FMassEntityManager::FObserverLock::IsDirty() const 
-{ 
-	return EntityCollections.IsEmpty() && (CreatedEntities.IsEmpty() == false); 
-}
-
-inline bool FMassEntityManager::FObserverLock::IsDuringCreation() const 
-{ 
-	return bDuringCreation; 
-}
-
-inline bool FMassEntityManager::FObserverLock::DebugAreEntityCollectionsUpToDate() const 
-{ 
-	return EntityCollections.IsEmpty() == CreatedEntities.IsEmpty(); 
-}
-
-//-----------------------------------------------------------------------------
-// DEPRECATED
-//-----------------------------------------------------------------------------
-inline bool FMassEntityManager::DirtyCreationContext()
-{
-	return DirtyObserverLock();
-}
-
-inline void FMassEntityManager::FObserverLock::AppendEntities(const TConstArrayView<FMassEntityHandle> InEntitiesToAppend)
-{
-	AppendCreatedEntities(InEntitiesToAppend);
+	return ActiveCreationContext.IsValid();
 }
-
-inline void FMassEntityManager::FObserverLock::AppendEntities(const TConstArrayView<FMassEntityHandle> InEntitiesToAppend, FMassArchetypeEntityCollection&& InEntityCollection)
-{
-	AppendCreatedEntities(InEntitiesToAppend, Forward<FMassArchetypeEntityCollection>(InEntityCollection));
-}
\ No newline at end of file
