[Mass] added LexToString for MassEntityHandle and updated FlushCommands and AppendCommands in MassEntityManager to use const value to facilitate calls from lambda captures (no need to use mutable)

diff --git a/Engine/Source/Runtime/MassEntity/Private/MassEntityManager.cpp b/Engine/Source/Runtime/MassEntity/Private/MassEntityManager.cpp
--- a/Engine/Source/Runtime/MassEntity/Private/MassEntityManager.cpp
+++ b/Engine/Source/Runtime/MassEntity/Private/MassEntityManager.cpp
@@ -2035,7 +2035,7 @@ FMassExecutionContext FMassEntityManager::CreateExecutionContext(const float Del
 	return MoveTemp(ExecutionContext);
 }
 
-void FMassEntityManager::FlushCommands(TSharedPtr<FMassCommandBuffer>& InCommandBuffer)
+void FMassEntityManager::FlushCommands(const TSharedPtr<FMassCommandBuffer>& InCommandBuffer)
 {
 	if (!ensureMsgf(IsInGameThread(), TEXT("Calling %hs is supported only on the Game Tread"), __FUNCTION__))
 	{
@@ -2094,7 +2094,7 @@ void FMassEntityManager::FlushCommands()
 	}
 }
 
-void FMassEntityManager::AppendCommands(TSharedPtr<FMassCommandBuffer>& InOutCommandBuffer)
+void FMassEntityManager::AppendCommands(const TSharedPtr<FMassCommandBuffer>& InOutCommandBuffer)
 {
 	if (!ensureMsgf(Algo::Find(DeferredCommandBuffers, InOutCommandBuffer) == nullptr
 		, TEXT("We don't expect AppendCommands to be called with EntityManager's command buffer as the input parameter")))
diff --git a/Engine/Source/Runtime/MassEntity/Public/MassEntityHandle.h b/Engine/Source/Runtime/MassEntity/Public/MassEntityHandle.h
--- a/Engine/Source/Runtime/MassEntity/Public/MassEntityHandle.h
+++ b/Engine/Source/Runtime/MassEntity/Public/MassEntityHandle.h
@@ -70,6 +70,11 @@ struct alignas(8) FMassEntityHandle
 		return HashCombine(Entity.Index, Entity.SerialNumber);
 	}
 
+	friend FString LexToString(const FMassEntityHandle Entity)
+	{
+		return Entity.DebugGetDescription();
+	}
+
 	FString DebugGetDescription() const
 	{
 		return FString::Printf(TEXT("i: %d sn: %d"), Index, SerialNumber);
diff --git a/Engine/Source/Runtime/MassEntity/Public/MassEntityManager.h b/Engine/Source/Runtime/MassEntity/Public/MassEntityManager.h
--- a/Engine/Source/Runtime/MassEntity/Public/MassEntityManager.h
+++ b/Engine/Source/Runtime/MassEntity/Public/MassEntityManager.h
@@ -664,9 +664,9 @@ public:
 	/** 
 	 * @param InCommandBuffer if not set then the default command buffer will be flushed. If set and there's already 
 	 *		a command buffer being flushed (be it the main one or a previously requested one) then this command buffer 
-	 *		will be queue itself.
+	 *		will be queued itself.
 	 */
-	void FlushCommands(TSharedPtr<FMassCommandBuffer>& InCommandBuffer);
+	void FlushCommands(const TSharedPtr<FMassCommandBuffer>& InCommandBuffer);
 
 	void FlushCommands();
 
@@ -675,7 +675,7 @@ public:
 	 * InOutCommandBuffer into the main command buffer or append it to the list of command buffers waiting to be flushed.
 	 * @note as a consequence of the call InOutCommandBuffer can get its contents emptied due some of the underlying code using Move semantics
 	 */
-	void AppendCommands(TSharedPtr<FMassCommandBuffer>& InOutCommandBuffer);
+	void AppendCommands(const TSharedPtr<FMassCommandBuffer>& InOutCommandBuffer);
 
 	template<typename T>
 	UE_DEPRECATED(5.5, "This method will no longer be exposed. Use GetOrCreateConstSharedFragment instead.")
