[MassEntity] added messages to FMassSubsystemAccess ensures to help users implementing the fix

diff --git a/Engine/Source/Runtime/MassEntity/Public/MassSubsystemAccess.h b/Engine/Source/Runtime/MassEntity/Public/MassSubsystemAccess.h
--- a/Engine/Source/Runtime/MassEntity/Public/MassSubsystemAccess.h
+++ b/Engine/Source/Runtime/MassEntity/Public/MassSubsystemAccess.h
@@ -30,7 +30,8 @@ struct MASSENTITY_API FMassSubsystemAccess
 	T* GetMutableSubsystem()
 	{
 		const uint32 SystemIndex = FMassExternalSubsystemBitSet::GetTypeIndex<T>();
-		if (ensure(MutableSubsystemsBitSet.IsBitSet(SystemIndex)))
+		if (ensureMsgf(MutableSubsystemsBitSet.IsBitSet(SystemIndex)
+			, TEXT("Missing 'template<> struct TMassExternalSubsystemTraits<U%s>' for this subsystem"), *GetNameSafe(T::StaticClass())))
 		{
 			return GetSubsystemInternal<T>(SystemIndex);
 		}
@@ -50,7 +51,8 @@ struct MASSENTITY_API FMassSubsystemAccess
 	const T* GetSubsystem()
 	{
 		const uint32 SystemIndex = FMassExternalSubsystemBitSet::GetTypeIndex<T>();
-		if (ensure(ConstSubsystemsBitSet.IsBitSet(SystemIndex) || MutableSubsystemsBitSet.IsBitSet(SystemIndex)))
+		if (ensureMsgf(ConstSubsystemsBitSet.IsBitSet(SystemIndex) || MutableSubsystemsBitSet.IsBitSet(SystemIndex)
+			, TEXT("Missing 'template<> struct TMassExternalSubsystemTraits<U%s>' for this subsystem"), *GetNameSafe(T::StaticClass())))
 		{
 			return GetSubsystemInternal<T>(SystemIndex);
 		}
@@ -73,7 +75,7 @@ struct MASSENTITY_API FMassSubsystemAccess
 	{
 		const uint32 SystemIndex = FMassExternalSubsystemBitSet::GetTypeIndex(**SubsystemClass);
 		if (ensureMsgf(MutableSubsystemsBitSet.IsBitSet(SystemIndex)
-			, TEXT("Not configured to access SubsystemClass %s"), *GetNameSafe(SubsystemClass)))
+			, TEXT("Missing 'template<> struct TMassExternalSubsystemTraits<U%s>' for this subsystem"), *GetNameSafe(SubsystemClass)))
 		{
 			return GetSubsystemInternal<T>(SystemIndex, SubsystemClass);
 		}
@@ -94,7 +96,7 @@ struct MASSENTITY_API FMassSubsystemAccess
 	{
 		const uint32 SystemIndex = FMassExternalSubsystemBitSet::GetTypeIndex(**SubsystemClass);
 		if (ensureMsgf(ConstSubsystemsBitSet.IsBitSet(SystemIndex) || MutableSubsystemsBitSet.IsBitSet(SystemIndex)
-			, TEXT("Not configured to access SubsystemClass %s"), *GetNameSafe(SubsystemClass)))
+			, TEXT("Missing 'template<> struct TMassExternalSubsystemTraits<U%s>' for this subsystem"), *GetNameSafe(SubsystemClass)))
 		{
 			return GetSubsystemInternal<T>(SystemIndex, SubsystemClass);
 		}
