[Mass] Fix invalid assertions that were comparing memory addresses to allocation size.
Instead asset on offset relative to alloc size.

#rnx

diff --git a/Engine/Source/Runtime/MassEntity/Internal/MassArchetypeData.h b/Engine/Source/Runtime/MassEntity/Internal/MassArchetypeData.h
--- a/Engine/Source/Runtime/MassEntity/Internal/MassArchetypeData.h
+++ b/Engine/Source/Runtime/MassEntity/Internal/MassArchetypeData.h
@@ -57,14 +57,18 @@ public:
 	// Returns the Entity array element at the specified index
 	FMassEntityHandle& GetEntityArrayElementRef(int32 ChunkBase, int32 IndexWithinChunk)
 	{
-		checkSlow(reinterpret_cast<SIZE_T>(RawMemory + ChunkBase + IndexWithinChunk * sizeof(FMassEntityHandle))< AllocSize);
-		return reinterpret_cast<FMassEntityHandle*>(RawMemory + ChunkBase)[IndexWithinChunk];
+		uint8* RawMemoryChunkBase = RawMemory + ChunkBase;
+		checkSlow(ChunkBase + IndexWithinChunk * sizeof(FMassEntityHandle) < AllocSize
+			&& (reinterpret_cast<SIZE_T>(RawMemoryChunkBase) % alignof(FMassEntityHandle)) == 0);
+		return reinterpret_cast<FMassEntityHandle*>(RawMemoryChunkBase)[IndexWithinChunk];
 	}
 
 	const FMassEntityHandle* GetEntityArray(int32 ChunkBase) const
 	{
-		checkSlow(reinterpret_cast<SIZE_T>(RawMemory + ChunkBase)< AllocSize);
-		return reinterpret_cast<const FMassEntityHandle*>(RawMemory + ChunkBase);
+		uint8* RawMemoryChunkBase = RawMemory + ChunkBase;
+		checkSlow(ChunkBase < AllocSize
+			&& (reinterpret_cast<SIZE_T>(RawMemoryChunkBase) % alignof(FMassEntityHandle)) == 0);
+		return reinterpret_cast<const FMassEntityHandle*>(RawMemoryChunkBase);
 	}
 
 	uint8* GetRawMemory() const
