[Mass] Added a new debug-time construct simplifying vislogging entities

diff --git a/Engine/Source/Runtime/MassEntity/Private/MassEntityQuery.cpp b/Engine/Source/Runtime/MassEntity/Private/MassEntityQuery.cpp
--- a/Engine/Source/Runtime/MassEntity/Private/MassEntityQuery.cpp
+++ b/Engine/Source/Runtime/MassEntity/Private/MassEntityQuery.cpp
@@ -16,6 +16,7 @@
 
 #if WITH_MASSENTITY_DEBUG
 #include "MassRequirementAccessDetector.h"
+#include "MassDebugLogging.h"
 #endif // WITH_MASSENTITY_DEBUG
 
 namespace UE::Mass::Tweakables
@@ -657,6 +658,16 @@ void FMassEntityQuery::ExportRequirements(FMassExecutionRequirements& OutRequire
 	FMassFragmentRequirements::ExportRequirements(OutRequirements);
 }
 
+void FMassEntityQuery::DebugEnableEntityOwnerLogging()
+{
+#if WITH_MASSENTITY_DEBUG
+	if (RequiredOptionalFragments.Contains<FMassDebugLogFragment>() == false)
+	{
+		AddRequirement<FMassDebugLogFragment>(EMassFragmentAccess::ReadOnly, EMassFragmentPresence::Optional);
+	}
+#endif // WITH_MASSENTITY_DEBUG
+}
+
 //-----------------------------------------------------------------------------
 // DEPRECATED
 //-----------------------------------------------------------------------------
diff --git a/Engine/Source/Runtime/MassEntity/Public/MassDebugLogging.h b/Engine/Source/Runtime/MassEntity/Public/MassDebugLogging.h
new file mode 100644
--- /dev/null
+++ b/Engine/Source/Runtime/MassEntity/Public/MassDebugLogging.h
@@ -0,0 +1,84 @@
+// Copyright Epic Games, Inc. All Rights Reserved.
+
+#pragma once
+
+#include "MassEntityElementTypes.h"
+#include "MassProcessingTypes.h"
+#if WITH_MASSENTITY_DEBUG
+#include "MassExecutionContext.h"
+#include "MassEntityHandle.h"
+#include "MassDebugger.h"
+#include "VisualLogger/VisualLogger.h"
+#endif // WITH_MASSENTITY_DEBUG
+#include "MassDebugLogging.generated.h"
+
+
+struct FMassExecutionContext;
+class UObject;
+
+USTRUCT()
+struct MASSENTITY_API FMassDebugLogFragment : public FMassFragment
+{
+	GENERATED_BODY()
+
+	UPROPERTY()
+	TWeakObjectPtr<UObject> LogOwner = nullptr;
+};
+
+namespace UE::Mass::Debug
+{
+	struct FLoggingContext
+	{
+#if WITH_MASSENTITY_DEBUG
+		explicit FLoggingContext(const FMassExecutionContext& Context, bool bInCheckVisualLoggerForRecording = true)
+			: DebugFragmentsView(Context.GetFragmentView<FMassDebugLogFragment>())
+			, EntityListView(Context.GetEntities())
+			, bCheckVisualLoggerForRecording(bInCheckVisualLoggerForRecording)
+		{	
+		}
+
+		FORCEINLINE bool ShouldLogEntity(int32 EntityIndex, FColor* OutEntityColor = nullptr) const
+		{
+#if ENABLE_VISUAL_LOG
+			if (DebugFragmentsView.IsValidIndex(EntityIndex)
+				&& DebugFragmentsView[EntityIndex].LogOwner != nullptr 
+				&& (bCheckVisualLoggerForRecording == false || FVisualLogger::IsRecording()))
+			{
+				return true;
+			}
+#endif // ENABLE_VISUAL_LOG
+			// If no actor owner or the visual logger is not recording, base it on the mass debugger
+			return EntityListView.IsValidIndex(EntityIndex)
+				&& IsDebuggingEntity(EntityListView[EntityIndex], OutEntityColor);
+		}
+
+		FORCEINLINE UObject* GetLogOwner(int32 EntityIndex, UObject* FallbackOwner = nullptr) const
+		{
+			return DebugFragmentsView.IsValidIndex(EntityIndex)
+				? DebugFragmentsView[EntityIndex].LogOwner.Get()
+				: FallbackOwner;
+		}
+
+	private:
+		const TConstArrayView<FMassDebugLogFragment> DebugFragmentsView;
+		const TConstArrayView<FMassEntityHandle> EntityListView;
+		const bool bCheckVisualLoggerForRecording = true;
+
+#else
+
+		explicit FLoggingContext(const FMassExecutionContext& Context, bool bInCheckVisualLoggerForRecording = true)
+		{	
+		}
+
+		bool ShouldLogEntity(int32 EntityIndex, FColor* OutEntityColor = nullptr) const
+		{
+			return false;
+		}
+
+		FORCEINLINE UObject* GetLogOwner(int32 EntityIndex, UObject* FallbackOwner = nullptr) const
+		{
+			return FallbackOwner;
+		}
+#endif // WITH_MASSENTITY_DEBUG
+	};
+}
\ No newline at end of file
diff --git a/Engine/Source/Runtime/MassEntity/Public/MassEntityQuery.h b/Engine/Source/Runtime/MassEntity/Public/MassEntityQuery.h
--- a/Engine/Source/Runtime/MassEntity/Public/MassEntityQuery.h
+++ b/Engine/Source/Runtime/MassEntity/Public/MassEntityQuery.h
@@ -213,6 +213,12 @@ public:
 	 */
 	void SetParallelCommandBufferEnabled(const bool bInAllowParallelCommands) { bAllowParallelCommands = bInAllowParallelCommands; }
 
+	/**
+	 * Configures the query to support per-entity logging based on their individual UObject "owners",
+	 * as declared via debug fragments.
+	 */
+	void DebugEnableEntityOwnerLogging();
+
 private:
 	/**
 	 * Incrementally sorts all ValidArchetypes to fill OrderedArchetypeIndices with the expected order of archetype processing.
