[Mass]
Fix crash from captured 'this' on a struct emplaced in a TArray

diff --git a/Engine/Source/Runtime/MassEntity/Private/MassDebugger.cpp b/Engine/Source/Runtime/MassEntity/Private/MassDebugger.cpp
--- a/Engine/Source/Runtime/MassEntity/Private/MassDebugger.cpp
+++ b/Engine/Source/Runtime/MassEntity/Private/MassDebugger.cpp
@@ -1224,14 +1224,15 @@ FMassDebugger::FEnvironment::FEnvironment(const FMassEntityManager& InEntityMana
 	: EntityManager(InEntityManager.AsWeak())
 {
 #if UE_MASS_TRACE_ENABLED
-	TraceStartedDelegateHandle = FTraceAuxiliary::OnTraceStarted.AddLambda([this](FTraceAuxiliary::EConnectionType TraceType, const FString& TraceDestination)
+	TWeakPtr<const FMassEntityManager> WeakEntityManager = EntityManager;
+	TraceStartedDelegateHandle = FTraceAuxiliary::OnTraceStarted.AddLambda([WeakEntityManager](FTraceAuxiliary::EConnectionType TraceType, const FString& TraceDestination)
 	{
-		if (!EntityManager.IsValid())
+		if (!WeakEntityManager.IsValid())
 		{
 			return;
 		}
 
-		const FMassEntityManager& Manager = *EntityManager.Pin();
+		const FMassEntityManager& Manager = *WeakEntityManager.Pin();
 
 		ForEachArchetype(Manager, [](FMassArchetypeHandle ArchetypeHandle)
 		{
