[Mass] Added a way to ask EntityManager whether a given shared fragment has already been registered

diff --git a/Engine/Source/Runtime/MassEntity/Public/MassEntityManager.h b/Engine/Source/Runtime/MassEntity/Public/MassEntityManager.h
--- a/Engine/Source/Runtime/MassEntity/Public/MassEntityManager.h
+++ b/Engine/Source/Runtime/MassEntity/Public/MassEntityManager.h
@@ -585,6 +585,12 @@ public:
 			: FConstStructView();
 	}
 
+	bool IsConstSharedFragmentRegistered(const FConstSharedStruct& Fragment) const
+	{
+		const uint32 Hash = UE::StructUtils::GetStructCrc32(Fragment);
+		return ConstSharedFragmentsContainer.Contains(Hash);
+	}
+
 	template<typename SharedFragmentType>
 	TConstArrayView<FSharedStruct> GetSharedFragmentsOfType()
 	{
@@ -622,6 +628,12 @@ public:
 			: FConstStructView();
 	}
 
+	bool IsSharedFragmentRegistered(const FSharedStruct& Fragment) const
+	{
+		const uint32 Hash = UE::StructUtils::GetStructCrc32(Fragment);
+		return SharedFragmentsContainer.Contains(Hash);
+	}
+
 	template<typename T>
 	FConstStructView GetElementDataStruct(FMassEntityHandle EntityHandle, TNotNull<const UScriptStruct*> FragmentType) const
 	{
@@ -1254,6 +1266,11 @@ private:
 			return Instances[Index];
 		}
 
+		bool Contains(const int32 Hash) const
+		{
+			return HashToInstanceIndexMap.FindByHash(Hash, Hash) != nullptr;
+		}
+
 		TSharedStructType& operator[](const int32 Index)
 		{
 			return Instances[Index];
