[Mass] Added functionality that creates EntityBuilder instances configured in-bulk with provided data. This functionality is then used to create a FMassEntityTemplateData function that spits out a builder configured with template's data.

diff --git a/Engine/Plugins/Runtime/MassGameplay/Source/MassSpawner/Private/MassEntityTemplate.cpp b/Engine/Plugins/Runtime/MassGameplay/Source/MassSpawner/Private/MassEntityTemplate.cpp
--- a/Engine/Plugins/Runtime/MassGameplay/Source/MassSpawner/Private/MassEntityTemplate.cpp
+++ b/Engine/Plugins/Runtime/MassGameplay/Source/MassSpawner/Private/MassEntityTemplate.cpp
@@ -5,6 +5,7 @@
 #include "VisualLogger/VisualLoggerTypes.h"
 #include "MassDebugger.h"
 #include "Algo/Find.h"
+#include "MassEntityBuilder.h"
 
 //----------------------------------------------------------------------//
 //  FMassEntityTemplateID
@@ -57,6 +58,20 @@ bool FMassEntityTemplateData::SlowIsEquivalent(const FMassEntityTemplateData& Ot
 	return true;
 }
 
+UE::Mass::FEntityBuilder FMassEntityTemplateData::CreateEntityBuilder(const TSharedRef<FMassEntityManager>& InEntityManager) const
+{
+	UE::Mass::FEntityBuilder Builder = UE::Mass::FEntityBuilder::Make(InEntityManager, Composition
+		, InitialFragmentValues
+		, SharedFragmentValues.GetConstSharedFragments()
+		, SharedFragmentValues.GetSharedFragments());
+
+	Builder.ConfigureArchetypeCreation(CreationParams);
+
+	ensureMsgf(ObjectInitializers.Num() == 0, TEXT("%hs: EntityBuilder doesn't support object initializers (yet)."), __FUNCTION__);
+
+	return Builder;
+}
+
 uint32 GetTypeHash(const FMassEntityTemplateData& Template)
 {
 	// @todo: using the template name is temporary solution to allow to tell two templates apart based on something 
@@ -154,6 +169,11 @@ FString FMassEntityTemplate::DebugGetDescription(FMassEntityManager* EntityManag
 	return MoveTemp(Ar);
 }
 
+UE::Mass::FEntityBuilder FMassEntityTemplate::CreateEntityBuilder(const TSharedRef<FMassEntityManager>& InEntityManager) const
+{
+	return TemplateData.CreateEntityBuilder(InEntityManager);
+}
+
 //-----------------------------------------------------------------------------
 // FMassEntityTemplateIDFactory
 //-----------------------------------------------------------------------------
diff --git a/Engine/Plugins/Runtime/MassGameplay/Source/MassSpawner/Public/MassEntityTemplate.h b/Engine/Plugins/Runtime/MassGameplay/Source/MassSpawner/Public/MassEntityTemplate.h
--- a/Engine/Plugins/Runtime/MassGameplay/Source/MassSpawner/Public/MassEntityTemplate.h
+++ b/Engine/Plugins/Runtime/MassGameplay/Source/MassSpawner/Public/MassEntityTemplate.h
@@ -21,7 +21,10 @@ struct FMassEntityView;
 struct FMassEntityTemplateIDFactory;
 struct FMassEntityTemplate;
 struct FMassEntityManager;
-
+namespace UE::Mass
+{
+	struct FEntityBuilder;
+}
 
 //ID of the template an entity is using
 USTRUCT()
@@ -310,6 +313,8 @@ struct MASSSPAWNER_API FMassEntityTemplateData
 	bool SlowIsEquivalent(const FMassEntityTemplateData& Other) const;
 
 	FMassArchetypeCreationParams& GetArchetypeCreationParams();
+
+	UE::Mass::FEntityBuilder CreateEntityBuilder(const TSharedRef<FMassEntityManager>& InEntityManager) const;
 	
 protected:
 	FMassArchetypeCompositionDescriptor Composition;
@@ -363,6 +368,8 @@ struct MASSSPAWNER_API FMassEntityTemplate final : public TSharedFromThis<FMassE
 	TConstArrayView<FInstancedStruct> GetInitialFragmentValues() const;
 	const FMassEntityTemplateData& GetTemplateData() const;
 
+	UE::Mass::FEntityBuilder CreateEntityBuilder(const TSharedRef<FMassEntityManager>& InEntityManager) const;
+
 private:
 	FMassEntityTemplateData TemplateData;
 	FMassArchetypeHandle Archetype;
@@ -485,3 +492,4 @@ inline const FMassEntityTemplateData& FMassEntityTemplate::GetTemplateData() con
 {
 	return TemplateData;
 }
+
diff --git a/Engine/Source/Runtime/MassEntity/Private/MassEntityBuilder.cpp b/Engine/Source/Runtime/MassEntity/Private/MassEntityBuilder.cpp
--- a/Engine/Source/Runtime/MassEntity/Private/MassEntityBuilder.cpp
+++ b/Engine/Source/Runtime/MassEntity/Private/MassEntityBuilder.cpp
@@ -5,10 +5,27 @@
 #include "MassCommands.h"
 #include "MassEntityManager.h"
 #include "VisualLogger/VisualLogger.h"
-
+#if WITH_MASSENTITY_DEBUG
+#include "HAL/IConsoleManager.h"
+#endif // WITH_MASSENTITY_DEBUG
 
 namespace UE::Mass
 {
+#if WITH_MASSENTITY_DEBUG
+	namespace Debug
+	{
+		bool bValidateEntityBuilderMakeInput = true;
+		namespace
+		{
+			FAutoConsoleVariableRef AnonymousCVars[] = {
+				{TEXT("mass.debug.ValidateEntityBuilderMakeInput"), bValidateEntityBuilderMakeInput
+					, TEXT("When set, every call to FEntityBuilder::Make will verify if the struct values provided match declared entity composition.")
+					, ECVF_Cheat}
+			};
+		}
+	}
+#endif // WITH_MASSENTITY_DEBUG
+
 	namespace Private
 	{
 		struct FEntityBuilderHelper
@@ -57,6 +74,45 @@ namespace UE::Mass
 				}
 			}
 		};
+
+#if WITH_MASSENTITY_DEBUG
+		template<typename TElement, typename TBitset, typename TWrapper>
+		bool CheckStructContainer(TConstArrayView<TWrapper> Container, const TBitset& Bitset, const UObject* LogOwner)
+		{
+			bool bIssuesFound = false;
+
+			for (const TWrapper& Element : Container)
+			{
+				if (Mass::IsA<TElement>(Element.GetScriptStruct()))
+				{
+					if (Bitset.Contains(*Element.GetScriptStruct()) == false)
+					{
+						bIssuesFound = true;
+						UE_VLOG_UELOG(LogOwner, LogMass, Error, TEXT("%hs: input Composition doesn't contain %s"), __FUNCTION__, *GetNameSafe(Element.GetScriptStruct()));
+					}
+				}
+				else
+				{
+					bIssuesFound = true;
+					UE_VLOG_UELOG(LogOwner, LogMass, Error, TEXT("%hs: %s is not a valid %s type")
+						, __FUNCTION__, *GetNameSafe(Element.GetScriptStruct()), *TElement::StaticStruct()->GetName());
+				}
+			}
+
+			return bIssuesFound;
+		}
+
+		bool ValidateMakeInput(const TSharedRef<FMassEntityManager>& InEntityManager, const FMassArchetypeCompositionDescriptor& Composition
+			, TConstArrayView<FInstancedStruct> InitialFragmentValues, TConstArrayView<FConstSharedStruct> ConstSharedFragments, TConstArrayView<FSharedStruct> SharedFragments)
+		{
+			const UObject* LogOwner = InEntityManager->GetOwner();
+			bool bIssuesFound = CheckStructContainer<FMassFragment>(InitialFragmentValues, Composition.Fragments, LogOwner);
+			bIssuesFound |= CheckStructContainer<FMassConstSharedFragment>(ConstSharedFragments, Composition.ConstSharedFragments, LogOwner);
+			bIssuesFound |= CheckStructContainer<FMassSharedFragment>(SharedFragments, Composition.SharedFragments, LogOwner);
+			
+			return !bIssuesFound;
+		}
+#endif // WITH_MASSENTITY_DEBUG
 	}
 
 //-----------------------------------------------------------------------------
@@ -156,9 +212,54 @@ FEntityBuilder::~FEntityBuilder()
 	ConditionallyReleaseEntityHandle();
 }
 
+FEntityBuilder FEntityBuilder::Make(const TSharedRef<FMassEntityManager>& InEntityManager, const FMassArchetypeCompositionDescriptor& Composition
+	, TConstArrayView<FInstancedStruct> InitialFragmentValues, TConstArrayView<FConstSharedStruct> ConstSharedFragments, TConstArrayView<FSharedStruct> SharedFragments)
+{
+	FEntityBuilder Builder(InEntityManager);
+
+#if WITH_MASSENTITY_DEBUG
+	if (Debug::bValidateEntityBuilderMakeInput)
+	{
+		ensureMsgf(Private::ValidateMakeInput(InEntityManager, Composition, InitialFragmentValues, ConstSharedFragments, SharedFragments)
+			, TEXT("%hs: failed input validation. See log for details."), __FUNCTION__);
+	}
+#endif // WITH_MASSENTITY_DEBUG
+
+	Builder.Composition = Composition;
+	Builder.Fragments = InitialFragmentValues;
+	Builder.SharedFragments = SharedFragments;
+	Builder.ConstSharedFragments = ConstSharedFragments;
+
+	return Builder;
+}
+
+FEntityBuilder FEntityBuilder::Make(const TSharedRef<FMassEntityManager>& InEntityManager
+		, const FMassArchetypeCompositionDescriptor& Composition
+		, TArray<FInstancedStruct>&& InitialFragmentValues
+		, TArray<FConstSharedStruct>&& ConstSharedFragments
+		, TArray<FSharedStruct>&& SharedFragments)
+{
+	FEntityBuilder Builder(InEntityManager);
+
+#if WITH_MASSENTITY_DEBUG
+	if (Debug::bValidateEntityBuilderMakeInput)
+	{
+		ensureMsgf(Private::ValidateMakeInput(InEntityManager, Composition, InitialFragmentValues, ConstSharedFragments, SharedFragments)
+			, TEXT("%hs: failed input validation. See log for details."), __FUNCTION__);
+	}
+#endif // WITH_MASSENTITY_DEBUG
+
+	Builder.Composition = Composition;
+	Builder.Fragments = InitialFragmentValues;
+	Builder.SharedFragments.Append(Forward<TArray<FSharedStruct>>(SharedFragments));
+	Builder.ConstSharedFragments.Append(Forward<TArray<FConstSharedStruct>>(ConstSharedFragments));
+
+	return Builder;
+}
+
 FMassEntityHandle FEntityBuilder::Commit()
 {
-	// @todo consider locking every builder instance to a single thread to prevent concurent add/flush
+	// @todo consider locking every builder instance to a single thread to prevent concurrent add/flush
 
 	if (!testableEnsureMsgf(State != EState::Committed, TEXT("Trying to commit an already committed EntityBuilder. This request will be ignored.")))
 	{
@@ -334,7 +435,7 @@ void FEntityBuilder::CacheArchetypeHandle()
 {
 	if (CachedArchetypeHandle.IsValid() == false)
 	{
-		CachedArchetypeHandle = EntityManager->CreateArchetype(Composition);
+		CachedArchetypeHandle = EntityManager->CreateArchetype(Composition, ArchetypeCreationParams);
 	}
 }
 
diff --git a/Engine/Source/Runtime/MassEntity/Public/MassEntityBuilder.h b/Engine/Source/Runtime/MassEntity/Public/MassEntityBuilder.h
--- a/Engine/Source/Runtime/MassEntity/Public/MassEntityBuilder.h
+++ b/Engine/Source/Runtime/MassEntity/Public/MassEntityBuilder.h
@@ -73,6 +73,20 @@ struct MASSENTITY_API FEntityBuilder
 	/** Destructor - automatically commits entity creation if not explicitly aborted or committed beforehand. */
 	~FEntityBuilder();
 
+	/** Creates an instance of FEntityBuilder and populates it with provided data */
+	static FEntityBuilder Make(const TSharedRef<FMassEntityManager>& InEntityManager
+		, const FMassArchetypeCompositionDescriptor& Composition
+		, TConstArrayView<FInstancedStruct> InitialFragmentValues = {}
+		, TConstArrayView<FConstSharedStruct> ConstSharedFragments = {}
+		, TConstArrayView<FSharedStruct> SharedFragments = {});
+
+	/** Creates an instance of FEntityBuilder and populates it with provided data, using move-semantics on said data */
+	static FEntityBuilder Make(const TSharedRef<FMassEntityManager>& InEntityManager
+		, const FMassArchetypeCompositionDescriptor& Composition
+		, TArray<FInstancedStruct>&& InitialFragmentValues
+		, TArray<FConstSharedStruct>&& ConstSharedFragments
+		, TArray<FSharedStruct>&& SharedFragments);
+
 	/**
 	 * Finalizes the creation of the entity with the specified fragments and configurations.
 	 * Note that this function needs to be called manually, no automated entity creation will take place upon builder's destruction. 
@@ -170,6 +184,13 @@ struct MASSENTITY_API FEntityBuilder
 	template<typename T>
 	T* Find() requires (CElement<T> && !(CTag<T> || CChunkFragment<T>));
 
+	/**
+	 * Advanced functionality. Can be used to provide additional parameters that will be used to
+	 * create the entity's target archetype. Note that these parameters will take effect only if
+	 * the target archetype doesn't exist yet.
+	 */
+	void ConfigureArchetypeCreation(const FMassArchetypeCreationParams& InCreationParams);
+
 	/** Converts the builder to a FMassEntityHandle, reserving the entity handle if not already committed. */
 	[[nodiscard]] FMassEntityHandle GetEntityHandle() const;
 
@@ -206,6 +227,9 @@ private:
 	FMassArchetypeSharedFragmentValues CachedSharedFragmentValues;
 	FMassArchetypeHandle CachedArchetypeHandle;
 
+	/** stores optional FMassArchetypeCreationParams, that will be used if the target archetype doesn't exist yet */
+	FMassArchetypeCreationParams ArchetypeCreationParams;
+
 	template<CFragment T> 
 	TArray<FInstancedStruct>& GetInstancedStructContainerInternal() 
 	{ 
@@ -388,4 +412,9 @@ inline TSharedRef<FMassEntityManager> FEntityBuilder::GetEntityManager()
 	return EntityManager;
 }
 
+inline void FEntityBuilder::ConfigureArchetypeCreation(const FMassArchetypeCreationParams& InCreationParams)
+{
+	ArchetypeCreationParams = InCreationParams;
+}
+
 } // namespace UE::Mass
