[Mass] Added a way to add entities to commands with move semantics

diff --git a/Engine/Source/Runtime/MassEntity/Public/MassCommands.h b/Engine/Source/Runtime/MassEntity/Public/MassCommands.h
--- a/Engine/Source/Runtime/MassEntity/Public/MassCommands.h
+++ b/Engine/Source/Runtime/MassEntity/Public/MassCommands.h
@@ -150,6 +150,13 @@ struct FMassBatchedEntityCommand : public FMassBatchedCommand
 		bHasWork = true;
 	}
 
+	void Add(TArray<FMassEntityHandle>&& Entities)
+	{
+		UE_MT_SCOPED_WRITE_ACCESS(EntitiesAccessDetector);
+		TargetEntities.Append(Forward<TArray<FMassEntityHandle>>(Entities));
+		bHasWork = true;
+	}
+
 protected:
 	virtual SIZE_T GetAllocatedSize() const
 	{
