[Mass] Added a functions to MassEntityQuery that return matching entities via arrays of  FMassArchetypeEntityCollection or handles.

#rnx The CL also contains a fix to FMassArchetypeData::ExportEntityHandles which wasn't handling 0-length FMassArchetypeEntityCollection::FArchetypeEntityRange-s. Now it does.

diff --git a/Engine/Source/Runtime/MassEntity/Internal/MassArchetypeData.h b/Engine/Source/Runtime/MassEntity/Internal/MassArchetypeData.h
--- a/Engine/Source/Runtime/MassEntity/Internal/MassArchetypeData.h
+++ b/Engine/Source/Runtime/MassEntity/Internal/MassArchetypeData.h
@@ -339,6 +339,8 @@ public:
 
 	void ExportEntityHandles(const TConstArrayView<FMassArchetypeEntityCollection::FArchetypeEntityRange> Ranges, TArray<FMassEntityHandle>& InOutHandles) const;
 
+	void ExportEntityHandles(TArray<FMassEntityHandle>& InOutHandles) const;
+
 	// Converts the list of fragments into a user-readable debug string
 	FString DebugGetDescription() const;
 
diff --git a/Engine/Source/Runtime/MassEntity/Private/MassArchetypeData.cpp b/Engine/Source/Runtime/MassEntity/Private/MassArchetypeData.cpp
--- a/Engine/Source/Runtime/MassEntity/Private/MassArchetypeData.cpp
+++ b/Engine/Source/Runtime/MassEntity/Private/MassArchetypeData.cpp
@@ -914,19 +914,29 @@ void FMassArchetypeData::ExportEntityHandles(const TConstArrayView<FMassArchetyp
 	int32 TotalEntities = 0;
 	for (const FMassArchetypeEntityCollection::FArchetypeEntityRange& Range : Ranges)
 	{
-		TotalEntities += Range.Length;
+		check(Chunks.IsValidIndex(Range.ChunkIndex));
+		TotalEntities += (Range.Length > 0) ? Range.Length : (Chunks[Range.ChunkIndex].GetNumInstances() - Range.SubchunkStart);
 	}
+
 	int32 StartIndex = InOutHandles.Num();
 	InOutHandles.AddUninitialized(TotalEntities);
 
 	for (const FMassArchetypeEntityCollection::FArchetypeEntityRange& Range : Ranges)
 	{
-		check(Chunks.IsValidIndex(Range.ChunkIndex))
 		const FMassEntityHandle* EntitiesArray = Chunks[Range.ChunkIndex].GetEntityArray(EntityListOffsetWithinChunk);
 
-		FMemory::Memcpy(&InOutHandles[StartIndex], &EntitiesArray[Range.SubchunkStart], Range.Length * sizeof(FMassEntityHandle));
+		const int32 RangeLength = Range.Length > 0 ? Range.Length : (Chunks[Range.ChunkIndex].GetNumInstances() - Range.SubchunkStart);
+		FMemory::Memcpy(&InOutHandles[StartIndex], &EntitiesArray[Range.SubchunkStart], RangeLength * sizeof(FMassEntityHandle));
+
+		StartIndex += RangeLength;
+	}
+}
 
-		StartIndex += Range.Length;
+void FMassArchetypeData::ExportEntityHandles(TArray<FMassEntityHandle>& InOutHandles) const
+{
+	for (const FMassArchetypeChunk& Chunk : Chunks)
+	{
+		InOutHandles.Append(Chunk.GetEntityArray(EntityListOffsetWithinChunk), Chunk.GetNumInstances());
 	}
 }
 
diff --git a/Engine/Source/Runtime/MassEntity/Private/MassEntityQuery.cpp b/Engine/Source/Runtime/MassEntity/Private/MassEntityQuery.cpp
--- a/Engine/Source/Runtime/MassEntity/Private/MassEntityQuery.cpp
+++ b/Engine/Source/Runtime/MassEntity/Private/MassEntityQuery.cpp
@@ -534,6 +534,34 @@ bool FMassEntityQuery::HasMatchingEntities()
 	return false;
 }
 
+TArray<FMassArchetypeEntityCollection> FMassEntityQuery::CreateMatchingEntitiesCollection()
+{
+	TArray<FMassArchetypeEntityCollection> Collections;
+	CacheArchetypes();
+
+	Collections.Reserve(ValidArchetypes.Num());
+	for (FMassArchetypeHandle& ArchetypeHandle : ValidArchetypes)
+	{
+		Collections.Emplace(ArchetypeHandle);
+	}
+
+	return Collections;
+}
+
+TArray<FMassEntityHandle> FMassEntityQuery::GetMatchingEntityHandles()
+{
+	TArray<FMassEntityHandle> Handles;
+	CacheArchetypes();
+
+	for (FMassArchetypeHandle& ArchetypeHandle : ValidArchetypes)
+	{
+		FMassArchetypeData& ArchetypeData = FMassArchetypeHelper::ArchetypeDataFromHandleChecked(ArchetypeHandle);
+		ArchetypeData.ExportEntityHandles(Handles);
+	}
+
+	return Handles;
+}
+
 void FMassEntityQuery::GroupBy(UE::Mass::FArchetypeGroupType GroupType)
 {
 	GroupBy(GroupType, [](const UE::Mass::FArchetypeGroupID A, const UE::Mass::FArchetypeGroupID B)
diff --git a/Engine/Source/Runtime/MassEntity/Public/MassEntityQuery.h b/Engine/Source/Runtime/MassEntity/Public/MassEntityQuery.h
--- a/Engine/Source/Runtime/MassEntity/Public/MassEntityQuery.h
+++ b/Engine/Source/Runtime/MassEntity/Public/MassEntityQuery.h
@@ -171,6 +171,17 @@ public:
 	 */
 	bool HasMatchingEntities();
 
+	/**
+	 * Creates an array of FMassArchetypeEntityCollection instances that identify all the entities
+	 * currently matching this query.
+	 */
+	TArray<FMassArchetypeEntityCollection> CreateMatchingEntitiesCollection();
+
+	/**
+	 * Fetches entity handles of all the entities currently matching this query.
+	 */
+	TArray<FMassEntityHandle> GetMatchingEntityHandles();
+
 	/** 
 	 * Sets a chunk filter condition that will be applied to each chunk of all valid archetypes. Note 
 	 * that this condition won't be applied when a specific entity collection is used (via FMassArchetypeEntityCollection )
