[Mass] removed FMassArchetypeHandle.HandleVersion that has been submitted in error

#rnx

diff --git a/Engine/Source/Runtime/MassEntity/Public/MassArchetypeTypes.h b/Engine/Source/Runtime/MassEntity/Public/MassArchetypeTypes.h
--- a/Engine/Source/Runtime/MassEntity/Public/MassArchetypeTypes.h
+++ b/Engine/Source/Runtime/MassEntity/Public/MassArchetypeTypes.h
@@ -41,21 +41,12 @@ struct FMassArchetypeHandle final
 	bool operator==(const FMassArchetypeHandle& Other) const { return DataPtr == Other.DataPtr; }
 	bool operator!=(const FMassArchetypeHandle& Other) const { return DataPtr != Other.DataPtr; }
 
-	bool IsUpToData() const;
-
 	MASSENTITY_API friend uint32 GetTypeHash(const FMassArchetypeHandle& Instance);
 private:
 	FMassArchetypeHandle(const TSharedPtr<FMassArchetypeData>& InDataPtr)
 	: DataPtr(InDataPtr)
 	{}
 	TSharedPtr<FMassArchetypeData> DataPtr;
-	/**
-	 * This value indicates whether the target archetype had its entities moved around since the handle creations.
-	 * The information is useful in a couple of scenarios (like making sure an entity collection is up to date),
-	 * but in most cases the users should not concern themselves with this value.
-	 * Note that the value is not used as part of hash calculation, it's effectively transient. 
-	 */
-	uint32 HandleVersion = 0;
 
 	friend FMassArchetypeHelper;
 	friend FMassEntityManager;
