[Mass] Fixed some typos and string format

diff --git a/Engine/Source/Runtime/MassEntity/Private/MassArchetypeData.cpp b/Engine/Source/Runtime/MassEntity/Private/MassArchetypeData.cpp
--- a/Engine/Source/Runtime/MassEntity/Private/MassArchetypeData.cpp
+++ b/Engine/Source/Runtime/MassEntity/Private/MassArchetypeData.cpp
@@ -989,11 +989,11 @@ void FMassArchetypeData::DebugPrintArchetype(FOutputDevice& Ar)
 {
 	Ar.Logf(ELogVerbosity::Log, TEXT("Name: %s"), *GetCombinedDebugNamesAsString());
 
-	FStringOutputDevice TagsDecription;
-	CompositionDescriptor.Tags.DebugGetStringDesc(TagsDecription);
-	Ar.Logf(ELogVerbosity::Log, TEXT("Tags: %s"), *TagsDecription);
+	FStringOutputDevice TagsDescription;
+	CompositionDescriptor.Tags.DebugGetStringDesc(TagsDescription);
+	Ar.Logf(ELogVerbosity::Log, TEXT("Tags: %s"), *TagsDescription);
 	Ar.Logf(ELogVerbosity::Log, TEXT("Fragments: %s"), *DebugGetDescription());
-	Ar.Logf(ELogVerbosity::Log, TEXT("\tChunks: %d x %d KB = %d KB total"), Chunks.Num(), GetChunkAllocSize() / 1024, (GetChunkAllocSize()*Chunks.Num()) / 1024);
+	Ar.Logf(ELogVerbosity::Log, TEXT("\tChunks: %d x %llu KB = %llu KB total"), Chunks.Num(), GetChunkAllocSize() / 1024, (GetChunkAllocSize()*Chunks.Num()) / 1024);
 	
 	int ChunkWithFragmentsCount = 0;
 	for (FMassArchetypeChunk& Chunk : Chunks)
@@ -1043,16 +1043,16 @@ void FMassArchetypeData::DebugPrintArchetype(FOutputDevice& Ar)
 
 	//@TODO: Print out padding in between things?
 
-	const SIZE_T UnusuablePaddingOffset = TotalBytesPerEntity * NumEntitiesPerChunk;
-	const SIZE_T UnusuablePaddingAmount = GetChunkAllocSize() - UnusuablePaddingOffset;
-	if (UnusuablePaddingAmount > 0)
+	const SIZE_T UnusablePaddingOffset = TotalBytesPerEntity * NumEntitiesPerChunk;
+	const SIZE_T UnusablePaddingAmount = GetChunkAllocSize() - UnusablePaddingOffset;
+	if (UnusablePaddingAmount > 0)
 	{
-		Ar.Logf(ELogVerbosity::Log, TEXT("\tOffset 0x%04X: WastePadding[] (%d bytes total)"), UnusuablePaddingOffset, UnusuablePaddingAmount);
+		Ar.Logf(ELogVerbosity::Log, TEXT("\tOffset 0x%04llX: WastePadding[] (%llu bytes total)"), UnusablePaddingOffset, UnusablePaddingAmount);
 	}
 
-	if (GetChunkAllocSize() != TotalBytesOfValidData + UnusuablePaddingAmount)
+	if (GetChunkAllocSize() != TotalBytesOfValidData + UnusablePaddingAmount)
 	{
-		Ar.Logf(ELogVerbosity::Log, TEXT("\t@TODO: EXTRA PADDING HERE:  TotalBytesOfValidData: %d (%d missing)"), TotalBytesOfValidData, GetChunkAllocSize() - TotalBytesOfValidData);
+		Ar.Logf(ELogVerbosity::Log, TEXT("\t@TODO: EXTRA PADDING HERE:  TotalBytesOfValidData: %d (%llu missing)"), TotalBytesOfValidData, GetChunkAllocSize() - TotalBytesOfValidData);
 	}
 }
 
@@ -1080,8 +1080,8 @@ void FMassArchetypeData::SetDebugColor(const FColor InDebugColor)
 	if (InDebugColor == FColor{0})
 	{
 		// pick a color based on the composition
-		const uint32 Compositionhash = CompositionDescriptor.CalculateHash();
-		const uint8* Bytes = reinterpret_cast<const uint8*>(&Compositionhash);
+		const uint32 CompositionHash = CompositionDescriptor.CalculateHash();
+		const uint8* Bytes = reinterpret_cast<const uint8*>(&CompositionHash);
 		
 		const FLinearColor AdjustedColor = FLinearColor::MakeFromHSV8(
 			static_cast<uint8>((Bytes[0] >> 1) + (Bytes[1] >> 1)),
@@ -1311,7 +1311,7 @@ FMassArchetypeEntityCollection::FArchetypeEntityRange FMassArchetypeData::Prepar
 
 	check(DestinationChunk);
 
-	// we might be able to fit in less entitites than requested
+	// we might be able to fit in less entities than requested
 	const int32 NumToAdd = FMath::Min(NumEntitiesPerChunk - StartIndexWithinChunk, Entities.Num());
 	check(NumToAdd);
 	DestinationChunk->AddMultipleInstances(NumToAdd);
@@ -1342,7 +1342,6 @@ void FMassArchetypeData::BatchRemoveEntitiesInternal(const int32 ChunkIndex, con
 	
 	const int32 NumberToMove = FMath::Min(Chunk.GetNumInstances() - (StartIndexWithinChunk + NumberToRemove), NumberToRemove);
 	checkf(NumberToMove >= 0, TEXT("Trying to move a negative number of elements indicates a problem with sub-chunk indicators, it's possibly out of date."));
-	const int32 NumberToCut = FMath::Max(NumberToRemove - NumberToMove, 0);
 
 	if (NumberToMove > 0)
 	{
@@ -1482,7 +1481,7 @@ bool FMassArchetypeHelper::DoesArchetypeMatchRequirements(const FMassArchetypeDa
 	if (OutputDevice)
 	{
 		// do logging
-		OutputDevice->Logf(TEXT("% s")
+		OutputDevice->Logf(TEXT("%s")
 			, *FMassDebugger::GetArchetypeRequirementCompatibilityDescription(Requirements, Archetype.GetCompositionDescriptor()));
 	}
 
