[Mass] Made FMassExecutionContext::FEntityIterator usable in range-based loops while keeping it non-copyable

#rnx

diff --git a/Engine/Source/Runtime/MassEntity/Private/MassExecutionContext.cpp b/Engine/Source/Runtime/MassEntity/Private/MassExecutionContext.cpp
--- a/Engine/Source/Runtime/MassEntity/Private/MassExecutionContext.cpp
+++ b/Engine/Source/Runtime/MassEntity/Private/MassExecutionContext.cpp
@@ -184,6 +184,12 @@ FMassExecutionContext::FEntityIterator FMassExecutionContext::CreateEntityIterat
 	return FEntityIterator(*this, QueriesStack.Last());
 }
 
+FMassExecutionContext& FMassExecutionContext::GetDummyInstance()
+{
+	static FMassExecutionContext DummyContext;
+	return DummyContext;
+}
+
 //-----------------------------------------------------------------------------
 // FMassExecutionContext::FQueryTransientRuntime
 //-----------------------------------------------------------------------------
@@ -197,6 +203,13 @@ FMassExecutionContext::FQueryTransientRuntime& FMassExecutionContext::FQueryTran
 //-----------------------------------------------------------------------------
 // FMassExecutionContext::FEntityIterator
 //-----------------------------------------------------------------------------
+FMassExecutionContext::FEntityIterator::FEntityIterator()
+	: ExecutionContext(FMassExecutionContext::GetDummyInstance())
+	, QueryRuntime(FQueryTransientRuntime::GetDummyInstance())
+{
+	
+}
+
 FMassExecutionContext::FEntityIterator::FEntityIterator(FMassExecutionContext& InExecutionContext)
 	: ExecutionContext(InExecutionContext)
 	, QueryRuntime(FQueryTransientRuntime::GetDummyInstance())
diff --git a/Engine/Source/Runtime/MassEntity/Public/MassExecutionContext.h b/Engine/Source/Runtime/MassEntity/Public/MassExecutionContext.h
--- a/Engine/Source/Runtime/MassEntity/Public/MassExecutionContext.h
+++ b/Engine/Source/Runtime/MassEntity/Public/MassExecutionContext.h
@@ -137,6 +137,9 @@ private:
 	 */
 	FMassExecutionContext& operator=(const FMassExecutionContext& Other) = default;
 
+	static FMassExecutionContext& GetDummyInstance();
+
+	FMassExecutionContext() = default;
 public:
 	explicit FMassExecutionContext(FMassEntityManager& InEntityManager, const float InDeltaTimeSeconds = 0.f, const bool bInFlushDeferredCommands = true);
 	FMassExecutionContext(FMassExecutionContext&& Other) = default;
@@ -207,38 +210,38 @@ public:
 			return *this;
 		}
 
-		FORCEINLINE FEntityIterator operator++(int)
+		FORCEINLINE void operator++(int)
 		{
-			FEntityIterator Temp  = *this;
 			++(*this);
-			return Temp;
 		}
 
-		FEntityIterator& begin()
+		FEntityIterator&& begin()
 		{
-			return *this;
+			return MoveTemp(*this);
 		}
 
 		FEntityIterator end() const
 		{
-			FEntityIterator End = *this;
+			FEntityIterator End;
 			End.EntityIndex = NumEntities;
 			return End;
 		}
 
+		MASSENTITY_API FEntityIterator();
+		MASSENTITY_API FEntityIterator(FEntityIterator&&) = default;
+
 		FEntityIterator& operator=(const FEntityIterator&) = delete;
 		FEntityIterator& operator=(FEntityIterator&&) = delete;
-	private:
-		friend FMassExecutionContext;
-		FEntityIterator(FMassExecutionContext& InExecutionContext);
-		FEntityIterator(FMassExecutionContext& InExecutionContext, FQueryTransientRuntime& InQueryRuntime);
 		/**
 		 * Iterator copying is disabled to avoid additional checks to detect if entity chunk being iterated on changed.
 		 * This decision is to be reconsidered when valid iterator-copying scenarios emerge. 
 		 */
-		FEntityIterator(const FEntityIterator&) = default;
-		FEntityIterator(FEntityIterator&&) = default;
-		
+		FEntityIterator(const FEntityIterator&) = delete;
+
+	private:
+		friend FMassExecutionContext;
+		FEntityIterator(FMassExecutionContext& InExecutionContext);
+		FEntityIterator(FMassExecutionContext& InExecutionContext, FQueryTransientRuntime& InQueryRuntime);
 
 #if WITH_MASSENTITY_DEBUG
 		MASSENTITY_API void TestBreakpoints();
