[Mass] Limited the set of functions API-exported from FMassExecutionContext
Also moved some inline functions to the end of the file, and removed code deprecated up to 5.4.
Also, made FMassExecutionContext::GetDummyInstance public (but not API-exported)

#rnx

diff --git a/Engine/Source/Runtime/MassEntity/Public/MassExecutionContext.h b/Engine/Source/Runtime/MassEntity/Public/MassExecutionContext.h
--- a/Engine/Source/Runtime/MassEntity/Public/MassExecutionContext.h
+++ b/Engine/Source/Runtime/MassEntity/Public/MassExecutionContext.h
@@ -25,7 +25,7 @@
 
 struct FMassEntityQuery;
 
-struct MASSENTITY_API FMassExecutionContext
+struct FMassExecutionContext
 {
 private:
 
@@ -137,15 +137,17 @@ private:
 	 */
 	FMassExecutionContext& operator=(const FMassExecutionContext& Other) = default;
 
+public:
+	MASSENTITY_API explicit FMassExecutionContext(FMassEntityManager& InEntityManager, const float InDeltaTimeSeconds = 0.f, const bool bInFlushDeferredCommands = true);
+	MASSENTITY_API FMassExecutionContext(FMassExecutionContext&& Other) = default;
+	MASSENTITY_API FMassExecutionContext(const FMassExecutionContext& Other);
+	MASSENTITY_API FMassExecutionContext(const FMassExecutionContext& Other, FMassEntityQuery& Query, const TSharedPtr<FMassCommandBuffer>& InCommandBuffer = {});
+	MASSENTITY_API ~FMassExecutionContext();
+
+	/** For internal use only, should never be exported as part of API */
 	static FMassExecutionContext& GetDummyInstance();
 
 	FMassExecutionContext() = default;
-public:
-	explicit FMassExecutionContext(FMassEntityManager& InEntityManager, const float InDeltaTimeSeconds = 0.f, const bool bInFlushDeferredCommands = true);
-	FMassExecutionContext(FMassExecutionContext&& Other) = default;
-	FMassExecutionContext(const FMassExecutionContext& Other);
-	FMassExecutionContext(const FMassExecutionContext& Other, FMassEntityQuery& Query, const TSharedPtr<FMassCommandBuffer>& InCommandBuffer = {});
-	~FMassExecutionContext();
 
 	FMassEntityManager& GetEntityManagerChecked() const;
 	const TSharedRef<FMassEntityManager>& GetSharedEntityManager();
@@ -258,14 +260,14 @@ public:
 	 * Creates an Entity Iterator for the current chunk.
 	 * Supports range-based for loop and can be used directly as an entity index for the current chunk.
 	 */
-	FEntityIterator CreateEntityIterator();
+	MASSENTITY_API FEntityIterator CreateEntityIterator();
 
 	/** Sets bFlushDeferredCommands. Note that setting to True while the system is being executed doesn't result in
 	 *  immediate commands flushing */
 	void SetFlushDeferredCommands(const bool bNewFlushDeferredCommands);
 	void SetDeferredCommandBuffer(const TSharedPtr<FMassCommandBuffer>& InDeferredCommandBuffer);
-	void SetEntityCollection(const FMassArchetypeEntityCollection& InEntityCollection);
-	void SetEntityCollection(FMassArchetypeEntityCollection&& InEntityCollection);
+	MASSENTITY_API void SetEntityCollection(const FMassArchetypeEntityCollection& InEntityCollection);
+	MASSENTITY_API void SetEntityCollection(FMassArchetypeEntityCollection&& InEntityCollection);
 	void ClearEntityCollection();
 	void SetAuxData(const FInstancedStruct& InAuxData);
 
@@ -522,7 +524,7 @@ public:
 		return FragmentType != nullptr && FragmentType == TFragment::StaticStruct();
 	}
 
-	void FlushDeferred();
+	MASSENTITY_API void FlushDeferred();
 
 	void ClearExecutionData();
 	void SetCurrentArchetypeCompositionDescriptor(const FMassArchetypeCompositionDescriptor& Descriptor)
