[Mass] Made UE_MASS_TRACE_ENABLED configurable via projects' build files. Also removed some unused header inclusions.

diff --git a/Engine/Source/Runtime/MassEntity/Private/MassEntityTrace.cpp b/Engine/Source/Runtime/MassEntity/Private/MassEntityTrace.cpp
--- a/Engine/Source/Runtime/MassEntity/Private/MassEntityTrace.cpp
+++ b/Engine/Source/Runtime/MassEntity/Private/MassEntityTrace.cpp
@@ -4,8 +4,6 @@
 
 #if UE_MASS_TRACE_ENABLED
 
-#include "Algo/ForEach.h"
-
 #include "MassArchetypeData.h"
 #include "MassEntityTypes.h"
 #include "MassEntityQuery.h"
diff --git a/Engine/Source/Runtime/MassEntity/Public/MassEntityTrace.h b/Engine/Source/Runtime/MassEntity/Public/MassEntityTrace.h
--- a/Engine/Source/Runtime/MassEntity/Public/MassEntityTrace.h
+++ b/Engine/Source/Runtime/MassEntity/Public/MassEntityTrace.h
@@ -4,14 +4,11 @@
 
 #include "MassProcessingTypes.h"
 #include "Misc/NotNull.h"
-#include "ProfilingDebugging/TraceAuxiliary.h"
 #include "Trace/Config.h"
 #include "Trace/Trace.h"
 
-#if UE_TRACE_ENABLED && WITH_MASSENTITY_DEBUG && !IS_PROGRAM && !UE_BUILD_SHIPPING
-#define UE_MASS_TRACE_ENABLED 1
-#else
-#define UE_MASS_TRACE_ENABLED 0
+#ifndef UE_MASS_TRACE_ENABLED
+#define UE_MASS_TRACE_ENABLED (UE_TRACE_ENABLED && WITH_MASSENTITY_DEBUG && !IS_PROGRAM && !UE_BUILD_SHIPPING)
 #endif
 
 #if UE_MASS_TRACE_ENABLED
