[Mass] CIS Issue 829020: added missing TraceAuxiliary.h header to MassDebugger.cpp

#rnx

diff --git a/Engine/Source/Runtime/MassEntity/Private/MassDebugger.cpp b/Engine/Source/Runtime/MassEntity/Private/MassDebugger.cpp
--- a/Engine/Source/Runtime/MassEntity/Private/MassDebugger.cpp
+++ b/Engine/Source/Runtime/MassEntity/Private/MassDebugger.cpp
@@ -19,6 +19,7 @@
 #include "MassEntityTrace.h"
 #include "Misc/MessageDialog.h"
 #include "HAL/PlatformMisc.h"
+#include "ProfilingDebugging/TraceAuxiliary.h"
 
 #define LOCTEXT_NAMESPACE "MassDebugger"
 
