Mass Entity Trace changes for Mass Insights support
* Fixed bug where phase ends were being recorded as phase begins
* Removed populate entity events - they were sent on session start, however this neglects to consider that event recording occurs before sessions start due to insights ring buffer.  This leads to seemingly out of order events or double "create entity" events

diff --git a/Engine/Source/Runtime/MassEntity/Private/MassDebugger.cpp b/Engine/Source/Runtime/MassEntity/Private/MassDebugger.cpp
--- a/Engine/Source/Runtime/MassEntity/Private/MassDebugger.cpp
+++ b/Engine/Source/Runtime/MassEntity/Private/MassDebugger.cpp
@@ -1254,8 +1254,6 @@ FMassDebugger::FEnvironment::FEnvironment(const FMassEntityManager& InEntityMana
 			UE_TRACE_MASS_ARCHETYPE_CREATED(ArchetypeHandle)
 		});
 
-		FMassTrace::PopulateEntities(Manager);
-
 	});
 #endif
 }
diff --git a/Engine/Source/Runtime/MassEntity/Private/MassEntityTrace.cpp b/Engine/Source/Runtime/MassEntity/Private/MassEntityTrace.cpp
--- a/Engine/Source/Runtime/MassEntity/Private/MassEntityTrace.cpp
+++ b/Engine/Source/Runtime/MassEntity/Private/MassEntityTrace.cpp
@@ -232,7 +232,7 @@ void FMassTrace::OnPhaseEnd(uint64 PhaseId)
 	if (UE_TRACE_CHANNELEXPR_IS_ENABLED(MassChannel))
 	{
 		const FString EnumName = StaticEnum<EMassProcessingPhase>()->GetNameStringByValue(PhaseId);
-		OutputBeginPhaseRegion(*EnumName);
+		OutputEndPhaseRegion(*EnumName);
 	}
 }
 
@@ -304,55 +304,6 @@ void FMassTrace::RegisterFragment(const UScriptStruct* Struct)
 	}
 }
 
-void FMassTrace::PopulateEntities(const FMassEntityManager& EntityManager)
-{
-	if (UE_TRACE_CHANNELEXPR_IS_ENABLED(MassChannel))
-	{
-		constexpr int32 ScratchSpace = 1024;
-
-		TArray<uint64, TInlineAllocator<1024>> EntityDataScratch;
-		TArray<uint64, TInlineAllocator<1024>> ArchetypeIDsScratch;
-
-		const uint64 Cycle = FPlatformTime::Cycles64();
-
-		FMassDebugger::ForEachArchetype(EntityManager, [Cycle, &EntityDataScratch, &ArchetypeIDsScratch](const FMassArchetypeHandle ArchetypeHandle)
-		{
-			const FMassArchetypeData* ArchetypeData = FMassDebugger::GetArchetypeData(ArchetypeHandle);
-
-			if (ArchetypeData)
-			{
-				uint64 ArchetypeID = FMassDebugger::GetArchetypeTraceID(*ArchetypeData);
-				FMassDebugger::EnumerateChunks(*ArchetypeData, [Cycle, ArchetypeID, ArchetypeData, &EntityDataScratch, &ArchetypeIDsScratch](const FMassArchetypeChunk& Chunk)
-					{
-						TConstArrayView<FMassEntityHandle> EntitiesHandles = FMassDebugger::GetEntitiesViewOfArchetype(*ArchetypeData, Chunk);
-
-						for (FMassEntityHandle EntityHandle : EntitiesHandles)
-						{
-							EntityDataScratch.Add(EntityHandle.AsNumber());
-							ArchetypeIDsScratch.Add(ArchetypeID);
-
-							if (EntityDataScratch.Num() == ScratchSpace)
-							{
-								UE_TRACE_LOG(MassTrace, MassBulkAddEntity, MassChannel)
-									<< MassBulkAddEntity.Cycle(Cycle)
-									<< MassBulkAddEntity.Entities(EntityDataScratch.GetData(), EntityDataScratch.Num())
-									<< MassBulkAddEntity.ArchetypeIDs(ArchetypeIDsScratch.GetData(), ArchetypeIDsScratch.Num());
-								EntityDataScratch.Reset();
-								ArchetypeIDsScratch.Reset();
-							}
-						}
-
-					});
-			}
-		});
-
-		UE_TRACE_LOG(MassTrace, MassBulkAddEntity, MassChannel)
-			<< MassBulkAddEntity.Cycle(Cycle)
-			<< MassBulkAddEntity.Entities(EntityDataScratch.GetData(), EntityDataScratch.Num())
-			<< MassBulkAddEntity.ArchetypeIDs(ArchetypeIDsScratch.GetData(), ArchetypeIDsScratch.Num());
-	}
-}
-
 void FMassTrace::EntityCreated(FMassEntityHandle Entity, const FMassArchetypeData& Archetype)
 {
 	if (UE_TRACE_CHANNELEXPR_IS_ENABLED(MassChannel))
diff --git a/Engine/Source/Runtime/MassEntity/Public/MassEntityTrace.h b/Engine/Source/Runtime/MassEntity/Public/MassEntityTrace.h
--- a/Engine/Source/Runtime/MassEntity/Public/MassEntityTrace.h
+++ b/Engine/Source/Runtime/MassEntity/Public/MassEntityTrace.h
@@ -52,8 +52,7 @@ public:
 
 	MASSENTITY_API static void OutputRegisterFragment(const UScriptStruct* Struct);
 	MASSENTITY_API static void RegisterFragment(const UScriptStruct* Struct);
-
-	MASSENTITY_API static void PopulateEntities(const FMassEntityManager& EntityManager);
+	
 	MASSENTITY_API static void EntityCreated(FMassEntityHandle Entity, const FMassArchetypeData& Archetype);
 	MASSENTITY_API static void EntityMoved(FMassEntityHandle Entity, const FMassArchetypeData& NewArchetype);
 	MASSENTITY_API static void EntityDestroyed(FMassEntityHandle Entity);
