[Mass] Added CalculateRangeLength to MassArchetypeData to unify how we calculate range length for given FMassArchetypeEntityCollection::FArchetypeEntityRange, and also so that the same calculation can be used for debugging purposes

#rnx

diff --git a/Engine/Source/Runtime/MassEntity/Internal/MassArchetypeData.h b/Engine/Source/Runtime/MassEntity/Internal/MassArchetypeData.h
--- a/Engine/Source/Runtime/MassEntity/Internal/MassArchetypeData.h
+++ b/Engine/Source/Runtime/MassEntity/Internal/MassArchetypeData.h
@@ -287,6 +287,20 @@ public:
 	int32 GetChunkCount() const { return Chunks.Num(); }
 	int32 GetNonEmptyChunkCount() const;
 
+	FORCEINLINE static int32 CalculateRangeLength(FMassArchetypeEntityCollection::FArchetypeEntityRange EntityRange, const FMassArchetypeChunk& Chunk)
+	{
+		return EntityRange.Length > 0
+			? EntityRange.Length
+			: (Chunk.GetNumInstances() - EntityRange.SubchunkStart);	
+	}
+
+	int32 CalculateRangeLength(FMassArchetypeEntityCollection::FArchetypeEntityRange EntityRange) const
+	{
+		check(Chunks.IsValidIndex(EntityRange.ChunkIndex));
+		const FMassArchetypeChunk& Chunk = Chunks[EntityRange.ChunkIndex];
+		return CalculateRangeLength(EntityRange, Chunk);
+	}
+
 	uint32 GetCreatedArchetypeDataVersion() const;
 	uint32 GetEntityOrderVersion() const;
 
diff --git a/Engine/Source/Runtime/MassEntity/Private/MassArchetypeData.cpp b/Engine/Source/Runtime/MassEntity/Private/MassArchetypeData.cpp
--- a/Engine/Source/Runtime/MassEntity/Private/MassArchetypeData.cpp
+++ b/Engine/Source/Runtime/MassEntity/Private/MassArchetypeData.cpp
@@ -337,7 +337,7 @@ void FMassArchetypeData::BatchDestroyEntityChunks(FMassArchetypeEntityCollection
 	for (const FMassArchetypeEntityCollection::FArchetypeEntityRange EntityRange : SortedRangeCollection)
 	{ 
 		FMassArchetypeChunk& Chunk = Chunks[EntityRange.ChunkIndex];
-		const int32 RangeLength = EntityRange.Length > 0 ? EntityRange.Length : (Chunk.GetNumInstances() - EntityRange.SubchunkStart);
+		const int32 RangeLength = CalculateRangeLength(EntityRange, Chunk);
 
 		// gather entities we're about to remove
 		FMassEntityHandle* DyingEntityPtr = &Chunk.GetEntityArrayElementRef(EntityListOffsetWithinChunk, EntityRange.SubchunkStart);
@@ -564,7 +564,7 @@ void FMassArchetypeData::ExecuteFunction(FMassExecutionContext& RunContext, cons
 void FMassArchetypeData::ExecutionFunctionForChunk(FMassExecutionContext& RunContext, const FMassExecuteFunction& Function, const FMassQueryRequirementIndicesMapping& RequirementMapping, const FMassArchetypeEntityCollection::FArchetypeEntityRange& EntityRange, const FMassChunkConditionFunction& ChunkCondition)
 {
 	FMassArchetypeChunk& Chunk = Chunks[EntityRange.ChunkIndex];
-	const int32 RangeLength = EntityRange.Length > 0 ? EntityRange.Length : (Chunk.GetNumInstances() - EntityRange.SubchunkStart);
+	const int32 RangeLength = CalculateRangeLength(EntityRange, Chunk);
 
 	if (RangeLength > 0)
 	{
@@ -923,9 +923,9 @@ void FMassArchetypeData::ExportEntityHandles(const TConstArrayView<FMassArchetyp
 
 	for (const FMassArchetypeEntityCollection::FArchetypeEntityRange& Range : Ranges)
 	{
-		const FMassEntityHandle* EntitiesArray = Chunks[Range.ChunkIndex].GetEntityArray(EntityListOffsetWithinChunk);
-
-		const int32 RangeLength = Range.Length > 0 ? Range.Length : (Chunks[Range.ChunkIndex].GetNumInstances() - Range.SubchunkStart);
+		const FMassArchetypeChunk& Chunk = Chunks[Range.ChunkIndex];
+		const FMassEntityHandle* EntitiesArray = Chunk.GetEntityArray(EntityListOffsetWithinChunk);
+		const int32 RangeLength = CalculateRangeLength(Range, Chunk);
 		FMemory::Memcpy(&InOutHandles[StartIndex], &EntitiesArray[Range.SubchunkStart], RangeLength * sizeof(FMassEntityHandle));
 
 		StartIndex += RangeLength;
@@ -1179,7 +1179,7 @@ void FMassArchetypeData::BatchMoveEntitiesToAnotherArchetype(const FMassArchetyp
 		}
 
 		FMassArchetypeChunk& Chunk = Chunks[EntityRange.ChunkIndex];
-		const int32 RangeLength = EntityRange.Length > 0 ? EntityRange.Length : (Chunk.GetNumInstances() - EntityRange.SubchunkStart);
+		const int32 RangeLength = CalculateRangeLength(EntityRange, Chunk);
 
 		// 0 - consider compacting new archetype to ensure larger empty spaces
 		// 1. find next free spot in the destination archetype
@@ -1430,7 +1430,7 @@ void FMassArchetypeData::BatchSetFragmentValues(TConstArrayView<FMassArchetypeEn
 	for (const FMassArchetypeEntityCollection::FArchetypeEntityRange EntityRange : EntityCollection)
 	{
 		FMassArchetypeChunk& Chunk = Chunks[EntityRange.ChunkIndex];
-		const int32 RangeLength = EntityRange.Length > 0 ? EntityRange.Length : (Chunk.GetNumInstances() - EntityRange.SubchunkStart);
+		const int32 RangeLength = CalculateRangeLength(EntityRange, Chunk);
 
 		for (int i = 0; i < Payload.Num(); ++i)
 		{
