Added a constructor that can take ownership of a FMassArchetypeEntityCollection to the FMassArchetypeEntityCollectionWithPayload struct.

diff --git a/Engine/Source/Runtime/MassEntity/Public/MassArchetypeTypes.h b/Engine/Source/Runtime/MassEntity/Public/MassArchetypeTypes.h
--- a/Engine/Source/Runtime/MassEntity/Public/MassArchetypeTypes.h
+++ b/Engine/Source/Runtime/MassEntity/Public/MassArchetypeTypes.h
@@ -190,6 +190,11 @@ struct MASSENTITY_API FMassArchetypeEntityCollectionWithPayload
 	{
 	}
 
+	explicit FMassArchetypeEntityCollectionWithPayload(FMassArchetypeEntityCollection&& InEntityCollection)
+		: Entities(MoveTempIfPossible(InEntityCollection))
+	{
+	}
+
 	static void CreateEntityRangesWithPayload(const FMassEntityManager& EntityManager, const TConstArrayView<FMassEntityHandle> Entities
 		, const FMassArchetypeEntityCollection::EDuplicatesHandling DuplicatesHandling, FMassGenericPayloadView Payload
 		, TArray<FMassArchetypeEntityCollectionWithPayload>& OutEntityCollections);
