[Mass] Extension to FMassArchetypeEntityCollection's functionality. It's now possible to compare two collection's archetype handles and append one compatible collection to another.

diff --git a/Engine/Source/Runtime/MassEntity/Private/MassArchetypeTypes.cpp b/Engine/Source/Runtime/MassEntity/Private/MassArchetypeTypes.cpp
--- a/Engine/Source/Runtime/MassEntity/Private/MassArchetypeTypes.cpp
+++ b/Engine/Source/Runtime/MassEntity/Private/MassArchetypeTypes.cpp
@@ -5,6 +5,7 @@
 #include "MassArchetypeData.h"
 #include "MassEntityUtils.h"
 #include "Containers/StridedView.h"
+#include "MassTestableEnsures.h"
 
 //-----------------------------------------------------------------------------
 // FMassArchetypeHandle
@@ -24,7 +25,7 @@ FMassArchetypeVersionedHandle::FMassArchetypeVersionedHandle(const FMassArchetyp
 }
 
 FMassArchetypeVersionedHandle::FMassArchetypeVersionedHandle(FMassArchetypeHandle&& InHandle)
-	: ArchetypeHandle(MoveTemp(InHandle))
+	: ArchetypeHandle(Forward<FMassArchetypeHandle>(InHandle))
 	, HandleVersion(ArchetypeHandle.IsValid() ? ArchetypeHandle.DataPtr->GetEntityOrderVersion() : 0)
 {
 }
@@ -207,15 +208,62 @@ bool FMassArchetypeEntityCollection::IsSame(const FMassArchetypeEntityCollection
 	return true;
 }
 
-void FMassArchetypeEntityCollection::ExportEntityHandles(TArray<FMassEntityHandle>& InOutHandles) const
+bool FMassArchetypeEntityCollection::ExportEntityHandles(TArray<FMassEntityHandle>& InOutHandles) const
 {
+	if (Ranges.IsEmpty())
+	{
+		return false;
+	}
+
+	const int32 InOutHandlesStartingNum = InOutHandles.Num();
 	const FMassArchetypeData* ArchetypeData = FMassArchetypeHelper::ArchetypeDataFromHandle(Archetype);
-	if (ensureMsgf(ArchetypeData, TEXT("This operation is supported only for collections containing created entities (i.e. not 'reserved'")))
+	if (testableEnsureMsgf(ArchetypeData, TEXT("This operation is supported only for collections containing created entities (i.e. not 'reserved'")))
 	{
-		checkf(IsUpToDate(), TEXT("The entity collection is out of that. The following action might crash and if not, the results might be incorrect."));
+		testableCheckfReturn(IsUpToDate(), return false, TEXT("The entity collection is out of date. Unable to reliably generate handle list."));
+
 		ArchetypeData->ExportEntityHandles(Ranges, InOutHandles);
+		ensureMsgf(InOutHandlesStartingNum != InOutHandles.Num(), TEXT("Exporting non-empty ranges resulted in no entity handles"));
+	}
+
+	return InOutHandlesStartingNum != InOutHandles.Num();
+}
+
+bool FMassArchetypeEntityCollection::DoesContainOverlappingRanges(FConstEntityRangeArrayView Ranges)
+{
+	for (int RangeIndex = 1; RangeIndex < Ranges.Num(); ++RangeIndex) 
+	{
+		if (Ranges[RangeIndex-1].IsOverlapping(Ranges[RangeIndex]))
+		{
+			return true;
+		}
 	}
+	return false;
+}
+
+#if WITH_MASSENTITY_DEBUG
+int32 FMassArchetypeEntityCollection::DebugCountEntities() const
+{
+	int32 EntitiesCount = 0;
+
+	if (const FMassArchetypeData* ArchetypeData = FMassArchetypeHelper::ArchetypeDataFromHandle(Archetype))
+	{
+		for (const FArchetypeEntityRange& Range : Ranges)
+		{
+			EntitiesCount += ArchetypeData->CalculateRangeLength(Range);
+		}
+	}
+	else
+	{
+		for (const FArchetypeEntityRange& Range : Ranges)
+		{
+			ensureAlwaysMsgf(Range.Length > 0, TEXT("Empty range length in archetype-less collection is unexpected. There's no way to determine what it means"));
+			EntitiesCount += Range.Length;
+		}
+	}
+
+	return EntitiesCount;
 }
+#endif // WITH_MASSENTITY_DEBUG
 
 //-----------------------------------------------------------------------------
 // FMassArchetypeEntityCollectionWithPayload
diff --git a/Engine/Source/Runtime/MassEntity/Public/MassArchetypeTypes.h b/Engine/Source/Runtime/MassEntity/Public/MassArchetypeTypes.h
--- a/Engine/Source/Runtime/MassEntity/Public/MassArchetypeTypes.h
+++ b/Engine/Source/Runtime/MassEntity/Public/MassArchetypeTypes.h
@@ -113,11 +113,30 @@ public:
 			return ChunkIndex == Other.ChunkIndex && SubchunkStart + Length == Other.SubchunkStart;
 		}
 
+		bool IsOverlapping(const FArchetypeEntityRange& Other) const
+		{
+			return ChunkIndex == Other.ChunkIndex
+				&& (*this < Other
+					// note that Length == 0 means "all the entities starting from SubchunkStart
+					? (SubchunkStart + Length > Other.SubchunkStart || Length == 0)
+					: (Other.SubchunkStart + Other.Length > SubchunkStart || Other.Length == 0)
+				);
+		}
+
 		bool operator==(const FArchetypeEntityRange& Other) const
 		{
 			return ChunkIndex == Other.ChunkIndex && SubchunkStart == Other.SubchunkStart && Length == Other.Length;
 		}
 		bool operator!=(const FArchetypeEntityRange& Other) const { return !(*this == Other); }
+
+		bool operator<(const FArchetypeEntityRange& Other) const
+		{
+			return (ChunkIndex != Other.ChunkIndex)
+				? ChunkIndex < Other.ChunkIndex
+				: (SubchunkStart != Other.SubchunkStart
+					? SubchunkStart < Other.SubchunkStart
+					: Length < Other.Length);
+		}
 	};
 
 	enum EDuplicatesHandling
@@ -148,7 +167,7 @@ public:
 	explicit FMassArchetypeEntityCollection(const FMassArchetypeHandle& InArchetypeHandle, const EInitializationType Initialization = EInitializationType::GatherAll);
 	explicit FMassArchetypeEntityCollection(TSharedPtr<FMassArchetypeData>& InArchetype, const EInitializationType Initialization = EInitializationType::GatherAll);
 	FMassArchetypeEntityCollection(const FMassArchetypeHandle& InArchetypeHandle, FEntityRangeArray&& InEntityRanges)
-		: Ranges(MoveTemp(InEntityRanges))
+		: Ranges(Forward<FEntityRangeArray>(InEntityRanges))
 		, Archetype(InArchetypeHandle)
 	{}
 
@@ -168,12 +187,31 @@ public:
 
 	/** The comparison function that checks if Other is identical to this. Intended for diagnostics/debugging. */
 	bool IsSame(const FMassArchetypeEntityCollection& Other) const;
+	bool IsSameArchetype(const FMassArchetypeEntityCollection& Other) const;
+
+	/**
+	 * Appends ranges of the given FMassArchetypeEntityCollection instance. Note that it can be safely done only
+	 * when both collections host entities of the same archetype, and both were created with the same version
+	 * of said archetype.
+	 * Additionally, we don't expect the operation to produce overlapping entity ranges and this assumption is
+	 * only verified in debug builds (i.e. use it only when you're certain no range overlaps are possible).
+	 */
+	template<typename T>
+	requires std::is_same_v<typename TDecay<T>::Type, FMassArchetypeEntityCollection>
+	void Append(T&& Other);
 
 	/**
 	 * Converts stored entity ranges to FMassEntityHandles and appends them to InOutHandles.
 	 * Note that the operation is only supported for already created entities (i.e. not "reserved")
+	 * @return whether any entity handles have been actually exported
 	 */
-	void ExportEntityHandles(TArray<FMassEntityHandle>& InOutHandles) const;
+	bool ExportEntityHandles(TArray<FMassEntityHandle>& InOutHandles) const;
+
+	static bool DoesContainOverlappingRanges(FConstEntityRangeArrayView Ranges);
+
+#if WITH_MASSENTITY_DEBUG
+	int32 DebugCountEntities() const;
+#endif // WITH_MASSENTITY_DEBUG
 
 protected:
 	friend struct FMassArchetypeEntityCollectionWithPayload;
@@ -334,6 +372,28 @@ inline bool FMassArchetypeEntityCollection::IsSet() const
 	return !IsEmpty();
 }
 
+inline bool FMassArchetypeEntityCollection::IsSameArchetype(const FMassArchetypeEntityCollection& Other) const
+{
+	return Archetype == Other.Archetype;
+}
+
+template<typename T>
+requires std::is_same_v<typename TDecay<T>::Type, FMassArchetypeEntityCollection>
+void FMassArchetypeEntityCollection::Append(T&& Other)
+{
+	const bool bWasEmpty = Ranges.IsEmpty();
+	checkf(IsSameArchetype(Other), TEXT("Unable to merge two entity collections representing different archetypes"));
+
+	Ranges.Append(Forward<T>(Other).Ranges);
+
+	if (bWasEmpty == false)
+	{
+		Ranges.Sort();
+		checkfSlow(DoesContainOverlappingRanges(Ranges) == false
+			, TEXT("Entity collection ranges overlap as a result of %hs"), __FUNCTION__);
+	}
+}
+
 inline FMassArchetypeVersionedHandle::operator FMassArchetypeHandle() const
 {
 	return ArchetypeHandle;
