[Mass] CIS Issue 835311: Fixed FMassArchetypeVersionedHandle::IsUpToDate not being API-exported

#rnx

diff --git a/Engine/Source/Runtime/MassEntity/Public/MassArchetypeTypes.h b/Engine/Source/Runtime/MassEntity/Public/MassArchetypeTypes.h
--- a/Engine/Source/Runtime/MassEntity/Public/MassArchetypeTypes.h
+++ b/Engine/Source/Runtime/MassEntity/Public/MassArchetypeTypes.h
@@ -62,7 +62,7 @@ struct FMassArchetypeVersionedHandle final
 	bool IsValid() const;
 	bool operator==(const FMassArchetypeVersionedHandle& Other) const;
 	bool operator!=(const FMassArchetypeVersionedHandle& Other) const;
-	bool IsUpToDate() const;
+	MASSENTITY_API bool IsUpToDate() const;
 
 	operator FMassArchetypeHandle() const;
 
@@ -347,6 +347,11 @@ inline bool FMassArchetypeVersionedHandle::operator!=(const FMassArchetypeVersio
 	return !(*this == Other);
 }
 
+inline FMassArchetypeVersionedHandle::operator FMassArchetypeHandle() const
+{
+	return ArchetypeHandle;
+}
+
 inline FMassArchetypeEntityCollection::FConstEntityRangeArrayView FMassArchetypeEntityCollection::GetRanges() const
 {
 	return Ranges;
@@ -394,11 +399,6 @@ void FMassArchetypeEntityCollection::Append(T&& Other)
 	}
 }
 
-inline FMassArchetypeVersionedHandle::operator FMassArchetypeHandle() const
-{
-	return ArchetypeHandle;
-}
-
 inline FMassRawEntityInChunkData::FMassRawEntityInChunkData(uint8* InChunkRawMemory, const int32 InIndexWithinChunk)
 	: ChunkRawMemory(InChunkRawMemory), IndexWithinChunk(InIndexWithinChunk)
 {}
