[Mass] renamed protected FMassObserverManager::HandleSingleEntityImpl to HandleSingleElementImpl

#rnx

diff --git a/Engine/Source/Runtime/MassEntity/Private/MassObserverManager.cpp b/Engine/Source/Runtime/MassEntity/Private/MassObserverManager.cpp
--- a/Engine/Source/Runtime/MassEntity/Private/MassObserverManager.cpp
+++ b/Engine/Source/Runtime/MassEntity/Private/MassObserverManager.cpp
@@ -294,12 +294,12 @@ void FMassObserverManager::OnElementOperation(TNotNull<const UScriptStruct*> Ele
 	{
 		if (ObservedFragments[(uint8)Operation].Contains(*ElementType))
 		{
-			HandleSingleEntityImpl(*ElementType, EntityCollection, FragmentObservers[(uint8)Operation]);
+			HandleSingleElementImpl(*ElementType, EntityCollection, FragmentObservers[(uint8)Operation]);
 		}
 	}
 	else if (ObservedTags[(uint8)Operation].Contains(*ElementType))
 	{
-		HandleSingleEntityImpl(*ElementType, EntityCollection, TagObservers[(uint8)Operation]);
+		HandleSingleElementImpl(*ElementType, EntityCollection, TagObservers[(uint8)Operation]);
 	}
 }
 
@@ -320,7 +320,7 @@ void FMassObserverManager::HandleFragmentsImpl(FMassProcessingContext& Processin
 	}
 }
 
-void FMassObserverManager::HandleSingleEntityImpl(const UScriptStruct& FragmentType, const FMassArchetypeEntityCollection& EntityCollection, FMassObserversMap& HandlersContainer)
+void FMassObserverManager::HandleSingleElementImpl(const UScriptStruct& FragmentType, const FMassArchetypeEntityCollection& EntityCollection, FMassObserversMap& HandlersContainer)
 {
 	TRACE_CPUPROFILER_EVENT_SCOPE(MassObserver_HandleSingleEntityImpl);
 
@@ -396,8 +396,13 @@ void FMassObserverManager::RemoveObserverInstance(const UScriptStruct& ElementTy
 }
 
 //----------------------------------------------------------------------//
-// FMassObserversMap
+// DEPRECATED
 //----------------------------------------------------------------------//
+void FMassObserverManager::HandleSingleEntityImpl(const UScriptStruct& FragmentType, const FMassArchetypeEntityCollection& EntityCollection, FMassObserversMap& HandlersContainer)
+{
+	HandleSingleElementImpl(FragmentType, EntityCollection, HandlersContainer);
+}
+
 void FMassObserverManager::OnPostFragmentOrTagAdded(const UScriptStruct& FragmentOrTagType, const FMassArchetypeEntityCollection& EntityCollection)
 {
 	OnPostElementAdded(&FragmentOrTagType, EntityCollection);
diff --git a/Engine/Source/Runtime/MassEntity/Public/MassObserverManager.h b/Engine/Source/Runtime/MassEntity/Public/MassObserverManager.h
--- a/Engine/Source/Runtime/MassEntity/Public/MassObserverManager.h
+++ b/Engine/Source/Runtime/MassEntity/Public/MassObserverManager.h
@@ -139,7 +139,7 @@ protected:
 	void DeInitialize();
 	void HandleFragmentsImpl(FMassProcessingContext& ProcessingContext, const FMassArchetypeEntityCollection& EntityCollection
 		, TArrayView<const UScriptStruct*> ObservedTypes, FMassObserversMap& HandlersContainer);
-	void HandleSingleEntityImpl(const UScriptStruct& FragmentType, const FMassArchetypeEntityCollection& EntityCollection, FMassObserversMap& HandlersContainer);
+	void HandleSingleElementImpl(const UScriptStruct& FragmentType, const FMassArchetypeEntityCollection& EntityCollection, FMassObserversMap& HandlersContainer);
 
 	FMassFragmentBitSet ObservedFragments[(uint8)EMassObservedOperation::MAX];
 	FMassTagBitSet ObservedTags[(uint8)EMassObservedOperation::MAX];
@@ -165,6 +165,9 @@ public:
 	void OnPreFragmentOrTagRemoved(const UScriptStruct& FragmentOrTagType, const FMassArchetypeEntityCollection& EntityCollection);
 	UE_DEPRECATED(5.5, "*FragmentOrTag* functions have been renamed to *Element*. Please use those instead.")
 	void OnFragmentOrTagOperation(const UScriptStruct& FragmentOrTagType, const FMassArchetypeEntityCollection& EntityCollection, const EMassObservedOperation Operation);
+protected:
+	UE_DEPRECATED(5.6, "Use HandleSingleElementImpl.")
+	void HandleSingleEntityImpl(const UScriptStruct& FragmentType, const FMassArchetypeEntityCollection& EntityCollection, FMassObserversMap& HandlersContainer);
 };
 
 template<>
