[Mass] Fixed FMassExecutionContext::GetDummyInstance, it had a flaw in FMassExecutionContext instance creation that in turn resulted in failing checks and effectively crashing.
Also, removed the defunct FMassExecutionContext default constructor

diff --git a/Engine/Source/Runtime/MassEntity/Private/MassExecutionContext.cpp b/Engine/Source/Runtime/MassEntity/Private/MassExecutionContext.cpp
--- a/Engine/Source/Runtime/MassEntity/Private/MassExecutionContext.cpp
+++ b/Engine/Source/Runtime/MassEntity/Private/MassExecutionContext.cpp
@@ -186,7 +186,7 @@ FMassExecutionContext::FEntityIterator FMassExecutionContext::CreateEntityIterat
 
 FMassExecutionContext& FMassExecutionContext::GetDummyInstance()
 {
-	static FMassExecutionContext DummyContext;
+	static FMassExecutionContext DummyContext(*TSharedRef<FMassEntityManager>(MakeShareable(new FMassEntityManager())));
 	return DummyContext;
 }
 
diff --git a/Engine/Source/Runtime/MassEntity/Public/MassExecutionContext.h b/Engine/Source/Runtime/MassEntity/Public/MassExecutionContext.h
--- a/Engine/Source/Runtime/MassEntity/Public/MassExecutionContext.h
+++ b/Engine/Source/Runtime/MassEntity/Public/MassExecutionContext.h
@@ -137,8 +137,6 @@ private:
 	 */
 	FMassExecutionContext& operator=(const FMassExecutionContext& Other) = default;
 
-	FMassExecutionContext() = default;
-
 public:
 	MASSENTITY_API explicit FMassExecutionContext(FMassEntityManager& InEntityManager, const float InDeltaTimeSeconds = 0.f, const bool bInFlushDeferredCommands = true);
 	MASSENTITY_API FMassExecutionContext(FMassExecutionContext&& Other) = default;
@@ -147,7 +145,7 @@ public:
 	MASSENTITY_API ~FMassExecutionContext();
 
 	/** For internal use only, should never be exported as part of API */
-	static FMassExecutionContext& GetDummyInstance();
+	MASSENTITY_API static FMassExecutionContext& GetDummyInstance();
 
 	FMassEntityManager& GetEntityManagerChecked() const;
 	const TSharedRef<FMassEntityManager>& GetSharedEntityManager();
