[Mass] The change makes MassProcessingContext capable of creating properly configured FMassExecutionContext instances, as well as hosting one and cleaning up when going out of scope.
Simplified logic for commands flushing at the end of processing

diff --git a/Engine/Config/BaseEngine.ini b/Engine/Config/BaseEngine.ini
--- a/Engine/Config/BaseEngine.ini
+++ b/Engine/Config/BaseEngine.ini
@@ -1355,6 +1355,7 @@ ClassName=/Script/OnlineSubsystemUtils.OnlineEngineInterfaceImpl
 
 +EnumRedirects=(OldName="EMaterialExpressionConvertType", ValueChanges=(("Float1", "Scalar"), ("Float2", "Vector2"), ("Float3", "Vector3"), ("Float4", "Vector4")))
 
++StructRedirects=(OldName="/Script/MassEntity.MassProcessingContext",NewName="/Script/MassEntity.MassProcessingContext_DEPRECATED")
 ; End of releases redirects
 
 ; Fixing the output of tags for blueprints to be reset.
diff --git a/Engine/Plugins/Runtime/MassGameplay/Source/MassSpawner/Private/MassSpawner.cpp b/Engine/Plugins/Runtime/MassGameplay/Source/MassSpawner/Private/MassSpawner.cpp
--- a/Engine/Plugins/Runtime/MassGameplay/Source/MassSpawner/Private/MassSpawner.cpp
+++ b/Engine/Plugins/Runtime/MassGameplay/Source/MassSpawner/Private/MassSpawner.cpp
@@ -9,7 +9,6 @@
 #include "MassSpawnerTypes.h"
 #include "MassSpawnerSubsystem.h"
 #include "MassSimulationSubsystem.h"
-#include "EnvironmentQuery/EnvQuery.h"
 #include "VisualLogger/VisualLogger.h"
 #include "MassEntityConfigAsset.h"
 #include "MassEntityManager.h"
@@ -19,6 +18,7 @@
 #include "MassSpawnLocationProcessor.h"
 #include "MassExecutor.h"
 #include "MassEntityUtils.h"
+#include "MassProcessingContext.h"
 #if WITH_EDITOR
 #include "Engine/Texture2D.h"
 #endif
diff --git a/Engine/Plugins/Runtime/MassGameplay/Source/MassSpawner/Private/MassSpawnerSubsystem.cpp b/Engine/Plugins/Runtime/MassGameplay/Source/MassSpawner/Private/MassSpawnerSubsystem.cpp
--- a/Engine/Plugins/Runtime/MassGameplay/Source/MassSpawner/Private/MassSpawnerSubsystem.cpp
+++ b/Engine/Plugins/Runtime/MassGameplay/Source/MassSpawner/Private/MassSpawnerSubsystem.cpp
@@ -10,11 +10,11 @@
 #include "StructUtils/InstancedStruct.h"
 #include "VisualLogger/VisualLogger.h"
 #include "MassSpawner.h"
-#include "MassObserverProcessor.h"
 #include "ProfilingDebugging/CpuProfilerTrace.h"
 #include "MassSimulationSubsystem.h"
 #include "MassProcessor.h"
 #include "MassEntityUtils.h"
+#include "MassProcessingContext.h"
 
 //----------------------------------------------------------------------//
 //  UMassSpawnerSubsystem
diff --git a/Engine/Source/Developer/MassEntityTestSuite/Private/MassChunkCollectionTest.cpp b/Engine/Source/Developer/MassEntityTestSuite/Private/MassChunkCollectionTest.cpp
--- a/Engine/Source/Developer/MassEntityTestSuite/Private/MassChunkCollectionTest.cpp
+++ b/Engine/Source/Developer/MassEntityTestSuite/Private/MassChunkCollectionTest.cpp
@@ -6,6 +6,7 @@
 #include "MassExecutionContext.h"
 #include "Algo/RandomShuffle.h"
 #include "MassExecutor.h"
+#include "MassProcessingContext.h"
 #include "Algo/Compare.h"
 #include "Algo/RemoveIf.h"
 
diff --git a/Engine/Source/Developer/MassEntityTestSuite/Private/MassCompositeProcessorTest.cpp b/Engine/Source/Developer/MassEntityTestSuite/Private/MassCompositeProcessorTest.cpp
--- a/Engine/Source/Developer/MassEntityTestSuite/Private/MassCompositeProcessorTest.cpp
+++ b/Engine/Source/Developer/MassEntityTestSuite/Private/MassCompositeProcessorTest.cpp
@@ -3,7 +3,7 @@
 #include "AITestsCommon.h"
 #include "MassEntityManager.h"
 #include "MassExecutionContext.h"
-#include "MassProcessingTypes.h"
+#include "MassProcessingContext.h"
 #include "MassEntityTestTypes.h"
 #include "MassExecutor.h"
 
diff --git a/Engine/Source/Developer/MassEntityTestSuite/Private/MassEntityBuilderTest.cpp b/Engine/Source/Developer/MassEntityTestSuite/Private/MassEntityBuilderTest.cpp
--- a/Engine/Source/Developer/MassEntityTestSuite/Private/MassEntityBuilderTest.cpp
+++ b/Engine/Source/Developer/MassEntityTestSuite/Private/MassEntityBuilderTest.cpp
@@ -6,6 +6,7 @@
 #include "MassEntityManager.h"
 #include "MassExecutionContext.h"
 #include "MassExecutor.h"
+#include "MassProcessingContext.h"
 
 
 #define LOCTEXT_NAMESPACE "MassTest"
diff --git a/Engine/Source/Developer/MassEntityTestSuite/Private/MassEntityTestFarmPlot.cpp b/Engine/Source/Developer/MassEntityTestSuite/Private/MassEntityTestFarmPlot.cpp
--- a/Engine/Source/Developer/MassEntityTestSuite/Private/MassEntityTestFarmPlot.cpp
+++ b/Engine/Source/Developer/MassEntityTestSuite/Private/MassEntityTestFarmPlot.cpp
@@ -4,9 +4,9 @@
 
 #include "Components/HierarchicalInstancedStaticMeshComponent.h"
 #include "Engine/CollisionProfile.h"
-#include "MassEntityUtils.h"
 #include "MassExecutor.h"
 #include "MassExecutionContext.h"
+#include "MassProcessingContext.h"
 
 #include UE_INLINE_GENERATED_CPP_BY_NAME(MassEntityTestFarmPlot)
 
diff --git a/Engine/Source/Developer/MassEntityTestSuite/Private/MassExecutionTest.cpp b/Engine/Source/Developer/MassEntityTestSuite/Private/MassExecutionTest.cpp
--- a/Engine/Source/Developer/MassEntityTestSuite/Private/MassExecutionTest.cpp
+++ b/Engine/Source/Developer/MassEntityTestSuite/Private/MassExecutionTest.cpp
@@ -4,6 +4,7 @@
 #include "MassProcessingTypes.h"
 #include "MassEntityTestTypes.h"
 #include "MassExecutor.h"
+#include "MassProcessingContext.h"
 
 #define LOCTEXT_NAMESPACE "MassTest"
 
@@ -57,26 +58,6 @@ struct FExecution_EmptyPipeline : FExecutionTestBase
 IMPLEMENT_AI_INSTANT_TEST(FExecution_EmptyPipeline, "System.Mass.Execution.EmptyPipeline");
 
 
-struct FExecution_InvalidProcessingContext : FExecutionTestBase
-{
-	virtual bool InstantTest() override
-	{
-		CA_ASSUME(EntityManager);
-		const float DeltaSeconds = 0.f;
-		FMassProcessingContext ProcessingContext;
-		// test assumption
-		AITEST_NULL("FMassProcessingContext\'s default constructor is expected to set FMassProcessingContext.EntitySubsystem to null", ProcessingContext.EntityManager.Get());
-		
-		AITEST_SCOPED_CHECK(TEXT("ProcessingContext.EntityManager is null"), 1);
-		// note that using RunProcessorsView is to bypass reasonable tests UE::Mass::Executor::Run(Pipeline,...) does that are 
-		// reported via ensures which are not handled by the automation framework
-		UE::Mass::Executor::RunProcessorsView(TArrayView<UMassProcessor*>(), ProcessingContext);
-		return true;
-	}
-};
-IMPLEMENT_AI_INSTANT_TEST(FExecution_InvalidProcessingContext, "System.Mass.Execution.InvalidProcessingContext");
-
-
 #if WITH_MASSENTITY_DEBUG
 struct FExecution_SingleNullProcessor : FExecutionTestBase
 {
diff --git a/Engine/Source/Developer/MassEntityTestSuite/Private/MassMultiThreadTest.cpp b/Engine/Source/Developer/MassEntityTestSuite/Private/MassMultiThreadTest.cpp
--- a/Engine/Source/Developer/MassEntityTestSuite/Private/MassMultiThreadTest.cpp
+++ b/Engine/Source/Developer/MassEntityTestSuite/Private/MassMultiThreadTest.cpp
@@ -7,6 +7,7 @@
 #include "MassExecutor.h"
 #include "MassEntityView.h"
 #include "MassExecutionContext.h"
+#include "MassProcessingContext.h"
 
 #define LOCTEXT_NAMESPACE "MassTest"
 
@@ -42,8 +43,8 @@ struct FMTTestBase : FEntityTestBase
 			CompositeProcessor->SetGroupName(TEXT("Test"));
 			CompositeProcessor->SetProcessors(MakeArrayView<UMassProcessor*>((UMassProcessor**)Processors.GetData(), Processors.Num()));
 
-			FMassProcessingContext Context(*EntityManager, /*DeltaTime=*/0);
-			FinishEvent = UE::Mass::Executor::TriggerParallelTasks(*CompositeProcessor, Context, []() {});
+			FMassProcessingContext Context(*EntityManager);
+			FinishEvent = UE::Mass::Executor::TriggerParallelTasks(*CompositeProcessor, MoveTemp(Context), []() {});
 		}
 
 		if (FinishEvent->IsComplete())
diff --git a/Engine/Source/Developer/MassEntityTestSuite/Private/MassProcessorTest.cpp b/Engine/Source/Developer/MassEntityTestSuite/Private/MassProcessorTest.cpp
--- a/Engine/Source/Developer/MassEntityTestSuite/Private/MassProcessorTest.cpp
+++ b/Engine/Source/Developer/MassEntityTestSuite/Private/MassProcessorTest.cpp
@@ -3,7 +3,7 @@
 #include "AITestsCommon.h"
 #include "MassEntityManager.h"
 #include "MassExecutionContext.h"
-#include "MassProcessingTypes.h"
+#include "MassProcessingContext.h"
 #include "MassEntityTestTypes.h"
 #include "MassExecutor.h"
 
diff --git a/Engine/Source/Developer/MassEntityTestSuite/Private/MassQueryTest.cpp b/Engine/Source/Developer/MassEntityTestSuite/Private/MassQueryTest.cpp
--- a/Engine/Source/Developer/MassEntityTestSuite/Private/MassQueryTest.cpp
+++ b/Engine/Source/Developer/MassEntityTestSuite/Private/MassQueryTest.cpp
@@ -3,7 +3,7 @@
 #include "AITestsCommon.h"
 #include "Engine/World.h"
 #include "MassEntityManager.h"
-#include "MassProcessingTypes.h"
+#include "MassProcessingContext.h"
 #include "MassEntityTestTypes.h"
 #include "MassExecutor.h"
 #include "MassExecutionContext.h"
diff --git a/Engine/Source/Developer/MassEntityTestSuite/Private/MassTestTypes.cpp b/Engine/Source/Developer/MassEntityTestSuite/Private/MassTestTypes.cpp
--- a/Engine/Source/Developer/MassEntityTestSuite/Private/MassTestTypes.cpp
+++ b/Engine/Source/Developer/MassEntityTestSuite/Private/MassTestTypes.cpp
@@ -4,6 +4,7 @@
 #include "MassEntityManager.h"
 #include "MassExecutor.h"
 #include "MassExecutionContext.h"
+#include "MassProcessingContext.h"
 #include "Engine/World.h"
 
 
diff --git a/Engine/Source/Runtime/MassEntity/Private/MassEntityManager.cpp b/Engine/Source/Runtime/MassEntity/Private/MassEntityManager.cpp
--- a/Engine/Source/Runtime/MassEntity/Private/MassEntityManager.cpp
+++ b/Engine/Source/Runtime/MassEntity/Private/MassEntityManager.cpp
@@ -16,6 +16,7 @@
 #include "MassEntityUtils.h"
 #include "MassEntityBuilder.h"
 #include "MassTypeManager.h"
+#include "MassProcessingContext.h"
 
 #define CHECK_SYNC_API() testableCheckf(IsProcessing() == false, TEXT("Synchronous API function %hs called during mass processing. Use asynchronous API instead."), __FUNCTION__)
 #define CHECK_SYNC_API_RETURN(ReturnValue) testableCheckfReturn(IsProcessing() == false, return ReturnValue, TEXT("Synchronous API function %hs called during mass processing. Use asynchronous API instead."), __FUNCTION__)
@@ -959,12 +960,7 @@ void FMassEntityManager::BatchDestroyEntityChunks(TConstArrayView<FMassArchetype
 	TRACE_CPUPROFILER_EVENT_SCOPE(Mass_BatchDestroyEntityChunks);
 
 	TArray<FMassEntityHandle> EntitiesRemoved;
-	// note that it's important to place the context instance in the same scope as the loop below that updates 
-	// FMassEntityManager.EntityData, otherwise, if there are commands flushed as part of FMassProcessingContext's 
-	// destruction the commands will work on outdated information (which might result in crashes).
-	FMassProcessingContext ProcessingContext(*this, /*TimeDelta=*/0.0f);
-	ProcessingContext.bFlushCommandBuffer = false;
-	ProcessingContext.CommandBuffer = MakeShareable(new FMassCommandBuffer());
+	FMassProcessingContext ProcessingContext(*this);
 
 	for (const FMassArchetypeEntityCollection& EntityCollection : Collections)
 	{
diff --git a/Engine/Source/Runtime/MassEntity/Private/MassExecutor.cpp b/Engine/Source/Runtime/MassEntity/Private/MassExecutor.cpp
--- a/Engine/Source/Runtime/MassEntity/Private/MassExecutor.cpp
+++ b/Engine/Source/Runtime/MassEntity/Private/MassExecutor.cpp
@@ -5,6 +5,7 @@
 #include "MassProcessor.h"
 #include "MassCommandBuffer.h"
 #include "MassExecutionContext.h"
+#include "MassProcessingContext.h"
 #include "ProfilingDebugging/CpuProfilerTrace.h"
 
 namespace UE::Mass::Executor
@@ -21,11 +22,10 @@ FORCEINLINE void ExecuteProcessors(FMassEntityManager& EntityManager, TArrayView
 	}
 }
 
-void Run(FMassRuntimePipeline& RuntimePipeline, FMassProcessingContext& ProcessingContext)
+void Run(FMassRuntimePipeline& RuntimePipeline, FProcessingContext& ProcessingContext)
 {
-	if (!ensure(ProcessingContext.EntityManager) || 
-		!ensure(ProcessingContext.DeltaSeconds >= 0.f) ||
-		!ensure(RuntimePipeline.GetProcessors().Find(nullptr) == INDEX_NONE))
+	if (!ensure(ProcessingContext.DeltaSeconds >= 0.f) 
+		|| !ensure(RuntimePipeline.GetProcessors().Find(nullptr) == INDEX_NONE))
 	{
 		return;
 	}
@@ -34,12 +34,11 @@ void Run(FMassRuntimePipeline& RuntimePipeline, FMassProcessingContext& Processi
 	RunProcessorsView(RuntimePipeline.GetMutableProcessors(), ProcessingContext);
 }
 
-void RunSparse(FMassRuntimePipeline& RuntimePipeline, FMassProcessingContext& ProcessingContext, FMassArchetypeHandle Archetype, TConstArrayView<FMassEntityHandle> Entities)
+void RunSparse(FMassRuntimePipeline& RuntimePipeline, FProcessingContext& ProcessingContext, FMassArchetypeHandle Archetype, TConstArrayView<FMassEntityHandle> Entities)
 {
-	if (!ensure(ProcessingContext.EntityManager) ||
-		!ensure(RuntimePipeline.GetProcessors().Find(nullptr) == INDEX_NONE) ||
-		RuntimePipeline.Num() == 0 ||
-		!ensureMsgf(Archetype.IsValid(), TEXT("The Archetype passed in to UE::Mass::Executor::RunSparse is invalid")))
+	if (!ensure(RuntimePipeline.GetProcessors().Find(nullptr) == INDEX_NONE) 
+		|| RuntimePipeline.Num() == 0 
+		|| !ensureMsgf(Archetype.IsValid(), TEXT("The Archetype passed in to UE::Mass::Executor::RunSparse is invalid")))
 	{
 		return;
 	}
@@ -50,12 +49,11 @@ void RunSparse(FMassRuntimePipeline& RuntimePipeline, FMassProcessingContext& Pr
 	RunProcessorsView(RuntimePipeline.GetMutableProcessors(), ProcessingContext, MakeArrayView(&EntityCollection, 1));
 }
 
-void RunSparse(FMassRuntimePipeline& RuntimePipeline, FMassProcessingContext& ProcessingContext, const FMassArchetypeEntityCollection& EntityCollection)
+void RunSparse(FMassRuntimePipeline& RuntimePipeline, FProcessingContext& ProcessingContext, const FMassArchetypeEntityCollection& EntityCollection)
 {
-	if (!ensure(ProcessingContext.EntityManager) ||
-		!ensure(RuntimePipeline.GetProcessors().Find(nullptr) == INDEX_NONE) ||
-		RuntimePipeline.Num() == 0 ||
-		!ensureMsgf(EntityCollection.GetArchetype().IsValid(), TEXT("The Archetype of EntityCollection passed in to UE::Mass::Executor::RunSparse is invalid")))
+	if (!ensure(RuntimePipeline.GetProcessors().Find(nullptr) == INDEX_NONE) 
+		|| RuntimePipeline.Num() == 0
+		|| !ensureMsgf(EntityCollection.GetArchetype().IsValid(), TEXT("The Archetype of EntityCollection passed in to UE::Mass::Executor::RunSparse is invalid")))
 	{
 		return;
 	}
@@ -65,9 +63,9 @@ void RunSparse(FMassRuntimePipeline& RuntimePipeline, FMassProcessingContext& Pr
 	RunProcessorsView(RuntimePipeline.GetMutableProcessors(), ProcessingContext, MakeArrayView(&EntityCollection, 1));
 }
 
-void Run(UMassProcessor& Processor, FMassProcessingContext& ProcessingContext)
+void Run(UMassProcessor& Processor, FProcessingContext& ProcessingContext)
 {
-	if (!ensure(ProcessingContext.EntityManager) || !ensure(ProcessingContext.DeltaSeconds >= 0.f))
+	if (!ensure(ProcessingContext.DeltaSeconds >= 0.f))
 	{
 		return;
 	}
@@ -78,15 +76,8 @@ void Run(UMassProcessor& Processor, FMassProcessingContext& ProcessingContext)
 	RunProcessorsView(MakeArrayView(&ProcPtr, 1), ProcessingContext);
 }
 
-void RunProcessorsView(TArrayView<UMassProcessor* const> Processors, FMassProcessingContext& ProcessingContext, TConstArrayView<FMassArchetypeEntityCollection> EntityCollections)
+void RunProcessorsView(TArrayView<UMassProcessor* const> Processors, FProcessingContext& ProcessingContext, TConstArrayView<FMassArchetypeEntityCollection> EntityCollections)
 {
-	TRACE_CPUPROFILER_EVENT_SCOPE(RunProcessorsView);
-
-	if (!ProcessingContext.EntityManager)
-	{
-		UE_LOG(LogMass, Error, TEXT("%s ProcessingContext.EntityManager is null. Baling out."), ANSI_TO_TCHAR(__FUNCTION__));
-		return;
-	}
 #if WITH_MASSENTITY_DEBUG
 	if (Processors.Find(nullptr) != INDEX_NONE)
 	{
@@ -95,63 +86,23 @@ void RunProcessorsView(TArrayView<UMassProcessor* const> Processors, FMassProces
 	}
 #endif // WITH_MASSENTITY_DEBUG
 
-	TRACE_CPUPROFILER_EVENT_SCOPE_STR("MassExecutor RunProcessorsView")
-
-	FMassExecutionContext ExecutionContext(*ProcessingContext.EntityManager.Get(), ProcessingContext.DeltaSeconds);
+	TRACE_CPUPROFILER_EVENT_SCOPE(MassExecutor_RunProcessorsView);
 
-	// if ProcessingContext points at a valid CommandBuffer use that one, otherwise manually create a new command buffer 
-	// to let the default one still be used by code unaware of mass processing
-	TSharedPtr<FMassCommandBuffer> CommandBuffer = ProcessingContext.CommandBuffer 
-		? ProcessingContext.CommandBuffer : MakeShareable(new FMassCommandBuffer());
-	ExecutionContext.SetDeferredCommandBuffer(CommandBuffer);
-	ExecutionContext.SetFlushDeferredCommands(false);
-	ExecutionContext.SetAuxData(ProcessingContext.AuxData);
-	ExecutionContext.SetExecutionType(EMassExecutionContextType::Processor);
+	FMassExecutionContext& ExecutionContext = ProcessingContext.GetExecutionContext();
+	FMassEntityManager& EntityManager = *ProcessingContext.GetEntityManager();
+	FMassEntityManager::FScopedProcessing ProcessingScope = EntityManager.NewProcessingScope();
 
-	FMassEntityManager& EntityManager = *ProcessingContext.EntityManager.Get();
+	if (EntityCollections.Num() == 0)
 	{
-		TRACE_CPUPROFILER_EVENT_SCOPE_STR("Execute Processors")
-		
-		FMassEntityManager::FScopedProcessing ProcessingScope = EntityManager.NewProcessingScope();
-
-		if (EntityCollections.Num() == 0)
-		{
-			ExecuteProcessors(*ProcessingContext.EntityManager, Processors, ExecutionContext);
-		}
-		else
-		{
-			// @todo change ExecutionContext to contain a TConstArrayView of collections. Will lead up to changes in the entity query as well.
-			for (const FMassArchetypeEntityCollection& Collection : EntityCollections)
-			{
-				ExecutionContext.SetEntityCollection(Collection);
-				ExecuteProcessors(*ProcessingContext.EntityManager, Processors, ExecutionContext);
-				ExecutionContext.ClearEntityCollection();
-			}
-		}
-	}
-	
-	if (ProcessingContext.bFlushCommandBuffer)
-	{		
-		TRACE_CPUPROFILER_EVENT_SCOPE_STR("Flush Deferred Commands")
-		
-		ExecutionContext.SetFlushDeferredCommands(true);
-		// append the commands added from other, non-processor sources (like MassAgentSubsystem)
-		ensure(!EntityManager.Defer().IsFlushing());
-		ensure(!ExecutionContext.Defer().IsFlushing());
-		ExecutionContext.Defer().MoveAppend(EntityManager.Defer());
-		ExecutionContext.FlushDeferred();
+		ExecuteProcessors(EntityManager, Processors, ExecutionContext);
 	}
-	// else make sure we don't just lose the commands. Append to the command buffer requested via
-	// ProcessingContext.CommandBuffer or to the default EntityManager's command buffer.
-	else if (CommandBuffer->HasPendingCommands() && CommandBuffer != ProcessingContext.CommandBuffer)
+	else
 	{
-		if (ProcessingContext.CommandBuffer)
+		for (const FMassArchetypeEntityCollection& Collection : EntityCollections)
 		{
-			ProcessingContext.CommandBuffer->MoveAppend(*CommandBuffer.Get());
-		}
-		else if (CommandBuffer.Get() != &EntityManager.Defer())
-		{
-			EntityManager.AppendCommands(CommandBuffer);
+			ExecutionContext.SetEntityCollection(Collection);
+			ExecuteProcessors(EntityManager, Processors, ExecutionContext);
+			ExecutionContext.ClearEntityCollection();
 		}
 	}
 }
@@ -201,29 +152,18 @@ private:
 	ENamedThreads::Type DesiredThread;
 };
 
-FGraphEventRef TriggerParallelTasks(UMassProcessor& Processor, FMassProcessingContext& ProcessingContext, TFunction<void()> OnDoneNotification, ENamedThreads::Type CurrentThread)
+FGraphEventRef TriggerParallelTasks(UMassProcessor& Processor, FProcessingContext&& ProcessingContext, TFunction<void()> OnDoneNotification, ENamedThreads::Type CurrentThread)
 {
-	if (!ProcessingContext.EntityManager)
-	{
-		UE_LOG(LogMass, Error, TEXT("%s ProcessingContext.EntityManager is null. Baling out."), ANSI_TO_TCHAR(__FUNCTION__));
-		return FGraphEventRef();
-	}
-
-	TRACE_CPUPROFILER_EVENT_SCOPE_STR("MassExecutor RunParallel")
+	TRACE_CPUPROFILER_EVENT_SCOPE(MassExecutor_RunParallel);
 
-	// not going through FMassEntityManager::CreateExecutionContext on purpose - we do need a separate command buffer
-	FMassExecutionContext ExecutionContext(*ProcessingContext.EntityManager.Get(), ProcessingContext.DeltaSeconds);
-	TSharedPtr<FMassCommandBuffer> CommandBuffer = ProcessingContext.CommandBuffer
-		? ProcessingContext.CommandBuffer : MakeShareable(new FMassCommandBuffer());
-	ExecutionContext.SetDeferredCommandBuffer(CommandBuffer);
-	ExecutionContext.SetFlushDeferredCommands(false);
-	ExecutionContext.SetAuxData(ProcessingContext.AuxData);
-	ExecutionContext.SetExecutionType(EMassExecutionContextType::Processor);
+	// We need to transfer ProcessingContext's ExecutionContext - otherwise ProcessingContext's destructor will attempt
+	// flushing stored commands. 
+	FMassExecutionContext ExecutionContext = MoveTemp(ProcessingContext).GetExecutionContext();
 
 	FGraphEventRef CompletionEvent;
 	{
 		TRACE_CPUPROFILER_EVENT_SCOPE_STR("Dispatch Processors")
-		CompletionEvent = Processor.DispatchProcessorTasks(ProcessingContext.EntityManager, ExecutionContext, {});
+		CompletionEvent = Processor.DispatchProcessorTasks(ProcessingContext.GetEntityManager(), ExecutionContext, {});
 	}
 
 	if (CompletionEvent.IsValid())
@@ -239,7 +179,7 @@ FGraphEventRef TriggerParallelTasks(UMassProcessor& Processor, FMassProcessingCo
 //-----------------------------------------------------------------------------
 // DEPRECATED
 //-----------------------------------------------------------------------------
-void RunProcessorsView(TArrayView<UMassProcessor* const> Processors, FMassProcessingContext& ProcessingContext, const FMassArchetypeEntityCollection* EntityCollection)
+void RunProcessorsView(TArrayView<UMassProcessor* const> Processors, FProcessingContext& ProcessingContext, const FMassArchetypeEntityCollection* EntityCollection)
 {
 	if (EntityCollection)
 	{
@@ -251,4 +191,11 @@ void RunProcessorsView(TArrayView<UMassProcessor* const> Processors, FMassProces
 	}
 }
 
+inline FGraphEventRef TriggerParallelTasks(UMassProcessor& Processor, FProcessingContext& ProcessingContext, TFunction<void()> OnDoneNotification
+		, ENamedThreads::Type CurrentThread)
+{
+	FProcessingContext LocalContext = ProcessingContext;
+	return TriggerParallelTasks(Processor, MoveTemp(ProcessingContext), OnDoneNotification, CurrentThread);
+}
+
 } // namespace UE::Mass::Executor
diff --git a/Engine/Source/Runtime/MassEntity/Private/MassObserverManager.cpp b/Engine/Source/Runtime/MassEntity/Private/MassObserverManager.cpp
--- a/Engine/Source/Runtime/MassEntity/Private/MassObserverManager.cpp
+++ b/Engine/Source/Runtime/MassEntity/Private/MassObserverManager.cpp
@@ -7,6 +7,8 @@
 #include "MassProcessingTypes.h"
 #include "MassObserverRegistry.h"
 #include "MassDebugger.h"
+#include "MassProcessingContext.h"
+
 
 #include UE_INLINE_GENERATED_CPP_BY_NAME(MassObserverManager)
 
@@ -139,28 +141,17 @@ bool FMassObserverManager::OnPostEntitiesCreated(const FMassArchetypeEntityColle
 
 bool FMassObserverManager::OnPostEntitiesCreated(TConstArrayView<FMassArchetypeEntityCollection> EntityCollections)
 {
-	FMassProcessingContext ProcessingContext(EntityManager, /*DeltaSeconds=*/0.f);
-	// requesting not to flush commands since handling creation of new entities can result in multiple collections of
-	// processors being executed and flushing commands between these runs would ruin EntityCollection since entities could
-	// get their composition changed and get moved to new archetypes
-	ProcessingContext.bFlushCommandBuffer = false;
-	ProcessingContext.CommandBuffer = MakeShareable(new FMassCommandBuffer());
-
-	if (OnPostEntitiesCreated(ProcessingContext, EntityCollections))
-	{
-		EntityManager.FlushCommands(ProcessingContext.CommandBuffer);
-		return true;
-	}
-	return false;
+	UE::Mass::FProcessingContext ProcessingContext(EntityManager, /*DeltaSeconds=*/0.f);
+	return OnPostEntitiesCreated(ProcessingContext, EntityCollections);
 }
 
 // DEPRECATED
-bool FMassObserverManager::OnPostEntitiesCreated(FMassProcessingContext& ProcessingContext, const FMassArchetypeEntityCollection& EntityCollection)
+bool FMassObserverManager::OnPostEntitiesCreated(UE::Mass::FProcessingContext& ProcessingContext, const FMassArchetypeEntityCollection& EntityCollection)
 {
 	return OnPostEntitiesCreated(ProcessingContext, MakeArrayView(&EntityCollection, 1));
 }
 
-bool FMassObserverManager::OnPostEntitiesCreated(FMassProcessingContext& ProcessingContext, TConstArrayView<FMassArchetypeEntityCollection> EntityCollections)
+bool FMassObserverManager::OnPostEntitiesCreated(UE::Mass::FProcessingContext& ProcessingContext, TConstArrayView<FMassArchetypeEntityCollection> EntityCollections)
 {
 	if (!ensureMsgf(EntityManager.IsDuringEntityCreation() == false, TEXT("Entity creation not done yet.")))
 	{
@@ -169,13 +160,11 @@ bool FMassObserverManager::OnPostEntitiesCreated(FMassProcessingContext& Process
 
 	TRACE_CPUPROFILER_EVENT_SCOPE_STR("OnPostEntitiesCreated")
 
-	check(ProcessingContext.EntityManager);
-
 	bool bReturnValue = false;
 
 	for (const FMassArchetypeEntityCollection& Collection : EntityCollections)
 	{
-		const FMassArchetypeCompositionDescriptor& ArchetypeComposition = ProcessingContext.EntityManager->GetArchetypeComposition(Collection.GetArchetype());
+		const FMassArchetypeCompositionDescriptor& ArchetypeComposition = ProcessingContext.GetEntityManager()->GetArchetypeComposition(Collection.GetArchetype());
 		bReturnValue = OnCompositionChanged(ProcessingContext, Collection, ArchetypeComposition, EMassObservedOperation::Add) || bReturnValue;
 	}
 
@@ -184,24 +173,15 @@ bool FMassObserverManager::OnPostEntitiesCreated(FMassProcessingContext& Process
 
 bool FMassObserverManager::OnPreEntitiesDestroyed(const FMassArchetypeEntityCollection& EntityCollection)
 {
-	FMassProcessingContext ProcessingContext(EntityManager, /*DeltaSeconds=*/0.f);
-	ProcessingContext.bFlushCommandBuffer = false;
-	ProcessingContext.CommandBuffer = MakeShareable(new FMassCommandBuffer());
-
-	if (OnPreEntitiesDestroyed(ProcessingContext, EntityCollection))
-	{
-		EntityManager.FlushCommands(ProcessingContext.CommandBuffer);
-		return true;
-	}
-	return false;
+	UE::Mass::FProcessingContext ProcessingContext(EntityManager);
+	return OnPreEntitiesDestroyed(ProcessingContext, EntityCollection);
 }
 
-bool FMassObserverManager::OnPreEntitiesDestroyed(FMassProcessingContext& ProcessingContext, const FMassArchetypeEntityCollection& EntityCollection)
+bool FMassObserverManager::OnPreEntitiesDestroyed(UE::Mass::FProcessingContext& ProcessingContext, const FMassArchetypeEntityCollection& EntityCollection)
 {
 	TRACE_CPUPROFILER_EVENT_SCOPE_STR("OnPreEntitiesDestroyed")
 
-	check(ProcessingContext.EntityManager);
-	const FMassArchetypeCompositionDescriptor& ArchetypeComposition = ProcessingContext.EntityManager->GetArchetypeComposition(EntityCollection.GetArchetype());
+	const FMassArchetypeCompositionDescriptor& ArchetypeComposition = ProcessingContext.GetEntityManager()->GetArchetypeComposition(EntityCollection.GetArchetype());
 	
 	return OnCompositionChanged(ProcessingContext, EntityCollection, ArchetypeComposition, EMassObservedOperation::Remove);
 }
@@ -212,7 +192,7 @@ bool FMassObserverManager::OnPreEntityDestroyed(const FMassArchetypeCompositionD
 	return OnCompositionChanged(Entity, ArchetypeComposition, EMassObservedOperation::Remove);
 }
 
-bool FMassObserverManager::OnCompositionChanged(FMassProcessingContext& ProcessingContext, const FMassArchetypeEntityCollection& EntityCollection, const FMassArchetypeCompositionDescriptor& CompositionDelta, const EMassObservedOperation Operation)
+bool FMassObserverManager::OnCompositionChanged(UE::Mass::FProcessingContext& ProcessingContext, const FMassArchetypeEntityCollection& EntityCollection, const FMassArchetypeCompositionDescriptor& CompositionDelta, const EMassObservedOperation Operation)
 {
 	using UE::Mass::ObserverManager::Tweakables::InlineAllocatorElementsForOverlapTypes;
 
@@ -248,6 +228,12 @@ bool FMassObserverManager::OnCompositionChanged(FMassProcessingContext& Processi
 	return false;
 }
 
+bool FMassObserverManager::OnCompositionChanged(const FMassArchetypeEntityCollection& EntityCollection, const FMassArchetypeCompositionDescriptor& Composition, const EMassObservedOperation Operation)
+{
+	UE::Mass::FProcessingContext LocalContext(EntityManager, /*DeltaSeconds=*/0.f);
+	return OnCompositionChanged(LocalContext, EntityCollection, Composition, Operation);
+}
+
 bool FMassObserverManager::OnCompositionChanged(const FMassEntityHandle Entity, const FMassArchetypeCompositionDescriptor& CompositionDelta, const EMassObservedOperation Operation)
 {
 	using UE::Mass::ObserverManager::Tweakables::InlineAllocatorElementsForOverlapTypes;
@@ -260,7 +246,7 @@ bool FMassObserverManager::OnCompositionChanged(const FMassEntityHandle Entity,
 	if (bHasFragmentsOverlap || bHasTagsOverlap)
 	{
 		TArray<const UScriptStruct*, TInlineAllocator<InlineAllocatorElementsForOverlapTypes>> ObservedTypesOverlap;
-		FMassProcessingContext ProcessingContext(EntityManager, /*DeltaSeconds=*/0.f);
+		UE::Mass::FProcessingContext ProcessingContext(EntityManager, /*DeltaSeconds=*/0.f);
 		ProcessingContext.bFlushCommandBuffer = false;
 		const FMassArchetypeHandle ArchetypeHandle = EntityManager.GetArchetypeForEntity(Entity);
 
@@ -303,7 +289,7 @@ void FMassObserverManager::OnElementOperation(TNotNull<const UScriptStruct*> Ele
 	}
 }
 
-void FMassObserverManager::HandleFragmentsImpl(FMassProcessingContext& ProcessingContext, const FMassArchetypeEntityCollection& EntityCollection
+void FMassObserverManager::HandleFragmentsImpl(UE::Mass::FProcessingContext& ProcessingContext, const FMassArchetypeEntityCollection& EntityCollection
 	, TArrayView<const UScriptStruct*> ObservedTypes
 	/*, const FMassFragmentBitSet& FragmentsBitSet*/, FMassObserversMap& HandlersContainer)
 {	
@@ -324,7 +310,7 @@ void FMassObserverManager::HandleSingleElementImpl(const UScriptStruct& Fragment
 {
 	TRACE_CPUPROFILER_EVENT_SCOPE(MassObserver_HandleSingleEntityImpl);
 
-	FMassProcessingContext ProcessingContext(EntityManager, /*DeltaSeconds=*/0.f);
+	UE::Mass::FProcessingContext ProcessingContext(EntityManager, /*DeltaSeconds=*/0.f);
 	ProcessingContext.bFlushCommandBuffer = false;
 	ProcessingContext.AuxData.InitializeAs(&FragmentType);
 	FMassRuntimePipeline& Pipeline = (*HandlersContainer).FindChecked(&FragmentType);
diff --git a/Engine/Source/Runtime/MassEntity/Private/MassProcessingContext.cpp b/Engine/Source/Runtime/MassEntity/Private/MassProcessingContext.cpp
new file mode 100644
--- /dev/null
+++ b/Engine/Source/Runtime/MassEntity/Private/MassProcessingContext.cpp
@@ -0,0 +1,33 @@
+// Copyright Epic Games, Inc. All Rights Reserved.
+
+#include "MassProcessingContext.h"
+
+namespace UE::Mass
+{
+	// @todo remove the depracation disabling once CommandBuffer and EntityManager are moved to `protected` and un-deprecated (around UE5.8)
+	PRAGMA_DISABLE_DEPRECATION_WARNINGS
+	FProcessingContext::~FProcessingContext()
+	{	
+		if (ExecutionContextPtr)
+		{
+			checkf(ExecutionContextPtr->GetSharedDeferredCommandBuffer(), TEXT("A valid execution context without a valid command buffer is unexpected"));
+			checkf(ExecutionContextPtr->GetSharedDeferredCommandBuffer() == CommandBuffer
+			   , TEXT("Hosted Execution Context's command buffer is not the same as FProcessingContext's command buffer. "
+				 "This is not supposed to happen. Make sure FProcessingContext.CommandBuffer never gets assigned after FMassExecutionContext's creation"));
+
+			ensure(!CommandBuffer->IsFlushing());
+
+			if (bFlushCommandBuffer)
+			{
+				EntityManager->FlushCommands(CommandBuffer);
+			}
+			else
+			{
+				EntityManager->AppendCommands(CommandBuffer);
+			}
+
+			ExecutionContextPtr->~FMassExecutionContext();
+		}
+	}
+	PRAGMA_ENABLE_DEPRECATION_WARNINGS
+} // namespace UE::Mass
diff --git a/Engine/Source/Runtime/MassEntity/Private/MassProcessingPhaseManager.cpp b/Engine/Source/Runtime/MassEntity/Private/MassProcessingPhaseManager.cpp
--- a/Engine/Source/Runtime/MassEntity/Private/MassProcessingPhaseManager.cpp
+++ b/Engine/Source/Runtime/MassEntity/Private/MassProcessingPhaseManager.cpp
@@ -11,6 +11,7 @@
 #include "Engine/World.h"
 #include "MassCommandBuffer.h"
 #include "MassEntityTrace.h"
+#include "MassProcessingContext.h"
 
 #include UE_INLINE_GENERATED_CPP_BY_NAME(MassProcessingPhaseManager)
 
@@ -85,7 +86,7 @@ void FMassProcessingPhase::ExecuteTick(float DeltaTime, ELevelTick TickType, ENa
 		bool bWorkRequested = false;
 		if (PhaseProcessor->IsEmpty() == false)
 		{
-			const FGraphEventRef PipelineCompletionEvent = UE::Mass::Executor::TriggerParallelTasks(*PhaseProcessor, Context, [this, DeltaTime]()
+			const FGraphEventRef PipelineCompletionEvent = UE::Mass::Executor::TriggerParallelTasks(*PhaseProcessor, MoveTemp(Context), [this, DeltaTime]()
 				{
 					OnParallelExecutionDone(DeltaTime);
 				}
@@ -586,6 +587,11 @@ void FMassProcessingPhaseManager::OnPhaseEnd(FMassProcessingPhase& Phase)
 		UE_VLOG_UELOG(Owner.Get(), LogMass, Log, TEXT("Phase Processing is now Paused"));
 #endif // WITH_MASSENTITY_DEBUG
 	}
+
+	if (GetEntityManagerRef().Defer().HasPendingCommands())
+	{
+		GetEntityManagerRef().FlushCommands();
+	}
 }
 
 FString FMassProcessingPhaseManager::GetName() const
diff --git a/Engine/Source/Runtime/MassEntity/Private/MassProcessingTypes.cpp b/Engine/Source/Runtime/MassEntity/Private/MassProcessingTypes.cpp
--- a/Engine/Source/Runtime/MassEntity/Private/MassProcessingTypes.cpp
+++ b/Engine/Source/Runtime/MassEntity/Private/MassProcessingTypes.cpp
@@ -11,35 +11,6 @@
 
 DEFINE_LOG_CATEGORY(LogMass);
 
-//----------------------------------------------------------------------//
-//  FMassProcessingContext
-//----------------------------------------------------------------------//
-FMassProcessingContext::FMassProcessingContext(FMassEntityManager& InEntityManager, const float InDeltaSeconds)
-	: EntityManager(InEntityManager.AsShared()), DeltaSeconds(InDeltaSeconds)
-{
-
-}
-
-FMassProcessingContext::FMassProcessingContext(TSharedPtr<FMassEntityManager>& InEntityManager, const float InDeltaSeconds)
-	: EntityManager(InEntityManager), DeltaSeconds(InDeltaSeconds)
-{
-
-}
-
-FMassProcessingContext::~FMassProcessingContext()
-{
-	if (CommandBuffer && CommandBuffer.IsUnique() && CommandBuffer->HasPendingCommands())
-	{
-		UE_CLOG(!EntityManager, LogMass, Error, TEXT("Unable to auto-flush FMassProcessingContext\'s commands due to missing EntityManager"));
-		if (ensure(EntityManager))
-		{
-			UE_VLOG(EntityManager->GetOwner(), LogMass, Log, TEXT("Auto-flushing command buffer as part of FMassProcessingContext destruction"));
-			checkf(CommandBuffer->IsFlushing() == false, TEXT("A totally unexpected scenario."));
-			EntityManager->FlushCommands(CommandBuffer);
-		}
-	}
-}
-
 //----------------------------------------------------------------------//
 //  FMassRuntimePipeline
 //----------------------------------------------------------------------//
diff --git a/Engine/Source/Runtime/MassEntity/Public/MassExecutor.h b/Engine/Source/Runtime/MassEntity/Public/MassExecutor.h
--- a/Engine/Source/Runtime/MassEntity/Public/MassExecutor.h
+++ b/Engine/Source/Runtime/MassEntity/Public/MassExecutor.h
@@ -7,7 +7,10 @@
 
 
 struct FMassRuntimePipeline;
-struct FMassProcessingContext;
+namespace UE::Mass
+{
+	struct FProcessingContext;
+}
 struct FMassEntityHandle;
 struct FMassArchetypeEntityCollection;
 class UMassProcessor;
@@ -16,26 +19,26 @@ class UMassProcessor;
 namespace UE::Mass::Executor
 {
 	/** Executes processors in a given RuntimePipeline */
-	MASSENTITY_API void Run(FMassRuntimePipeline& RuntimePipeline, FMassProcessingContext& ProcessingContext);
+	MASSENTITY_API void Run(FMassRuntimePipeline& RuntimePipeline, FProcessingContext& ProcessingContext);
 
 	/** Executes given Processor. Used mainly for triggering calculations via MassCompositeProcessors, e.g processing phases */
-	MASSENTITY_API void Run(UMassProcessor& Processor, FMassProcessingContext& ProcessingContext);
+	MASSENTITY_API void Run(UMassProcessor& Processor, FProcessingContext& ProcessingContext);
 
 	/** Similar to the Run function, but instead of using all the entities hosted by ProcessingContext.EntitySubsystem 
 	 *  it is processing only the entities given by EntityID via the Entities input parameter. 
 	 *  Note that all the entities need to be of Archetype archetype. 
 	 *  Under the hood the function converts Archetype-Entities pair to FMassArchetypeEntityCollection and calls the other flavor of RunSparse
 	 */
-	MASSENTITY_API void RunSparse(FMassRuntimePipeline& RuntimePipeline, FMassProcessingContext& ProcessingContext, FMassArchetypeHandle Archetype, TConstArrayView<FMassEntityHandle> Entities);
+	MASSENTITY_API void RunSparse(FMassRuntimePipeline& RuntimePipeline, FProcessingContext& ProcessingContext, FMassArchetypeHandle Archetype, TConstArrayView<FMassEntityHandle> Entities);
 
 	/** Similar to the Run function, but instead of using all the entities hosted by ProcessingContext.EntitySubsystem 
 	 *  it is processing only the entities given by SparseEntities input parameter.
 	 *  @todo rename
 	 */
-	MASSENTITY_API void RunSparse(FMassRuntimePipeline& RuntimePipeline, FMassProcessingContext& ProcessingContext, const FMassArchetypeEntityCollection& EntityCollection);
+	MASSENTITY_API void RunSparse(FMassRuntimePipeline& RuntimePipeline, FProcessingContext& ProcessingContext, const FMassArchetypeEntityCollection& EntityCollection);
 
 	/** Executes given Processors array view. This function gets called under the hood by the rest of Run* functions */
-	MASSENTITY_API void RunProcessorsView(TArrayView<UMassProcessor* const> Processors, FMassProcessingContext& ProcessingContext, TConstArrayView<FMassArchetypeEntityCollection> EntityCollections = {});
+	MASSENTITY_API void RunProcessorsView(TArrayView<UMassProcessor* const> Processors, FProcessingContext& ProcessingContext, TConstArrayView<FMassArchetypeEntityCollection> EntityCollections = {});
 
 	/** 
 	 *  Triggers tasks executing Processor (and potentially it's children) and returns the task graph event representing 
@@ -43,9 +46,13 @@ namespace UE::Mass::Executor
 	 *  @param OnDoneNotification will be called after all the processors are done, just after flushing the command buffer.
 	 *    Note that OnDoneNotification will be executed on GameThread.
 	 */
-	MASSENTITY_API FGraphEventRef TriggerParallelTasks(UMassProcessor& Processor, FMassProcessingContext& ProcessingContext, TFunction<void()> OnDoneNotification
+	MASSENTITY_API FGraphEventRef TriggerParallelTasks(UMassProcessor& Processor, FProcessingContext&& ProcessingContext, TFunction<void()> OnDoneNotification
 		, ENamedThreads::Type CurrentThread = ENamedThreads::GameThread);
 
 	UE_DEPRECATED(5.5, "This flavor of RunProcessorsView is deprecated. Use the one with TConstArrayView<FMassArchetypeEntityCollection> parameter instead.")
-	MASSENTITY_API void RunProcessorsView(TArrayView<UMassProcessor* const> Processors, FMassProcessingContext& ProcessingContext, const FMassArchetypeEntityCollection* EntityCollection);
+	MASSENTITY_API void RunProcessorsView(TArrayView<UMassProcessor* const> Processors, FProcessingContext& ProcessingContext, const FMassArchetypeEntityCollection* EntityCollection);
+
+	UE_DEPRECATED(5.6, "lvalue flavor of TriggerParallelTasks has been deprevate. Use the rvalue version.")
+	MASSENTITY_API FGraphEventRef TriggerParallelTasks(UMassProcessor& Processor, FProcessingContext& ProcessingContext, TFunction<void()> OnDoneNotification
+		, ENamedThreads::Type CurrentThread = ENamedThreads::GameThread);
 };
diff --git a/Engine/Source/Runtime/MassEntity/Public/MassObserverManager.h b/Engine/Source/Runtime/MassEntity/Public/MassObserverManager.h
--- a/Engine/Source/Runtime/MassEntity/Public/MassObserverManager.h
+++ b/Engine/Source/Runtime/MassEntity/Public/MassObserverManager.h
@@ -11,7 +11,10 @@
 struct FMassEntityManager;
 struct FMassArchetypeEntityCollection;
 class UMassProcessor;
-struct FMassProcessingContext;
+namespace UE::Mass
+{
+	struct FProcessingContext;
+}
 
 /** 
  * A wrapper type for a TMap to support having array-of-maps UPROPERTY members in FMassObserverManager
@@ -81,13 +84,13 @@ public:
 	bool OnPostEntitiesCreated(TConstArrayView<FMassArchetypeEntityCollection> EntityCollections);
 
 	/** @return whether any observers have been executed */
-	bool OnPostEntitiesCreated(FMassProcessingContext& InProcessingContext, TConstArrayView<FMassArchetypeEntityCollection> EntityCollections);
+	bool OnPostEntitiesCreated(UE::Mass::FProcessingContext& InProcessingContext, TConstArrayView<FMassArchetypeEntityCollection> EntityCollections);
 
 	/** @return whether any observers have been executed */
 	bool OnPreEntitiesDestroyed(const FMassArchetypeEntityCollection& EntityCollection);
 
 	/** @return whether any observers have been executed */
-	bool OnPreEntitiesDestroyed(FMassProcessingContext& ProcessingContext, const FMassArchetypeEntityCollection& EntityCollection);
+	bool OnPreEntitiesDestroyed(UE::Mass::FProcessingContext& ProcessingContext, const FMassArchetypeEntityCollection& EntityCollection);
 
 	/** @return whether any observers have been executed */
 	bool OnPreEntityDestroyed(const FMassArchetypeCompositionDescriptor& ArchetypeComposition, const FMassEntityHandle Entity);
@@ -107,14 +110,10 @@ public:
 	}
 
 	/** @return whether any observers have been executed */
-	bool OnCompositionChanged(FMassProcessingContext& InProcessingContext, const FMassArchetypeEntityCollection& EntityCollection, const FMassArchetypeCompositionDescriptor& Composition, const EMassObservedOperation Operation);
+	bool OnCompositionChanged(UE::Mass::FProcessingContext& InProcessingContext, const FMassArchetypeEntityCollection& EntityCollection, const FMassArchetypeCompositionDescriptor& Composition, const EMassObservedOperation Operation);
+
 	/** @return whether any observers have been executed */
-	bool OnCompositionChanged(const FMassArchetypeEntityCollection& EntityCollection, const FMassArchetypeCompositionDescriptor& Composition, const EMassObservedOperation Operation)
-	{
-		FMassProcessingContext LocalContext(EntityManager, /*DeltaSeconds=*/0.f);
-		LocalContext.bFlushCommandBuffer = false;
-		return OnCompositionChanged(LocalContext, EntityCollection, Composition, Operation);
-	}
+	bool OnCompositionChanged(const FMassArchetypeEntityCollection& EntityCollection, const FMassArchetypeCompositionDescriptor& Composition, const EMassObservedOperation Operation);
 
 	void OnPostElementAdded(TNotNull<const UScriptStruct*> ElementType, const FMassArchetypeEntityCollection& EntityCollection)
 	{
@@ -137,7 +136,7 @@ protected:
 
 	void Initialize();
 	void DeInitialize();
-	void HandleFragmentsImpl(FMassProcessingContext& ProcessingContext, const FMassArchetypeEntityCollection& EntityCollection
+	void HandleFragmentsImpl(UE::Mass::FProcessingContext& ProcessingContext, const FMassArchetypeEntityCollection& EntityCollection
 		, TArrayView<const UScriptStruct*> ObservedTypes, FMassObserversMap& HandlersContainer);
 	void HandleSingleElementImpl(const UScriptStruct& FragmentType, const FMassArchetypeEntityCollection& EntityCollection, FMassObserversMap& HandlersContainer);
 
@@ -158,7 +157,7 @@ protected:
 
 public:
 	UE_DEPRECATED(5.5, "This flavor of OnPostEntitiesCreated is deprecated. Please use the one taking a TConstArrayView<FMassArchetypeEntityCollection> parameter instead")
-	bool OnPostEntitiesCreated(FMassProcessingContext& InProcessingContext, const FMassArchetypeEntityCollection& EntityCollection);
+	bool OnPostEntitiesCreated(UE::Mass::FProcessingContext& InProcessingContext, const FMassArchetypeEntityCollection& EntityCollection);
 	UE_DEPRECATED(5.5, "*FragmentOrTag* functions have been renamed to *Element*. Please use those instead.")
 	void OnPostFragmentOrTagAdded(const UScriptStruct& FragmentOrTagType, const FMassArchetypeEntityCollection& EntityCollection);
 	UE_DEPRECATED(5.5, "*FragmentOrTag* functions have been renamed to *Element*. Please use those instead.")
diff --git a/Engine/Source/Runtime/MassEntity/Public/MassProcessingContext.h b/Engine/Source/Runtime/MassEntity/Public/MassProcessingContext.h
new file mode 100644
--- /dev/null
+++ b/Engine/Source/Runtime/MassEntity/Public/MassProcessingContext.h
@@ -0,0 +1,146 @@
+// Copyright Epic Games, Inc. All Rights Reserved.
+
+#pragma once
+
+#include "MassExecutionContext.h"
+#include "MassCommandBuffer.h"
+#include "MassEntityManager.h"
+#include "StructUtils/InstancedStruct.h"
+
+
+namespace UE::Mass
+{
+	struct MASSENTITY_API FProcessingContext
+	{
+		UE_DEPRECATED_FORGAME(5.6, "This constructor is deprecated. Use one of the other ones.")
+		FProcessingContext() = default;
+		UE_DEPRECATED_FORGAME(5.6, "This constructor is deprecated. Use one of the other ones.")
+		FProcessingContext(const FProcessingContext&) = default;
+
+		FProcessingContext(FProcessingContext&&) = default;
+
+		explicit FProcessingContext(FMassEntityManager& InEntityManager, const float InDeltaSeconds = 0.f);
+		explicit FProcessingContext(const TSharedRef<FMassEntityManager>& InEntityManager, const float InDeltaSeconds = 0.f);
+		explicit FProcessingContext(TSharedRef<FMassEntityManager>&& InEntityManager, const float InDeltaSeconds = 0.f);
+		explicit FProcessingContext(const TSharedPtr<FMassEntityManager>& InEntityManager, const float InDeltaSeconds = 0.f);
+		~FProcessingContext();
+		
+		[[nodiscard]] FMassExecutionContext& GetExecutionContext() &;
+		[[nodiscard]] FMassExecutionContext&& GetExecutionContext() &&;
+
+		void SetCommandBuffer(TSharedPtr<FMassCommandBuffer>&& InCommandBuffer);
+		void SetCommandBuffer(const TSharedPtr<FMassCommandBuffer>& InCommandBuffer);
+
+		const TSharedRef<FMassEntityManager>& GetEntityManager() const;
+
+		UE_DEPRECATED_FORGAME(5.6, "Direct access to FProcessingContext.EntityManager has been deprecated. Use GetEntityManager instead.")
+		TSharedRef<FMassEntityManager> EntityManager;
+
+		float DeltaSeconds = 0.f;
+
+		FInstancedStruct AuxData;
+
+		/** 
+		 * If set to "true" the MassExecutor will flush commands at the end of given execution function. 
+		 * If "false" the caller is responsible for manually flushing the commands.
+		 */
+		bool bFlushCommandBuffer = true; 
+
+		UE_DEPRECATED_FORGAME(5.6, "Direct access to FProcessingContext's CommandBuffer has been deprecated. Use SetCommandBuffer()")
+		TSharedPtr<FMassCommandBuffer> CommandBuffer;
+
+	protected:
+		uint8 ExecutionContextBuffer[sizeof(FMassExecutionContext)];
+		FMassExecutionContext* ExecutionContextPtr = nullptr;
+	};
+
+	//----------------------------------------------------------------------//
+	// INLINES
+	//----------------------------------------------------------------------//
+	// @todo remove the depracation disabling once CommandBuffer and EntityManager are moved to `protected` and un-deprecated (around UE5.8)
+	PRAGMA_DISABLE_DEPRECATION_WARNINGS
+	inline FProcessingContext::FProcessingContext(FMassEntityManager& InEntityManager, const float InDeltaSeconds)
+		: FProcessingContext(/*MoveTemp*/(InEntityManager.AsShared()), InDeltaSeconds)
+	{
+
+	}
+
+	inline FProcessingContext::FProcessingContext(const TSharedRef<FMassEntityManager>& InEntityManager, const float InDeltaSeconds)
+		: EntityManager(InEntityManager)
+		, DeltaSeconds(InDeltaSeconds)
+	{
+		
+	}
+
+	inline FProcessingContext::FProcessingContext(TSharedRef<FMassEntityManager>&& InEntityManager, const float InDeltaSeconds)
+		: EntityManager(MoveTemp(InEntityManager))
+		, DeltaSeconds(InDeltaSeconds)
+	{
+		
+	}
+
+	inline FProcessingContext::FProcessingContext(const TSharedPtr<FMassEntityManager>& InEntityManager, const float InDeltaSeconds)
+		: FProcessingContext(InEntityManager.ToSharedRef(), InDeltaSeconds)
+	{
+		
+	}
+
+	inline FMassExecutionContext& FProcessingContext::GetExecutionContext() &
+	{
+		if (ExecutionContextPtr == nullptr)
+		{
+			ExecutionContextPtr = new(&ExecutionContextBuffer) FMassExecutionContext(*EntityManager, DeltaSeconds);
+
+			if (CommandBuffer.IsValid() == false)
+			{
+				CommandBuffer = MakeShareable(new FMassCommandBuffer());
+			}
+
+			ExecutionContextPtr->SetDeferredCommandBuffer(CommandBuffer);
+			
+			ExecutionContextPtr->SetFlushDeferredCommands(false);
+			ExecutionContextPtr->SetAuxData(AuxData);
+			ExecutionContextPtr->SetExecutionType(EMassExecutionContextType::Processor);
+		}
+		return *ExecutionContextPtr;
+	}
+
+	inline FMassExecutionContext&& FProcessingContext::GetExecutionContext() &&
+	{
+		// Note: it's fine to store a reference to created execution context
+		// while nulling-out the pointer, since the FMassExecutionContext data
+		// lives in the ExecutionContextBuffer buffer. Nulling out ExecutionContextPtr
+		// only signals that the execution context has been moved out.
+
+		FMassExecutionContext& LocalExecutionContext = GetExecutionContext();
+		ExecutionContextPtr = nullptr;
+		CommandBuffer.Reset();
+
+		return MoveTemp(LocalExecutionContext);
+	}
+
+	inline void FProcessingContext::SetCommandBuffer(TSharedPtr<FMassCommandBuffer>&& InCommandBuffer)
+	{
+		checkf(ExecutionContextPtr == nullptr, TEXT("Setting command buffer after ExecutionContext creation is not supported"));
+		if (ExecutionContextPtr == nullptr)
+		{
+			CommandBuffer = MoveTemp(InCommandBuffer);
+		}
+	}
+
+	inline void FProcessingContext::SetCommandBuffer(const TSharedPtr<FMassCommandBuffer>& InCommandBuffer)
+	{
+		checkf(ExecutionContextPtr == nullptr, TEXT("Setting command buffer after ExecutionContext creation is not supported"));
+		if (ExecutionContextPtr == nullptr)
+		{
+			CommandBuffer = InCommandBuffer;
+		}
+	}
+
+	inline const TSharedRef<FMassEntityManager>& FProcessingContext::GetEntityManager() const
+	{
+		return EntityManager;
+	}
+
+	PRAGMA_ENABLE_DEPRECATION_WARNINGS
+} // namespace UE::Mass
diff --git a/Engine/Source/Runtime/MassEntity/Public/MassProcessingTypes.h b/Engine/Source/Runtime/MassEntity/Public/MassProcessingTypes.h
--- a/Engine/Source/Runtime/MassEntity/Public/MassProcessingTypes.h
+++ b/Engine/Source/Runtime/MassEntity/Public/MassProcessingTypes.h
@@ -39,34 +39,27 @@ struct FProcessorAuxDataBase
 	GENERATED_BODY()
 };
 
-USTRUCT(BlueprintType)
-struct MASSENTITY_API FMassProcessingContext
+USTRUCT(BlueprintType, meta = (Deprecated = "5.6"))
+struct MASSENTITY_API FMassProcessingContext_DEPRECATED
 {
 	GENERATED_BODY()
 
-	TSharedPtr<FMassEntityManager> EntityManager;
-	
 	UPROPERTY()
 	float DeltaSeconds = 0.f;
 
 	UPROPERTY()
 	FInstancedStruct AuxData;
 
-	/** 
-	 * If set to "true" the MassExecutor will flush commands at the end of given execution function. 
-	 * If "false" the caller is responsible for manually flushing the commands.
-	 */
 	UPROPERTY()
 	bool bFlushCommandBuffer = true; 
-		
-	TSharedPtr<FMassCommandBuffer> CommandBuffer;
-	
-	FMassProcessingContext() = default;
-	FMassProcessingContext(FMassEntityManager& InEntityManager, const float InDeltaSeconds);
-	FMassProcessingContext(TSharedPtr<FMassEntityManager>& InEntityManager, const float InDeltaSeconds);
-	~FMassProcessingContext();
 };
 
+namespace UE::Mass
+{
+	struct FProcessingContext;
+}
+using FMassProcessingContext = UE::Mass::FProcessingContext;
+
 /** 
  *  Runtime-usable array of MassProcessor copies
  */
