[Mass] added a FMassArchetypeEntityCollection constructor optimized for a single-entity case

diff --git a/Engine/Source/Runtime/MassEntity/Private/MassArchetypeTypes.cpp b/Engine/Source/Runtime/MassEntity/Private/MassArchetypeTypes.cpp
--- a/Engine/Source/Runtime/MassEntity/Private/MassArchetypeTypes.cpp
+++ b/Engine/Source/Runtime/MassEntity/Private/MassArchetypeTypes.cpp
@@ -38,6 +38,24 @@ bool FMassArchetypeVersionedHandle::IsUpToDate() const
 //-----------------------------------------------------------------------------
 // FMassArchetypeEntityCollection 
 //-----------------------------------------------------------------------------
+FORCEINLINE FMassArchetypeEntityCollection::FArchetypeEntityRange FMassArchetypeEntityCollection::CreateRangeForEntity(const FMassArchetypeHandle& InArchetype, const FMassEntityHandle EntityHandle)
+{
+	if (const FMassArchetypeData* ArchetypeData = FMassArchetypeHelper::ArchetypeDataFromHandle(InArchetype))
+	{
+		const int32* TrueIndexPtr = ArchetypeData->GetInternalIndexForEntity(EntityHandle.Index);
+		check(TrueIndexPtr);
+		const int32 TrueIndex = *TrueIndexPtr;
+		checkf(TrueIndex >= 0, TEXT("Negative values are unexpected here"));
+
+		const int32 NumEntitiesPerChunk = ArchetypeData->GetNumEntitiesPerChunk();
+		const int32 ChunkIndex = TrueIndex / NumEntitiesPerChunk;
+		const int32 SubchunkStart = TrueIndex % NumEntitiesPerChunk;
+		return FArchetypeEntityRange(ChunkIndex, SubchunkStart, 1);
+	}
+
+	return FArchetypeEntityRange(0, EntityHandle.Index, 1);
+}
+
 FMassArchetypeEntityCollection::FMassArchetypeEntityCollection(const FMassArchetypeHandle& InArchetype, TConstArrayView<FMassEntityHandle> InEntities, EDuplicatesHandling DuplicatesHandling)
 	: Archetype(InArchetype)
 {
@@ -45,6 +63,11 @@ FMassArchetypeEntityCollection::FMassArchetypeEntityCollection(const FMassArchet
 	{
 		return;
 	}
+	if (InEntities.Num() == 1)
+	{
+		Ranges.Add(CreateRangeForEntity(InArchetype, InEntities[0]));
+		return;
+	}
 
 	const FMassArchetypeData* ArchetypeData = FMassArchetypeHelper::ArchetypeDataFromHandle(InArchetype);
 
@@ -123,6 +146,12 @@ FMassArchetypeEntityCollection::FMassArchetypeEntityCollection(const FMassArchet
 	BuildEntityRanges(MakeStridedView<const int32>(TrueIndices));
 }
 
+FMassArchetypeEntityCollection::FMassArchetypeEntityCollection(const FMassArchetypeHandle& InArchetype, const FMassEntityHandle EntityHandle)
+	: Ranges({CreateRangeForEntity(InArchetype, EntityHandle)})
+	, Archetype(InArchetype)
+{	
+}
+
 void FMassArchetypeEntityCollection::BuildEntityRanges(TStridedView<const int32> TrueIndices)
 {
 	checkf(Ranges.Num() == 0, TEXT("Calling %s is valid only for initial configuration"), ANSI_TO_TCHAR(__FUNCTION__));
diff --git a/Engine/Source/Runtime/MassEntity/Public/MassArchetypeTypes.h b/Engine/Source/Runtime/MassEntity/Public/MassArchetypeTypes.h
--- a/Engine/Source/Runtime/MassEntity/Public/MassArchetypeTypes.h
+++ b/Engine/Source/Runtime/MassEntity/Public/MassArchetypeTypes.h
@@ -164,6 +164,8 @@ private:
 public:
 	FMassArchetypeEntityCollection() = default;
 	FMassArchetypeEntityCollection(const FMassArchetypeHandle& InArchetype, TConstArrayView<FMassEntityHandle> InEntities, EDuplicatesHandling DuplicatesHandling);
+	/** optimized, special case for a single-entity */
+	FMassArchetypeEntityCollection(const FMassArchetypeHandle& InArchetype, const FMassEntityHandle EntityHandle);
 	explicit FMassArchetypeEntityCollection(const FMassArchetypeHandle& InArchetypeHandle, const EInitializationType Initialization = EInitializationType::GatherAll);
 	explicit FMassArchetypeEntityCollection(TSharedPtr<FMassArchetypeData>& InArchetype, const EInitializationType Initialization = EInitializationType::GatherAll);
 	FMassArchetypeEntityCollection(const FMassArchetypeHandle& InArchetypeHandle, FEntityRangeArray&& InEntityRanges)
@@ -216,7 +218,8 @@ public:
 protected:
 	friend struct FMassArchetypeEntityCollectionWithPayload;
 	void BuildEntityRanges(TStridedView<const int32> TrueIndices);
-	
+	static FArchetypeEntityRange CreateRangeForEntity(const FMassArchetypeHandle& InArchetype, const FMassEntityHandle EntityHandle);
+
 private:
 	void GatherChunksFromArchetype();
 };
