[Mass] Changed how `WITH_MASSENTITY_DEBUG` is defined so that it's possible to disable it from *.Build.cs files
Useful for perf-measuring sessions

diff --git a/Engine/Source/Runtime/MassEntity/Public/MassProcessingTypes.h b/Engine/Source/Runtime/MassEntity/Public/MassProcessingTypes.h
--- a/Engine/Source/Runtime/MassEntity/Public/MassProcessingTypes.h
+++ b/Engine/Source/Runtime/MassEntity/Public/MassProcessingTypes.h
@@ -11,7 +11,9 @@
 #define MASS_DO_PARALLEL !UE_SERVER
 #endif // MASS_DO_PARALLEL
 
+#ifndef WITH_MASSENTITY_DEBUG
 #define WITH_MASSENTITY_DEBUG (!(UE_BUILD_SHIPPING || UE_BUILD_SHIPPING_WITH_EDITOR || UE_BUILD_TEST) && WITH_STRUCTUTILS_DEBUG && 1)
+#endif
 
 struct FMassEntityManager;
 class UMassProcessor;
