[Mass] Introduced ExecutionPriority for MassProcessors that is used to position the given processor (and its dependencies) earlier in the processing graph.
This includes a cvar mass.dependencies.ProcessorExecutionPriorityEnabled which can be used to disable the functionality.
Also, introduced LogMassDependencies as a dedicated category used by FMassProcessorDependencySolver to make debugging dependency calculations easier.

diff --git a/Engine/Source/Runtime/MassEntity/Private/MassProcessorDependencySolver.cpp b/Engine/Source/Runtime/MassEntity/Private/MassProcessorDependencySolver.cpp
--- a/Engine/Source/Runtime/MassEntity/Private/MassProcessorDependencySolver.cpp
+++ b/Engine/Source/Runtime/MassEntity/Private/MassProcessorDependencySolver.cpp
@@ -13,6 +13,8 @@
 
 #define LOCTEXT_NAMESPACE "Mass"
 
+DEFINE_LOG_CATEGORY_STATIC(LogMassDependencies, Warning, All);
+
 namespace UE::Mass::Private
 {
 	FString NameViewToString(TConstArrayView<FName> View)
@@ -85,13 +87,13 @@ namespace UE::Mass::Private
 		{
 			InOutReportedCycleHashes.Add(CycleHash);
 
-			UE_LOG(LogMass, Error, TEXT("Detected processing dependency cycle:"));
+			UE_LOG(LogMassDependencies, Error, TEXT("Detected processing dependency cycle:"));
 
 			for (const int32 CycleNodeIndex : CycleView)
 			{
 				if (const UMassProcessor* Processor = AllNodes[CycleNodeIndex].Processor)
 				{
-					UE_LOG(LogMass, Warning, TEXT("\t%s, group: %s, before: %s, after %s")
+					UE_LOG(LogMassDependencies, Warning, TEXT("\t%s, group: %s, before: %s, after %s")
 						, *Processor->GetName()
 						, *Processor->GetExecutionOrder().ExecuteInGroup.ToString()
 						, *NameViewToString(Processor->GetExecutionOrder().ExecuteBefore)
@@ -100,12 +102,31 @@ namespace UE::Mass::Private
 				else
 				{
 					// group
-					UE_LOG(LogMass, Warning, TEXT("\tGroup %s"), *AllNodes[CycleNodeIndex].Name.ToString());
+					UE_LOG(LogMassDependencies, Warning, TEXT("\tGroup %s"), *AllNodes[CycleNodeIndex].Name.ToString());
 				}
 			}
 		}
 	}
 #endif // WITH_MASSENTITY_DEBUG
+
+	bool bProcessorExecutionPriorityEnabled = true;
+	bool bPickHigherPriorityNodesRegardlessOfRequirements = true;
+	namespace
+	{
+		FAutoConsoleVariableRef ConsoleVariables[] =
+		{
+			FAutoConsoleVariableRef(
+				TEXT("mass.dependencies.ProcessorExecutionPriorityEnabled"),
+				bProcessorExecutionPriorityEnabled,
+				TEXT("Controls whether UMassProcessor.ExecutionPriority value is being used during dependency calculations"),
+				ECVF_Default)
+			, FAutoConsoleVariableRef(
+				TEXT("mass.dependencies.PickHigherPriorityNodesRegardlessOfRequirements"),
+				bPickHigherPriorityNodesRegardlessOfRequirements,
+				TEXT("If enabled, will result in lower priority nodes not being picked, even if they could run without obstructing anything else"),
+				ECVF_Default)
+		};
+	}
 }
 
 //----------------------------------------------------------------------//
@@ -354,6 +375,32 @@ bool FMassProcessorDependencySolver::FNode::IncreaseWaitingNodesCount(TArrayView
 
 	return true;
 }
+
+bool FMassProcessorDependencySolver::FNode::IncreaseWaitingNodesCountAndPriority(TArrayView<FNode> InAllNodes, const int32 IterationsLimit
+	, TArray<int32>& OutCycleIndices, const int32 InChildPriority)
+{
+	// cycle-protection check. If true it means we have a cycle and the whole algorithm result will be unreliable 
+	if (IterationsLimit < 0)
+	{
+		OutCycleIndices.Add(NodeIndex);
+		return false;
+	}
+
+	++TotalWaitingNodes;
+	UpdateExecutionPriority(InChildPriority);
+
+	for (const int32 DependencyIndex : OriginalDependencies)
+	{
+		check(&InAllNodes[DependencyIndex] != this);
+		if (InAllNodes[DependencyIndex].IncreaseWaitingNodesCountAndPriority(InAllNodes, IterationsLimit - 1, OutCycleIndices, MaxExecutionPriority) == false)
+		{
+			OutCycleIndices.Add(NodeIndex);
+			return false;
+		}
+	}
+
+	return true;
+}
 //----------------------------------------------------------------------//
 //  FProcessorDependencySolver
 //----------------------------------------------------------------------//
@@ -367,6 +414,7 @@ bool FMassProcessorDependencySolver::PerformSolverStep(FResourceUsage& ResourceU
 	int32 AcceptedNodeIndex = INDEX_NONE;
 	int32 FallbackAcceptedNodeIndex = INDEX_NONE;
 
+	const int32 HighestPriority = AllNodes[InOutIndicesRemaining[0]].MaxExecutionPriority;
 	for (int32 i = 0; i < InOutIndicesRemaining.Num(); ++i)
 	{
 		const int32 NodeIndex = InOutIndicesRemaining[i];
@@ -384,6 +432,15 @@ bool FMassProcessorDependencySolver::PerformSolverStep(FResourceUsage& ResourceU
 				// we'll just pick this one up and go with it. 
 				FallbackAcceptedNodeIndex = NodeIndex;
 			}
+			else if (UE::Mass::Private::bPickHigherPriorityNodesRegardlessOfRequirements 
+				&& AllNodes[NodeIndex].MaxExecutionPriority < HighestPriority)
+			{
+				// subsequent nodes are of lower execution priority, we break now and will use FallbackAcceptedNodeIndex
+				check(FallbackAcceptedNodeIndex != INDEX_NONE);
+				checkf(UE::Mass::Private::bProcessorExecutionPriorityEnabled == true
+					, TEXT("We never expect to hit this case when execution priorities are disabled - all nodes should have the same priority."))
+				break;
+			}
 		}
 	}
 
@@ -398,7 +455,7 @@ bool FMassProcessorDependencySolver::PerformSolverStep(FResourceUsage& ResourceU
 		// But that can't last forever, at some point we'll end up in a situation where every node left waits for 
 		// something that has been submitted already. Then we just pick one of the waiting ones (the one indicated by 
 		// FallbackAcceptedNodeIndex), "run it" and proceed.
-		UE_CLOG(AcceptedNodeIndex == INDEX_NONE, LogMass, Verbose, TEXT("No dependency-free node can be picked, due to resource requirements. Picking %s as the next node.")
+		UE_CLOG(AcceptedNodeIndex == INDEX_NONE, LogMassDependencies, Verbose, TEXT("No dependency-free node can be picked, due to resource requirements. Picking %s as the next node.")
 			, *Node.Name.ToString());
 
 		ResourceUsage.SubmitNode(NodeIndex, Node);
@@ -455,12 +512,12 @@ int32 FMassProcessorDependencySolver::CreateNodes(UMassProcessor& Processor)
 	{
 		if (Processor.ShouldAllowMultipleInstances())
 		{
-			UE_LOG(LogMass, Warning, TEXT("%hs Processor %s, name %s, already registered. This processor class does suport duplicates, but individual instances need to have a unique name.")
+			UE_LOG(LogMassDependencies, Warning, TEXT("%hs Processor %s, name %s, already registered. This processor class does suport duplicates, but individual instances need to have a unique name.")
 				, __FUNCTION__, *Processor.GetFullName(), *ProcName.ToString());
 		}
 		else
 		{
-			UE_LOG(LogMass, Warning, TEXT("%hs Processor %s already registered. Duplicates are not supported by this processor class.")
+			UE_LOG(LogMassDependencies, Warning, TEXT("%hs Processor %s already registered. Duplicates are not supported by this processor class.")
 				, __FUNCTION__, *ProcName.ToString());
 		}
 		return *NodeIndexPtr;
@@ -517,6 +574,9 @@ int32 FMassProcessorDependencySolver::CreateNodes(UMassProcessor& Processor)
 	ProcessorNode.Requirements.RequiredSubsystems.Write -= MultiThreadedSystemsBitSet;
 	ProcessorNode.Requirements.RequiredSubsystems.Read -= MultiThreadedSystemsBitSet;
 	ProcessorNode.Requirements.CountResourcesUsed();
+	ProcessorNode.MaxExecutionPriority = UE::Mass::Private::bProcessorExecutionPriorityEnabled
+		? Processor.GetExecutionPriority()
+		: 0;
 
 	if (ParentGroupNodeIndex > 0)
 	{
@@ -547,7 +607,7 @@ void FMassProcessorDependencySolver::BuildDependencies()
 				NodeIndexMap.Add(BeforeDependencyName, DependentNodeIndex);
 				AllNodes.Add({ BeforeDependencyName, nullptr, DependentNodeIndex });
 
-				UE_LOG(LogMass, Log, TEXT("Unable to find dependency \"%s\" declared by %s. Creating a dummy dependency node.")
+				UE_LOG(LogMassDependencies, Log, TEXT("Unable to find dependency \"%s\" declared by %s. Creating a dummy dependency node.")
 					, *BeforeDependencyName.ToString(), *AllNodes[NodeIndex].Name.ToString());
 			}
 			else
@@ -588,7 +648,7 @@ void FMassProcessorDependencySolver::BuildDependencies()
 				NodeIndexMap.Add(AfterDependencyName, PrerequisiteNodeIndex);
 				AllNodes.Add({ AfterDependencyName, nullptr, PrerequisiteNodeIndex });
 
-				UE_LOG(LogMass, Log, TEXT("Unable to find dependency \"%s\" declared by %s. Creating a dummy dependency node.")
+				UE_LOG(LogMassDependencies, Log, TEXT("Unable to find dependency \"%s\" declared by %s. Creating a dummy dependency node.")
 					, *AfterDependencyName.ToString(), *AllNodes[NodeIndex].Name.ToString());
 			}
 			else
@@ -646,7 +706,7 @@ void FMassProcessorDependencySolver::LogNode(const FNode& Node, int Indent)
 
 	if (Node.IsGroup())
 	{
-		UE_LOG(LogMass, Log, TEXT("%*s%s before:%s after:%s"), Indent, TEXT(""), *Node.Name.ToString()
+		UE_LOG(LogMassDependencies, Log, TEXT("%*s%s before:%s after:%s"), Indent, TEXT(""), *Node.Name.ToString()
 			, *NameViewToString(Node.ExecuteBefore)
 			, *NameViewToString(Node.ExecuteAfter));
 
@@ -658,7 +718,7 @@ void FMassProcessorDependencySolver::LogNode(const FNode& Node, int Indent)
 	else
 	{
 		CA_ASSUME(Node.Processor); // as implied by Node.IsGroup() == false
-		UE_LOG(LogMass, Log, TEXT("%*s%s before:%s after:%s"), Indent, TEXT(""), *Node.Name.ToString()
+		UE_LOG(LogMassDependencies, Log, TEXT("%*s%s before:%s after:%s"), Indent, TEXT(""), *Node.Name.ToString()
 			, *NameViewToString(Node.Processor->GetExecutionOrder().ExecuteBefore)
 			, *NameViewToString(Node.Processor->GetExecutionOrder().ExecuteAfter));
 	}
@@ -690,27 +750,58 @@ void FMassProcessorDependencySolver::Solve(TArray<FMassProcessorOrderInfo>& OutR
 #endif // WITH_MASSENTITY_DEBUG
 
 	TArray<int32> IndicesRemaining;
-	IndicesRemaining.Reserve(AllNodes.Num());
-	for (int32 NodeIndex = 0; NodeIndex < AllNodes.Num(); ++NodeIndex)
+	// @todo code duplication in this if-else block is temporary, will be reduced to one or the other
+	// once the bProcessorExecutionPriorityEnabled feature is accepted or removed. 
+	if (UE::Mass::Private::bProcessorExecutionPriorityEnabled)
 	{
-		// skip all the group nodes, all group dependencies have already been converted to individual processor dependencies
-		if (AllNodes[NodeIndex].IsGroup() == false)
+		IndicesRemaining.Reserve(AllNodes.Num());
+		for (int32 NodeIndex = 0; NodeIndex < AllNodes.Num(); ++NodeIndex)
 		{
-			IndicesRemaining.Add(NodeIndex);
-			if (AllNodes[NodeIndex].IncreaseWaitingNodesCount(AllNodes, TotalDependingNodes, CycleIndices) == false)
+			// skip all the group nodes, all group dependencies have already been converted to individual processor dependencies
+			if (AllNodes[NodeIndex].IsGroup() == false)
 			{
+				IndicesRemaining.Add(NodeIndex);
+				if (AllNodes[NodeIndex].IncreaseWaitingNodesCountAndPriority(AllNodes, TotalDependingNodes, CycleIndices) == false)
+				{
 #if WITH_MASSENTITY_DEBUG
-				// we have a cycle. Report it here
-				UE::Mass::Private::LogCycle(AllNodes, CycleIndices, ReportedCycleHashes);
+					// we have a cycle. Report it here
+					UE::Mass::Private::LogCycle(AllNodes, CycleIndices, ReportedCycleHashes);
 #endif // WITH_MASSENTITY_DEBUG
-				CycleIndices.Reset();
+					CycleIndices.Reset();
+				}
 			}
 		}
+
+		IndicesRemaining.Sort([this](const int32 IndexA, const int32 IndexB){
+			return AllNodes[IndexA].MaxExecutionPriority > AllNodes[IndexB].MaxExecutionPriority
+				|| (AllNodes[IndexA].MaxExecutionPriority == AllNodes[IndexB].MaxExecutionPriority
+					&& AllNodes[IndexA].TotalWaitingNodes > AllNodes[IndexB].TotalWaitingNodes);
+		});
 	}
+	else
+	{
+		IndicesRemaining.Reserve(AllNodes.Num());
+		for (int32 NodeIndex = 0; NodeIndex < AllNodes.Num(); ++NodeIndex)
+		{
+			// skip all the group nodes, all group dependencies have already been converted to individual processor dependencies
+			if (AllNodes[NodeIndex].IsGroup() == false)
+			{
+				IndicesRemaining.Add(NodeIndex);
+				if (AllNodes[NodeIndex].IncreaseWaitingNodesCount(AllNodes, TotalDependingNodes, CycleIndices) == false)
+				{
+#if WITH_MASSENTITY_DEBUG
+					// we have a cycle. Report it here
+					UE::Mass::Private::LogCycle(AllNodes, CycleIndices, ReportedCycleHashes);
+#endif // WITH_MASSENTITY_DEBUG
+					CycleIndices.Reset();
+				}
+			}
+		}
 
-	IndicesRemaining.Sort([this](const int32 IndexA, const int32 IndexB){
-		return AllNodes[IndexA].TotalWaitingNodes > AllNodes[IndexB].TotalWaitingNodes;
-	});
+		IndicesRemaining.Sort([this](const int32 IndexA, const int32 IndexB){
+			return AllNodes[IndexA].TotalWaitingNodes > AllNodes[IndexB].TotalWaitingNodes;
+		});
+	}
 
 	// this is where we'll be tracking what's being accessed by whom
 	FResourceUsage ResourceUsage(AllNodes);
@@ -724,7 +815,7 @@ void FMassProcessorDependencySolver::Solve(TArray<FMassProcessorOrderInfo>& OutR
 
 		if (bStepSuccessful == false)
 		{
-			UE_LOG(LogMass, Error, TEXT("Encountered processing dependency cycle - cutting the chain at an arbitrary location."));
+			UE_LOG(LogMassDependencies, Error, TEXT("Encountered processing dependency cycle - cutting the chain at an arbitrary location."));
 
 			// remove first dependency
 			// note that if we're in a cycle handling scenario every node does have some dependencies left
@@ -770,7 +861,7 @@ void FMassProcessorDependencySolver::ResolveDependencies(TArray<FMassProcessorOr
 		DependencyGraphFileName = InOutOptionalResult->DependencyGraphFileName;
 	}
 
-	UE_LOG(LogMass, Log, TEXT("Gathering dependencies data:"));
+	UE_LOG(LogMassDependencies, Log, TEXT("Gathering dependencies data:"));
 
 	AllNodes.Reset();
 	NodeIndexMap.Reset();
@@ -798,7 +889,7 @@ void FMassProcessorDependencySolver::ResolveDependencies(TArray<FMassProcessorOr
 	{
 		if (Processor == nullptr)
 		{
-			UE_LOG(LogMass, Warning, TEXT("%s nullptr found in Processors collection being processed"), ANSI_TO_TCHAR(__FUNCTION__));
+			UE_LOG(LogMassDependencies, Warning, TEXT("%s nullptr found in Processors collection being processed"), ANSI_TO_TCHAR(__FUNCTION__));
 			continue;
 		}
 
@@ -812,7 +903,7 @@ void FMassProcessorDependencySolver::ResolveDependencies(TArray<FMassProcessorOr
 		}
 	}
 
-	UE_LOG(LogMass, Verbose, TEXT("Pruning processors..."));
+	UE_LOG(LogMassDependencies, Verbose, TEXT("Pruning processors..."));
 
 	int32 PrunedProcessorsCount = 0;
 	for (FNode& Node : AllNodes)
@@ -839,7 +930,7 @@ void FMassProcessorDependencySolver::ResolveDependencies(TArray<FMassProcessorOr
 			// prune the archetype-less processors
 			if (Node.ValidArchetypes.Num() == 0 && bDoQueryBasedPruning)
 			{
-				UE_LOG(LogMass, Verbose, TEXT("\t%s"), *Node.Processor->GetName());
+				UE_LOG(LogMassDependencies, Verbose, TEXT("\t%s"), *Node.Processor->GetName());
 
 				if (InOutOptionalResult)
 				{
@@ -854,7 +945,7 @@ void FMassProcessorDependencySolver::ResolveDependencies(TArray<FMassProcessorOr
 		}
 	}
 
-	UE_LOG(LogMass, Verbose, TEXT("Number of processors pruned: %d"), PrunedProcessorsCount);
+	UE_LOG(LogMassDependencies, Verbose, TEXT("Number of processors pruned: %d"), PrunedProcessorsCount);
 
 	check(AllNodes.Num());
 	LogNode(AllNodes[0]);
@@ -867,10 +958,10 @@ void FMassProcessorDependencySolver::ResolveDependencies(TArray<FMassProcessorOr
 
 	Solve(OutResult);
 
-	UE_LOG(LogMass, Verbose, TEXT("Dependency order:"));
+	UE_LOG(LogMassDependencies, Verbose, TEXT("Dependency order:"));
 	for (const FMassProcessorOrderInfo& Info : OutResult)
 	{
-		UE_LOG(LogMass, Verbose, TEXT("\t%s"), *Info.Name.ToString());
+		UE_LOG(LogMassDependencies, Verbose, TEXT("\t%s"), *Info.Name.ToString());
 	}
 
 	int32 MaxSequenceLength = 0;
@@ -879,7 +970,7 @@ void FMassProcessorDependencySolver::ResolveDependencies(TArray<FMassProcessorOr
 		MaxSequenceLength = FMath::Max(MaxSequenceLength, Node.SequencePositionIndex);
 	}
 
-	UE_LOG(LogMass, Verbose, TEXT("Max sequence length: %d"), MaxSequenceLength);
+	UE_LOG(LogMassDependencies, Verbose, TEXT("Max sequence length: %d"), MaxSequenceLength);
 
 	if (InOutOptionalResult)
 	{
diff --git a/Engine/Source/Runtime/MassEntity/Public/MassProcessor.h b/Engine/Source/Runtime/MassEntity/Public/MassProcessor.h
--- a/Engine/Source/Runtime/MassEntity/Public/MassProcessor.h
+++ b/Engine/Source/Runtime/MassEntity/Public/MassProcessor.h
@@ -121,6 +121,21 @@ public:
 
 	const FMassSubsystemRequirements& GetProcessorRequirements() const;
 
+	/**
+	 * @return the current value of ExecutionPriority
+	 * @see SetExecutionPriority
+	 * @see ExecutionPriority
+	 */
+	int16 GetExecutionPriority() const;
+
+	/**
+	 * Sets new ExecutionPriority for this processor. The change will take effect the next time the processing graph is built
+	 * Note that at this point this operation does not cause processing graph rebuilding, so this function should be used
+	 * before processor's initialization or as part of code that will cause processing graph rebuilding anyway.
+	 * @see ExecutionPriority
+	 */
+	void SetExecutionPriority(const int16 NewExecutionPriority);
+
 	/** Adds Query to RegisteredQueries list. Query is required to be a member variable of this processor. Not meeting
 	 *  this requirement will cause check failure and the query won't be registered. */
 	void RegisterQuery(FMassEntityQuery& Query);
@@ -239,6 +254,15 @@ private:
 	EActivationState ActivationState = EActivationState::Active;
 
 protected:
+	/**
+	 * Denoted how important it is for this processor to be executed as soon as possible within a processing graph.
+	 * The larger the number the higher the priority. It's used in two ways:
+	 * - used when sorting nodes that otherwise seem similar in terms of "which processor to pick for execution next"
+	 * - affects the priority of the dependencies - if this super-important processor is waiting for processor A and B,
+	 *		then A and B become super important as well.
+	 */
+	UPROPERTY(EditDefaultsOnly, Category = Processor, config)
+	int16 ExecutionPriority = 0;
 
 	/**
 	 * Determines whether given processor wants to be pruned from the execution graph when there are
@@ -440,6 +464,16 @@ inline const FMassSubsystemRequirements& UMassProcessor::GetProcessorRequirement
 	return ProcessorRequirements;
 }
 
+inline int16 UMassProcessor::GetExecutionPriority() const
+{
+	return ExecutionPriority; 
+}
+
+inline void UMassProcessor::SetExecutionPriority(const int16 NewExecutionPriority)
+{
+	ExecutionPriority = NewExecutionPriority; 
+}
+
 inline void UMassProcessor::MarkAsDynamic()
 {
 	bIsDynamic = true;
diff --git a/Engine/Source/Runtime/MassEntity/Public/MassProcessorDependencySolver.h b/Engine/Source/Runtime/MassEntity/Public/MassProcessorDependencySolver.h
--- a/Engine/Source/Runtime/MassEntity/Public/MassProcessorDependencySolver.h
+++ b/Engine/Source/Runtime/MassEntity/Public/MassProcessorDependencySolver.h
@@ -95,6 +95,16 @@ struct FMassProcessorDependencySolver
 		 *	happens `OutCycleIndices` gets filled with the relevant node indices.
 		 */
 		bool IncreaseWaitingNodesCount(TArrayView<FNode> InAllNodes, const int32 IterationsLimit, TArray<int32>& OutCycleIndices);
+		bool IncreaseWaitingNodesCountAndPriority(TArrayView<FNode> InAllNodes, const int32 IterationsLimit, TArray<int32>& OutCycleIndices
+			, const int32 InChildPriority = TNumericLimits<int32>::Min());
+		void UpdateExecutionPriority(const int32 ChildExecutionPriority)
+		{
+			// picking the max execution priority - note that we're increasing child
+			// priority to ensure dependencies  always have a higher stored priority
+			// than the nodes that depend on them
+			MaxExecutionPriority = FMath::Max(int32(ChildExecutionPriority) + 1, MaxExecutionPriority);
+		}
+
 
 		FName Name = TEXT("");
 		UMassProcessor* Processor = nullptr;
@@ -106,6 +116,14 @@ struct FMassProcessorDependencySolver
 		int32 NodeIndex = INDEX_NONE;
 		/** indicates how often given node can be found in dependencies sequence for other nodes  */
 		int32 TotalWaitingNodes = 0;
+		/**
+		 * Indicates the maximum execution priority represented by this node or any od the nodes
+		 * that depend on it in a logical sense - i.e. it does not include the nodes that are dependencies just by blocing required resources
+		 * @todo reword this comment
+		 * Note that we're using a larger type than UMassProcessor.ExecutionPriority (int32 vs int16)
+		 * to not have to handle overflow in UpdateExecutionPriority
+		 */
+		int32 MaxExecutionPriority = 0;
 		/** 
 		 * indicates how deep within dependencies graph this give node is, or in other words, what's the longest sequence 
 		 * from this node to a dependency-less "parent" node 
