[MassEntity] Minor EntityCollection extension, supporting destructive archetype collections fetching with move semantics

diff --git a/Engine/Source/Runtime/MassEntity/Private/MassEntityCollection.cpp b/Engine/Source/Runtime/MassEntity/Private/MassEntityCollection.cpp
--- a/Engine/Source/Runtime/MassEntity/Private/MassEntityCollection.cpp
+++ b/Engine/Source/Runtime/MassEntity/Private/MassEntityCollection.cpp
@@ -18,18 +18,18 @@ namespace UE::Mass
 		AppendCollection(InEntityCollection);
 	}
 
-	FEntityCollection::FEntityCollection(const TConstArrayView<FMassEntityHandle> InCreatedEntities)
-		: EntityHandles(InCreatedEntities)
+	FEntityCollection::FEntityCollection(const TConstArrayView<FMassEntityHandle> InEntityHandles)
+		: EntityHandles(InEntityHandles)
 	{	
 	}
 
-	FEntityCollection::FEntityCollection(const TConstArrayView<FMassEntityHandle> InCreatedEntities, FMassArchetypeEntityCollection&& InEntityCollection)
-		: EntityHandles(InCreatedEntities)
+	FEntityCollection::FEntityCollection(const TConstArrayView<FMassEntityHandle> InEntityHandles, FMassArchetypeEntityCollection&& InEntityCollection)
+		: EntityHandles(InEntityHandles)
 	{
 		CachedCollections.Add(Forward<FMassArchetypeEntityCollection>(InEntityCollection));
 	}
 
-	TConstArrayView<FMassArchetypeEntityCollection> FEntityCollection::GetUpToDatePerArchetypeCollections(const FMassEntityManager& EntityManager) const
+	void FEntityCollection::ConditionallyUpdate(const FMassEntityManager& EntityManager) const
 	{
 		if (IsUpToDate() == false)
 		{
@@ -38,8 +38,6 @@ namespace UE::Mass
 				, CollectionCreationDuplicatesHandling
 				, CachedCollections);
 		}
-
-		return CachedCollections;
 	}
 
 	void FEntityCollection::AppendHandles(TConstArrayView<FMassEntityHandle> Handles)
diff --git a/Engine/Source/Runtime/MassEntity/Public/MassEntityCollection.h b/Engine/Source/Runtime/MassEntity/Public/MassEntityCollection.h
--- a/Engine/Source/Runtime/MassEntity/Public/MassEntityCollection.h
+++ b/Engine/Source/Runtime/MassEntity/Public/MassEntityCollection.h
@@ -32,8 +32,8 @@ namespace UE::Mass
 		 */
 		MASSENTITY_API explicit FEntityCollection(FMassArchetypeEntityCollection&& InEntityCollection);
 		MASSENTITY_API explicit FEntityCollection(const FMassArchetypeEntityCollection& InEntityCollection);
-		MASSENTITY_API explicit FEntityCollection(const TConstArrayView<FMassEntityHandle> InCreatedEntities);
-		MASSENTITY_API FEntityCollection(const TConstArrayView<FMassEntityHandle> InCreatedEntities, FMassArchetypeEntityCollection&& InEntityCollection);
+		MASSENTITY_API explicit FEntityCollection(const TConstArrayView<FMassEntityHandle> InEntityHandles);
+		MASSENTITY_API FEntityCollection(const TConstArrayView<FMassEntityHandle> InEntityHandles, FMassArchetypeEntityCollection&& InEntityCollection);
 
 		//-----------------------------------------------------------------------------
 		// mutating API
@@ -137,7 +137,20 @@ namespace UE::Mass
 		 * Fetches up-to-date FMassArchetypeEntityCollection instances matching stored
 		 * entity handles.
 		 */
-		MASSENTITY_API TConstArrayView<FMassArchetypeEntityCollection> GetUpToDatePerArchetypeCollections(const FMassEntityManager& EntityManager) const;
+		MASSENTITY_API TConstArrayView<FMassArchetypeEntityCollection> GetUpToDatePerArchetypeCollections(const FMassEntityManager& EntityManager) const
+		{
+			ConditionallyUpdate(EntityManager);
+			return CachedCollections;
+		}
+
+		/**
+		 * Updates cached archetype collections and returns the container with move semantics
+		 */
+		MASSENTITY_API TArray<FMassArchetypeEntityCollection> ConsumeArchetypeCollections(const FMassEntityManager& EntityManager) &&
+		{
+			ConditionallyUpdate(EntityManager);
+			return MoveTemp(CachedCollections);
+		}
 
 	private:
 		template<typename T> requires std::is_same_v<typename TDecay<T>::Type, FMassArchetypeEntityCollection>
@@ -165,6 +178,8 @@ namespace UE::Mass
 			}
 		}
 
+		void ConditionallyUpdate(const FMassEntityManager& EntityManager) const;
+
 		/**
 		 * these are the entities represented by given instance of FEntityCollection.
 		 * EntityHandles are the authority, source of truth regarding the contents
