[Mass] Fixed: Issue 861258: Undefined symbols: UE::Mass::FEntityCollection::ConditionallyUpdate

#rnx

diff --git a/Engine/Source/Runtime/MassEntity/Public/MassEntityCollection.h b/Engine/Source/Runtime/MassEntity/Public/MassEntityCollection.h
--- a/Engine/Source/Runtime/MassEntity/Public/MassEntityCollection.h
+++ b/Engine/Source/Runtime/MassEntity/Public/MassEntityCollection.h
@@ -178,7 +178,7 @@ namespace UE::Mass
 			}
 		}
 
-		void ConditionallyUpdate(const FMassEntityManager& EntityManager) const;
+		MASSENTITY_API void ConditionallyUpdate(const FMassEntityManager& EntityManager) const;
 
 		/**
 		 * these are the entities represented by given instance of FEntityCollection.
