[MassEntity] handling calls to FlushCommands calls after manager deinitialization

diff --git a/Engine/Source/Runtime/MassEntity/Private/MassEntityManager.cpp b/Engine/Source/Runtime/MassEntity/Private/MassEntityManager.cpp
--- a/Engine/Source/Runtime/MassEntity/Private/MassEntityManager.cpp
+++ b/Engine/Source/Runtime/MassEntity/Private/MassEntityManager.cpp
@@ -162,7 +162,7 @@ FMassEntityManager::FMassEntityManager(UObject* InOwner)
 
 FMassEntityManager::~FMassEntityManager()
 {
-	if (bInitialized)
+	if (InitializationState == EInitializationState::Initialized)
 	{
 		Deinitialize();
 	}
@@ -170,7 +170,7 @@ FMassEntityManager::~FMassEntityManager()
 
 void FMassEntityManager::GetResourceSizeEx(FResourceSizeEx& CumulativeResourceSize)
 {
-	SIZE_T MyExtraSize = (bInitialized ? GetEntityStorageInterface().GetAllocatedSize() : 0)
+	SIZE_T MyExtraSize = (InitializationState == EInitializationState::Initialized ? GetEntityStorageInterface().GetAllocatedSize() : 0)
 		+ FragmentHashToArchetypeMap.GetAllocatedSize()
 		+ FragmentTypeToArchetypeMap.GetAllocatedSize();
 
@@ -192,7 +192,13 @@ void FMassEntityManager::GetResourceSizeEx(FResourceSizeEx& CumulativeResourceSi
 
 void FMassEntityManager::AddReferencedObjects(FReferenceCollector& Collector)
 {
-	if (bInitialized == false)
+	if (InitializationState == EInitializationState::Uninitialized)
+	{
+		UE_VLOG_UELOG(GetOwner(), LogMass, Log, TEXT("AddReferencedObjects called before Initialize call"));
+		return;
+	}
+
+	if (InitializationState == EInitializationState::Deinitialized)
 	{
 		// this is AddReferencedObjects called after Deinitialize call, which means we don't want to retain any object refs
 		// since this FMassEntityManager instance is going away even if it's kept alive by some stored shared refs at the moment.
@@ -246,9 +252,9 @@ namespace UE::Mass::Private
 
 void FMassEntityManager::Initialize(const FMassEntityManagerStorageInitParams& InitializationParams)
 {
-	if (bInitialized)
+	if (InitializationState == EInitializationState::Initialized)
 	{
-		UE_LOG(LogMass, Log, TEXT("Calling %hs on already initialized entity manager owned by %s")
+		UE_VLOG_UELOG(GetOwner(), LogMass, Log, TEXT("Calling %hs on already initialized entity manager owned by %s")
 			, __FUNCTION__, *GetNameSafe(Owner.Get()));
 		return;
 	}
@@ -321,21 +327,21 @@ void FMassEntityManager::Initialize(const FMassEntityManagerStorageInitParams& I
 	FMassDebugger::RegisterEntityManager(*this);
 #endif // WITH_MASSENTITY_DEBUG
 
-	bInitialized = true;
+	InitializationState = EInitializationState::Initialized;
 	bFirstCommandFlush = true;
 }
 
 void FMassEntityManager::PostInitialize()
 {
-	ensure(bInitialized);
-	// this needs to be done after all the subsystems have been initialized since some processors might want to access
-	// them during processors' initialization
+	ensureMsgf(InitializationState == EInitializationState::Initialized,
+		TEXT("This needs to be done after all the subsystems have been initialized since some processors might want to access"
+			" them during processors' initialization"));
 	ObserverManager.Initialize();
 }
 
 void FMassEntityManager::Deinitialize()
 {
-	if (bInitialized)
+	if (InitializationState == EInitializationState::Initialized)
 	{
 		FCoreDelegates::OnPostFork.Remove(OnPostForkHandle);
 
@@ -356,11 +362,11 @@ void FMassEntityManager::Deinitialize()
 
 		ObserverManager.DeInitialize();
 		
-		bInitialized = false;
+		InitializationState = EInitializationState::Deinitialized;
 	}
 	else
 	{
-		UE_LOG(LogMass, Log, TEXT("Calling %hs on already deinitialized entity manager owned by %s")
+		UE_VLOG_UELOG(GetOwner(), LogMass, Log, TEXT("Calling %hs on already deinitialized entity manager owned by %s")
 			, __FUNCTION__, *GetNameSafe(Owner.Get()));
 	}
 }
@@ -2057,11 +2063,28 @@ void FMassEntityManager::FlushCommands(const TSharedPtr<FMassCommandBuffer>& InC
 	{
 		return;
 	}
+
 	if (!ensureMsgf(IsProcessing() == false, TEXT("Calling %hs is not supported while Mass Processing is active. Call FMassEntityManager::AppendCommands instead."), __FUNCTION__))
 	{
 		return;
 	}
 
+	if (UNLIKELY(InitializationState == EInitializationState::Uninitialized))
+	{
+		UE_VLOG_UELOG(GetOwner(), LogMass, Warning, TEXT("FlushCommands called before Initialize call, which means this FMassEntityManager instance is not ready to ,"
+			" process commands and will cancel them."));
+		InCommandBuffer->CancelCommands();
+		return;
+	}
+	
+	if (UNLIKELY(InitializationState == EInitializationState::Deinitialized))
+	{
+		UE_VLOG_UELOG(GetOwner(), LogMass, Log, TEXT("FlushCommands called after Deinitialize call, which means this FMassEntityManager instance is going away,"
+			" can't process commands and will cancel them."));
+		InCommandBuffer->CancelCommands();
+		return;
+	}
+
 	if (InCommandBuffer && InCommandBuffer->HasPendingCommands()
 		&& (Algo::Find(DeferredCommandBuffers, InCommandBuffer) == nullptr))
 	{
diff --git a/Engine/Source/Runtime/MassEntity/Public/MassEntityManager.h b/Engine/Source/Runtime/MassEntity/Public/MassEntityManager.h
--- a/Engine/Source/Runtime/MassEntity/Public/MassEntityManager.h
+++ b/Engine/Source/Runtime/MassEntity/Public/MassEntityManager.h
@@ -1196,19 +1196,6 @@ private:
 	/** Never access directly, use GetOrMakeCreationContext instead. */
 	TWeakPtr<FEntityCreationContext> ActiveCreationContext;
 
-	std::atomic<bool> bCommandBufferFlushingInProgress = false;
-
-	/**
-	 * This index will be enough to control which buffer is available for pushing commands since flashing is taking place 
-	 * in the game thread and pushing commands to the buffer fetched by Defer() is only supported also on the game thread
-	 * (due to checking the cached thread ID).
-	 * The whole CL aims to support non-mass code trying to push commands while the flushing is going on (as triggered
-	 * by MassObservers reacting to the commands being flushed currently).
-	 */
-	uint8 OpenedCommandBufferIndex = 0;
-	static constexpr int32 NumCommandBuffers = 2;
-	TStaticArray<TSharedPtr<FMassCommandBuffer>, NumCommandBuffers> DeferredCommandBuffers;
-	
 	std::atomic<int32> ProcessingScopeCount = 0;
 
 	// the "version" number increased every time an archetype gets added
@@ -1306,10 +1293,29 @@ private:
 
 	FOnNewArchetypeDelegate OnNewArchetypeEvent;
 
-	bool bInitialized = false;
+	FDelegateHandle OnPostForkHandle;
+
+	/**
+	 * This index will be enough to control which buffer is available for pushing commands since flashing is taking place 
+	 * in the game thread and pushing commands to the buffer fetched by Defer() is only supported also on the game thread
+	 * (due to checking the cached thread ID).
+	 * The whole CL aims to support non-mass code trying to push commands while the flushing is going on (as triggered
+	 * by MassObservers reacting to the commands being flushed currently).
+	 */
+	static constexpr int32 NumCommandBuffers = 2;
+	TStaticArray<TSharedPtr<FMassCommandBuffer>, NumCommandBuffers> DeferredCommandBuffers;
+	uint8 OpenedCommandBufferIndex = 0;
+	std::atomic<bool> bCommandBufferFlushingInProgress = false;
 	bool bFirstCommandFlush = true;
 
-	FDelegateHandle OnPostForkHandle;
+	enum class EInitializationState : uint8
+	{
+		Uninitialized,
+		Initialized,
+		Deinitialized
+	};
+
+	EInitializationState InitializationState = EInitializationState::Uninitialized;
 
 	//-----------------------------------------------------------------------------
 	// DEPRECATED
