[Mass] Obesrvers deprecating stuff

diff --git a/Engine/Source/Runtime/MassEntity/Private/MassObserverManager.cpp b/Engine/Source/Runtime/MassEntity/Private/MassObserverManager.cpp
--- a/Engine/Source/Runtime/MassEntity/Private/MassObserverManager.cpp
+++ b/Engine/Source/Runtime/MassEntity/Private/MassObserverManager.cpp
@@ -14,6 +14,26 @@
 
 namespace UE::Mass::ObserverManager
 {
+	struct FDeprecationHelper
+	{
+		static void HandleSingleElement(TNotNull<FMassObserverManager*> ObserverManager, const UScriptStruct& ElementType, const FMassArchetypeEntityCollection& EntityCollection, FMassObserversMap& HandlersContainer)
+		{
+			UE::Mass::FProcessingContext ProcessingContext(ObserverManager->EntityManager);
+			ProcessingContext.AuxData.InitializeAs(&ElementType);
+			FMassObserverManager::HandleElementsImpl(ProcessingContext, { EntityCollection }, MakeArrayView((const UScriptStruct**)&ElementType, 1), HandlersContainer);
+		}
+
+		static void HandleSingleElement(TNotNull<FMassObserverManager*> ObserverManager, const UScriptStruct& ElementType, const FMassArchetypeEntityCollection& EntityCollection, const EMassObservedOperation Operation)
+		{
+			const int32 OperationIndex = int32(Operation);
+			const bool bIsFragment = UE::Mass::IsA<FMassFragment>(&ElementType);
+			check(bIsFragment || UE::Mass::IsA<FMassTag>(&ElementType));
+
+			HandleSingleElement(ObserverManager, ElementType, EntityCollection
+				, bIsFragment ? ObserverManager->FragmentObservers[OperationIndex] : ObserverManager->TagObservers[OperationIndex]);
+		}
+	};
+
 	namespace Tweakables
 	{
 		// Used as a template parameter for TInlineAllocator that we use when gathering UScriptStruct* of the observed types to process.
@@ -211,7 +231,7 @@ bool FMassObserverManager::OnCompositionChanged(UE::Mass::FProcessingContext& Pr
 		{
 			FragmentOverlap.ExportTypes(ObservedTypesOverlap);
 
-			HandleFragmentsImpl(ProcessingContext, EntityCollection, ObservedTypesOverlap, FragmentObservers[(uint8)Operation]);
+			HandleElementsImpl(ProcessingContext, {EntityCollection}, ObservedTypesOverlap, FragmentObservers[(uint8)Operation]);
 		}
 
 		if (bHasTagsOverlap)
@@ -219,7 +239,7 @@ bool FMassObserverManager::OnCompositionChanged(UE::Mass::FProcessingContext& Pr
 			ObservedTypesOverlap.Reset();
 			TagOverlap.ExportTypes(ObservedTypesOverlap);
 
-			HandleFragmentsImpl(ProcessingContext, EntityCollection, ObservedTypesOverlap, TagObservers[(uint8)Operation]);
+			HandleElementsImpl(ProcessingContext, {EntityCollection}, ObservedTypesOverlap, TagObservers[(uint8)Operation]);
 		}
 
 		return true;
@@ -249,51 +269,29 @@ bool FMassObserverManager::OnCompositionChanged(const FMassEntityHandle Entity,
 		UE::Mass::FProcessingContext ProcessingContext(EntityManager, /*DeltaSeconds=*/0.f);
 		ProcessingContext.bFlushCommandBuffer = false;
 		const FMassArchetypeHandle ArchetypeHandle = EntityManager.GetArchetypeForEntity(Entity);
+		FMassArchetypeEntityCollection EntityCollection(ArchetypeHandle, {Entity}, FMassArchetypeEntityCollection::NoDuplicates);
 
 		if (bHasFragmentsOverlap)
 		{
 			FragmentOverlap.ExportTypes(ObservedTypesOverlap);
-
-			HandleFragmentsImpl(ProcessingContext, FMassArchetypeEntityCollection(ArchetypeHandle, MakeArrayView(&Entity, 1)
-				, FMassArchetypeEntityCollection::NoDuplicates), ObservedTypesOverlap, FragmentObservers[(uint8)Operation]);
+			HandleElementsImpl(ProcessingContext, {EntityCollection}, ObservedTypesOverlap, FragmentObservers[(uint8)Operation]);
 		}
 
 		if (bHasTagsOverlap)
 		{
 			ObservedTypesOverlap.Reset();
 			TagOverlap.ExportTypes(ObservedTypesOverlap);
-
-			HandleFragmentsImpl(ProcessingContext, FMassArchetypeEntityCollection(ArchetypeHandle, MakeArrayView(&Entity, 1)
-				, FMassArchetypeEntityCollection::NoDuplicates), ObservedTypesOverlap, TagObservers[(uint8)Operation]);
+			HandleElementsImpl(ProcessingContext, {EntityCollection}, ObservedTypesOverlap, TagObservers[(uint8)Operation]);
 		}
 	}
 
 	return bHasFragmentsOverlap || bHasTagsOverlap;
 }
 
-void FMassObserverManager::OnElementOperation(TNotNull<const UScriptStruct*> ElementType, const FMassArchetypeEntityCollection& EntityCollection, const EMassObservedOperation Operation)
-{
-	const bool bIsFragment = UE::Mass::IsA<FMassFragment>(ElementType);
-	check(bIsFragment || UE::Mass::IsA<FMassTag>(ElementType));
-
-	if (bIsFragment)
-	{
-		if (ObservedFragments[(uint8)Operation].Contains(*ElementType))
-		{
-			HandleSingleElementImpl(*ElementType, EntityCollection, FragmentObservers[(uint8)Operation]);
-		}
-	}
-	else if (ObservedTags[(uint8)Operation].Contains(*ElementType))
-	{
-		HandleSingleElementImpl(*ElementType, EntityCollection, TagObservers[(uint8)Operation]);
-	}
-}
-
-void FMassObserverManager::HandleFragmentsImpl(UE::Mass::FProcessingContext& ProcessingContext, const FMassArchetypeEntityCollection& EntityCollection
-	, TArrayView<const UScriptStruct*> ObservedTypes
-	/*, const FMassFragmentBitSet& FragmentsBitSet*/, FMassObserversMap& HandlersContainer)
+void FMassObserverManager::HandleElementsImpl(UE::Mass::FProcessingContext& ProcessingContext, TConstArrayView<FMassArchetypeEntityCollection> EntityCollections
+	, TArrayView<const UScriptStruct*> ObservedTypes, FMassObserversMap& HandlersContainer)
 {	
-	TRACE_CPUPROFILER_EVENT_SCOPE(MassObserver_HandleFragmentsImpl);
+	TRACE_CPUPROFILER_EVENT_SCOPE(MassObserver_HandleElementsImpl);
 
 	check(ObservedTypes.Num() > 0);
 
@@ -302,22 +300,10 @@ void FMassObserverManager::HandleFragmentsImpl(UE::Mass::FProcessingContext& Pro
 		ProcessingContext.AuxData.InitializeAs(Type);
 		FMassRuntimePipeline& Pipeline = (*HandlersContainer).FindChecked(Type);
 
-		UE::Mass::Executor::RunProcessorsView(Pipeline.GetMutableProcessors(), ProcessingContext, MakeArrayView(&EntityCollection, 1));
+		UE::Mass::Executor::RunProcessorsView(Pipeline.GetMutableProcessors(), ProcessingContext, EntityCollections);
 	}
 }
 
-void FMassObserverManager::HandleSingleElementImpl(const UScriptStruct& FragmentType, const FMassArchetypeEntityCollection& EntityCollection, FMassObserversMap& HandlersContainer)
-{
-	TRACE_CPUPROFILER_EVENT_SCOPE(MassObserver_HandleSingleEntityImpl);
-
-	UE::Mass::FProcessingContext ProcessingContext(EntityManager, /*DeltaSeconds=*/0.f);
-	ProcessingContext.bFlushCommandBuffer = false;
-	ProcessingContext.AuxData.InitializeAs(&FragmentType);
-	FMassRuntimePipeline& Pipeline = (*HandlersContainer).FindChecked(&FragmentType);
-
-	UE::Mass::Executor::RunProcessorsView(Pipeline.GetMutableProcessors(), ProcessingContext, MakeArrayView(&EntityCollection, 1));
-}
-
 void FMassObserverManager::AddObserverInstance(const UScriptStruct& ElementType, const EMassObservedOperation Operation, UMassProcessor& ObserverProcessor)
 {
 	const bool bIsFragment = UE::Mass::IsA<FMassFragment>(&ElementType);
@@ -386,20 +372,26 @@ void FMassObserverManager::RemoveObserverInstance(const UScriptStruct& ElementTy
 //----------------------------------------------------------------------//
 void FMassObserverManager::HandleSingleEntityImpl(const UScriptStruct& FragmentType, const FMassArchetypeEntityCollection& EntityCollection, FMassObserversMap& HandlersContainer)
 {
-	HandleSingleElementImpl(FragmentType, EntityCollection, HandlersContainer);
+	UE::Mass::ObserverManager::FDeprecationHelper::HandleSingleElement(this, FragmentType, EntityCollection, HandlersContainer);
 }
 
 void FMassObserverManager::OnPostFragmentOrTagAdded(const UScriptStruct& FragmentOrTagType, const FMassArchetypeEntityCollection& EntityCollection)
 {
-	OnPostElementAdded(&FragmentOrTagType, EntityCollection);
+	UE::Mass::ObserverManager::FDeprecationHelper::HandleSingleElement(this, FragmentOrTagType, EntityCollection, EMassObservedOperation::Add);
 }
 
 void FMassObserverManager::OnPreFragmentOrTagRemoved(const UScriptStruct& FragmentOrTagType, const FMassArchetypeEntityCollection& EntityCollection)
 {
-	OnPreElementRemoved(&FragmentOrTagType, EntityCollection);
+	UE::Mass::ObserverManager::FDeprecationHelper::HandleSingleElement(this, FragmentOrTagType, EntityCollection, EMassObservedOperation::Remove);
 }
 
 void FMassObserverManager::OnFragmentOrTagOperation(const UScriptStruct& FragmentOrTagType, const FMassArchetypeEntityCollection& EntityCollection, const EMassObservedOperation Operation)
 {
-	OnElementOperation(&FragmentOrTagType, EntityCollection, Operation);
+	UE::Mass::ObserverManager::FDeprecationHelper::HandleSingleElement(this, FragmentOrTagType, EntityCollection, Operation);
+}
+
+void FMassObserverManager::HandleFragmentsImpl(UE::Mass::FProcessingContext& ProcessingContext, const FMassArchetypeEntityCollection& EntityCollection
+		, TArrayView<const UScriptStruct*> ObservedTypes, FMassObserversMap& HandlersContainer)
+{
+	HandleElementsImpl(ProcessingContext, {EntityCollection}, ObservedTypes, HandlersContainer);
 }
diff --git a/Engine/Source/Runtime/MassEntity/Public/MassObserverManager.h b/Engine/Source/Runtime/MassEntity/Public/MassObserverManager.h
--- a/Engine/Source/Runtime/MassEntity/Public/MassObserverManager.h
+++ b/Engine/Source/Runtime/MassEntity/Public/MassObserverManager.h
@@ -14,6 +14,10 @@ class UMassProcessor;
 namespace UE::Mass
 {
 	struct FProcessingContext;
+	namespace ObserverManager
+	{
+		struct FDeprecationHelper;
+	}
 }
 
 /** 
@@ -115,16 +119,6 @@ public:
 	/** @return whether any observers have been executed */
 	bool OnCompositionChanged(const FMassArchetypeEntityCollection& EntityCollection, const FMassArchetypeCompositionDescriptor& Composition, const EMassObservedOperation Operation);
 
-	void OnPostElementAdded(TNotNull<const UScriptStruct*> ElementType, const FMassArchetypeEntityCollection& EntityCollection)
-	{
-		OnElementOperation(ElementType, EntityCollection, EMassObservedOperation::Add);
-	}
-	void OnPreElementRemoved(TNotNull<const UScriptStruct*> ElementType, const FMassArchetypeEntityCollection& EntityCollection)
-	{
-		OnElementOperation(ElementType, EntityCollection, EMassObservedOperation::Remove);
-	}
-	void OnElementOperation(TNotNull<const UScriptStruct*> ElementType, const FMassArchetypeEntityCollection& EntityCollection, const EMassObservedOperation Operation);
-
 	void AddObserverInstance(const UScriptStruct& ElementType, const EMassObservedOperation Operation, UMassProcessor& ObserverProcessor);
 	void RemoveObserverInstance(const UScriptStruct& ElementType, const EMassObservedOperation Operation, UMassProcessor& ObserverProcessor);
 
@@ -136,9 +130,8 @@ protected:
 
 	void Initialize();
 	void DeInitialize();
-	void HandleFragmentsImpl(UE::Mass::FProcessingContext& ProcessingContext, const FMassArchetypeEntityCollection& EntityCollection
+	static void HandleElementsImpl(UE::Mass::FProcessingContext& ProcessingContext, TConstArrayView<FMassArchetypeEntityCollection> EntityCollections
 		, TArrayView<const UScriptStruct*> ObservedTypes, FMassObserversMap& HandlersContainer);
-	void HandleSingleElementImpl(const UScriptStruct& FragmentType, const FMassArchetypeEntityCollection& EntityCollection, FMassObserversMap& HandlersContainer);
 
 	FMassFragmentBitSet ObservedFragments[(uint8)EMassObservedOperation::MAX];
 	FMassTagBitSet ObservedTags[(uint8)EMassObservedOperation::MAX];
@@ -158,15 +151,20 @@ protected:
 public:
 	UE_DEPRECATED(5.5, "This flavor of OnPostEntitiesCreated is deprecated. Please use the one taking a TConstArrayView<FMassArchetypeEntityCollection> parameter instead")
 	bool OnPostEntitiesCreated(UE::Mass::FProcessingContext& InProcessingContext, const FMassArchetypeEntityCollection& EntityCollection);
-	UE_DEPRECATED(5.5, "*FragmentOrTag* functions have been renamed to *Element*. Please use those instead.")
+	UE_DEPRECATED(5.6, "*FragmentOrTag* functions have been deprecated. Use OnCompositionChanged instead.")
 	void OnPostFragmentOrTagAdded(const UScriptStruct& FragmentOrTagType, const FMassArchetypeEntityCollection& EntityCollection);
-	UE_DEPRECATED(5.5, "*FragmentOrTag* functions have been renamed to *Element*. Please use those instead.")
+	UE_DEPRECATED(5.6, "*FragmentOrTag* functions have been deprecated. Use OnCompositionChanged instead.")
 	void OnPreFragmentOrTagRemoved(const UScriptStruct& FragmentOrTagType, const FMassArchetypeEntityCollection& EntityCollection);
-	UE_DEPRECATED(5.5, "*FragmentOrTag* functions have been renamed to *Element*. Please use those instead.")
+	UE_DEPRECATED(5.6, "*FragmentOrTag* functions have been deprecated. Use OnCompositionChanged instead.")
 	void OnFragmentOrTagOperation(const UScriptStruct& FragmentOrTagType, const FMassArchetypeEntityCollection& EntityCollection, const EMassObservedOperation Operation);
 protected:
-	UE_DEPRECATED(5.6, "Use HandleSingleElementImpl.")
+	UE_DEPRECATED(5.6, "Use HandleElementsImpl instead.")
 	void HandleSingleEntityImpl(const UScriptStruct& FragmentType, const FMassArchetypeEntityCollection& EntityCollection, FMassObserversMap& HandlersContainer);
+	UE_DEPRECATED(5.6, "Use HandleElementsImpl instead.")
+	void HandleFragmentsImpl(UE::Mass::FProcessingContext& ProcessingContext, const FMassArchetypeEntityCollection& EntityCollection
+		, TArrayView<const UScriptStruct*> ObservedTypes, FMassObserversMap& HandlersContainer);
+
+	friend UE::Mass::ObserverManager::FDeprecationHelper;
 };
 
 template<>
