[MassEntity] merged two error handling cases

diff --git a/Engine/Source/Runtime/MassEntity/Private/MassEntityManager.cpp b/Engine/Source/Runtime/MassEntity/Private/MassEntityManager.cpp
--- a/Engine/Source/Runtime/MassEntity/Private/MassEntityManager.cpp
+++ b/Engine/Source/Runtime/MassEntity/Private/MassEntityManager.cpp
@@ -2069,22 +2069,16 @@ void FMassEntityManager::FlushCommands(const TSharedPtr<FMassCommandBuffer>& InC
 		return;
 	}
 
-	if (UNLIKELY(InitializationState == EInitializationState::Uninitialized))
+	if (UNLIKELY(InitializationState != EInitializationState::Initialized))
 	{
-		UE_VLOG_UELOG(GetOwner(), LogMass, Warning, TEXT("FlushCommands called before Initialize call, which means this FMassEntityManager instance is not ready to ,"
-			" process commands and will cancel them."));
+		UE_CVLOG_UELOG(InitializationState == EInitializationState::Uninitialized, GetOwner(), LogMass, Warning
+			, TEXT("FlushCommands called before Initialize call, which means this FMassEntityManager instance is not ready to process commands and will cancel them."));
+		UE_CVLOG_UELOG(InitializationState == EInitializationState::Deinitialized, GetOwner(), LogMass, Log
+			, TEXT("FlushCommands called after Deinitialize call, which means this FMassEntityManager instance is going away, can't process commands and will cancel them."));
 		InCommandBuffer->CancelCommands();
 		return;
 	}
 	
-	if (UNLIKELY(InitializationState == EInitializationState::Deinitialized))
-	{
-		UE_VLOG_UELOG(GetOwner(), LogMass, Log, TEXT("FlushCommands called after Deinitialize call, which means this FMassEntityManager instance is going away,"
-			" can't process commands and will cancel them."));
-		InCommandBuffer->CancelCommands();
-		return;
-	}
-
 	if (InCommandBuffer && InCommandBuffer->HasPendingCommands()
 		&& (Algo::Find(DeferredCommandBuffers, InCommandBuffer) == nullptr))
 	{
