[MassEntity] fixed typo

diff --git a/Engine/Source/Runtime/MassEntity/Private/MassEntityQuery.cpp b/Engine/Source/Runtime/MassEntity/Private/MassEntityQuery.cpp
--- a/Engine/Source/Runtime/MassEntity/Private/MassEntityQuery.cpp
+++ b/Engine/Source/Runtime/MassEntity/Private/MassEntityQuery.cpp
@@ -22,8 +22,8 @@
 namespace UE::Mass::Tweakables
 {
 	/**
-	 * Controls whether ParallelForEachEntityChunk actually performs ParallelFor operations. If `false` the call is passed
-	 * the the regular ForEachEntityChunk call.
+	 * Controls whether ParallelForEachEntityChunk actually performs ParallelFor operations.
+	 * If `false` the call is passed to the regular ForEachEntityChunk call.
 	 */
 	bool bAllowParallelExecution = true;
 
