[Mass] Added an archetype handle rvalue ref constructor to FMassArchetypeEntityCollection

#rnx

diff --git a/Engine/Source/Runtime/MassEntity/Private/MassArchetypeTypes.cpp b/Engine/Source/Runtime/MassEntity/Private/MassArchetypeTypes.cpp
--- a/Engine/Source/Runtime/MassEntity/Private/MassArchetypeTypes.cpp
+++ b/Engine/Source/Runtime/MassEntity/Private/MassArchetypeTypes.cpp
@@ -152,6 +152,12 @@ FMassArchetypeEntityCollection::FMassArchetypeEntityCollection(const FMassArchet
 {	
 }
 
+FMassArchetypeEntityCollection::FMassArchetypeEntityCollection(FMassArchetypeHandle&& InArchetype, const FMassEntityHandle EntityHandle)
+	: Ranges({CreateRangeForEntity(InArchetype, EntityHandle)})
+	, Archetype(MoveTemp(InArchetype))
+{	
+}
+
 void FMassArchetypeEntityCollection::BuildEntityRanges(TStridedView<const int32> TrueIndices)
 {
 	checkf(Ranges.Num() == 0, TEXT("Calling %s is valid only for initial configuration"), ANSI_TO_TCHAR(__FUNCTION__));
diff --git a/Engine/Source/Runtime/MassEntity/Public/MassArchetypeTypes.h b/Engine/Source/Runtime/MassEntity/Public/MassArchetypeTypes.h
--- a/Engine/Source/Runtime/MassEntity/Public/MassArchetypeTypes.h
+++ b/Engine/Source/Runtime/MassEntity/Public/MassArchetypeTypes.h
@@ -166,6 +166,7 @@ public:
 	FMassArchetypeEntityCollection(const FMassArchetypeHandle& InArchetype, TConstArrayView<FMassEntityHandle> InEntities, EDuplicatesHandling DuplicatesHandling);
 	/** optimized, special case for a single-entity */
 	FMassArchetypeEntityCollection(const FMassArchetypeHandle& InArchetype, const FMassEntityHandle EntityHandle);
+	FMassArchetypeEntityCollection(FMassArchetypeHandle&& InArchetype, const FMassEntityHandle EntityHandle);
 	explicit FMassArchetypeEntityCollection(const FMassArchetypeHandle& InArchetypeHandle, const EInitializationType Initialization = EInitializationType::GatherAll);
 	explicit FMassArchetypeEntityCollection(TSharedPtr<FMassArchetypeData>& InArchetype, const EInitializationType Initialization = EInitializationType::GatherAll);
 	FMassArchetypeEntityCollection(const FMassArchetypeHandle& InArchetypeHandle, FEntityRangeArray&& InEntityRanges)
