[Mass] added FMassArchetypeCompositionDescriptor.IsIdentical to check whether two FMassArchetypeCompositionDescriptor instances are identical

#rnx

diff --git a/Engine/Source/Runtime/MassEntity/Public/MassEntityTypes.h b/Engine/Source/Runtime/MassEntity/Public/MassEntityTypes.h
--- a/Engine/Source/Runtime/MassEntity/Public/MassEntityTypes.h
+++ b/Engine/Source/Runtime/MassEntity/Public/MassEntityTypes.h
@@ -100,6 +100,9 @@ struct FMassArchetypeCompositionDescriptor
 		ConstSharedFragments.Reset();
 	}
 
+	/**
+	 * Compares contents of two FMassArchetypeCompositionDescriptor instances, ignoring the trailing empty bits in the bitsets
+	 */
 	bool IsEquivalent(const FMassArchetypeCompositionDescriptor& OtherDescriptor) const
 	{
 		return Fragments.IsEquivalent(OtherDescriptor.Fragments)
@@ -109,6 +112,18 @@ struct FMassArchetypeCompositionDescriptor
 			&& ConstSharedFragments.IsEquivalent(OtherDescriptor.ConstSharedFragments);
 	}
 
+	/**
+	 * Checks whether contents of two FMassArchetypeCompositionDescriptor instances are identical.
+	 */
+	bool IsIdentical(const FMassArchetypeCompositionDescriptor& OtherDescriptor) const
+	{
+		return Fragments == OtherDescriptor.Fragments
+			&& Tags == OtherDescriptor.Tags
+			&& ChunkFragments == OtherDescriptor.ChunkFragments
+			&& SharedFragments == OtherDescriptor.SharedFragments
+			&& ConstSharedFragments == OtherDescriptor.ConstSharedFragments;
+	}
+
 	bool IsEmpty() const 
 	{ 
 		return Fragments.IsEmpty()
