[Mass] Mass Observers overhaul.
There are two observes locking mechanisms now:
* FObserverLock - locks all notifications, and every opertation information gets stored so that notification can be sent out once the lock is released
* FCreationContext - a wrapper for FObserverLock, that makes the ObserversManager interpret the operations being performed as operations on the freshly created entities - these operations don't track specific compsition changes, we only care about the end result.

In addition, MassObserversManager's API has been simplified, and there no longer a need on EntityManager side to conditionally notify the observers manager - now it will get notified about all the operations, on whether the notified go out right away or get buffered depends o whether a FObserverLock has been created and is kept alive.

diff --git a/Engine/Plugins/Runtime/MassGameplay/Source/MassSpawner/Private/MassSpawnerSubsystem.cpp b/Engine/Plugins/Runtime/MassGameplay/Source/MassSpawner/Private/MassSpawnerSubsystem.cpp
--- a/Engine/Plugins/Runtime/MassGameplay/Source/MassSpawner/Private/MassSpawnerSubsystem.cpp
+++ b/Engine/Plugins/Runtime/MassGameplay/Source/MassSpawner/Private/MassSpawnerSubsystem.cpp
@@ -15,6 +15,7 @@
 #include "MassProcessor.h"
 #include "MassEntityUtils.h"
 #include "MassProcessingContext.h"
+#include "MassObserverNotificationTypes.h"
 
 //----------------------------------------------------------------------//
 //  UMassSpawnerSubsystem
@@ -135,7 +136,7 @@ TSharedPtr<FMassEntityManager::FEntityCreationContext> UMassSpawnerSubsystem::Do
 	// 2. Copy data from FMassEntityTemplate.Fragments.
 	//		a. @todo, could be done as part of creation?
 	TConstArrayView<FInstancedStruct> FragmentInstances = EntityTemplate.GetInitialFragmentValues();
-	EntityManager->BatchSetEntityFragmentValues(CreationContext->GetEntityCollections(), FragmentInstances);
+	EntityManager->BatchSetEntityFragmentValues(CreationContext->GetEntityCollections(*EntityManager.Get()), FragmentInstances);
 	
 	// 3. Run SpawnDataInitializer, if set. This is a special type of processor that operates on the entities to initialize them.
 	// e.g., will run UInstancedActorsInitializerProcessor for Mass InstancedActors
@@ -147,7 +148,7 @@ TSharedPtr<FMassEntityManager::FEntityCreationContext> UMassSpawnerSubsystem::Do
 	{
 		FMassProcessingContext ProcessingContext(EntityManager, /*TimeDelta=*/0.0f);
 		ProcessingContext.AuxData = SpawnData;
-		UE::Mass::Executor::RunProcessorsView(MakeArrayView(&SpawnDataInitializer, 1), ProcessingContext, CreationContext->GetEntityCollections());
+		UE::Mass::Executor::RunProcessorsView(MakeArrayView(&SpawnDataInitializer, 1), ProcessingContext, CreationContext->GetEntityCollections(*EntityManager.Get()));
 	}
 
 	OutEntities.Append(MoveTemp(SpawnedEntities));
diff --git a/Engine/Source/Developer/MassEntityTestSuite/Private/MassCreationContextTest.cpp b/Engine/Source/Developer/MassEntityTestSuite/Private/MassCreationContextTest.cpp
--- a/Engine/Source/Developer/MassEntityTestSuite/Private/MassCreationContextTest.cpp
+++ b/Engine/Source/Developer/MassEntityTestSuite/Private/MassCreationContextTest.cpp
@@ -1,11 +1,10 @@
 // Copyright Epic Games, Inc. All Rights Reserved.
 
 #include "MassEntityManager.h"
-#include "MassProcessingTypes.h"
 #include "MassEntityTestTypes.h"
 #include "MassEntityTypes.h"
-#include "MassCommandBuffer.h"
 #include "MassExecutionContext.h"
+#include "MassObserverNotificationTypes.h"
 
 #define LOCTEXT_NAMESPACE "MassTest"
 
@@ -29,9 +28,10 @@ struct FCreationContextTest : FEntityTestBase
 		const int32 NumDifferentArchetypesUsed = 2;
 
 		AITEST_EQUAL(TEXT("Two back to back entity creation operations should result in the same creation context"), CreationContextInt, CreationContextFloat);
-		AITEST_FALSE(TEXT("CreationContext's entity collection should be invalidated at this moment"), CreationContextInt->DebugAreEntityCollectionsUpToDate());
+		AITEST_TRUE(TEXT("CreationContext's entity collection should be still valid since we only created two consistent collections of entities")
+			, CreationContextInt->DebugAreEntityCollectionsUpToDate());
 
-		TConstArrayView<FMassArchetypeEntityCollection> EntityCollections = CreationContextInt->GetEntityCollections();
+		TArray<FMassArchetypeEntityCollection> EntityCollections = CreationContextInt->GetEntityCollections(*EntityManager.Get());
 		AITEST_EQUAL(TEXT("We expect the number of resulting collections to match expectations"), EntityCollections.Num(), NumDifferentArchetypesUsed);
 
 		return true;
@@ -67,11 +67,11 @@ struct FManualCreate : FEntityTestBase
 			TSharedRef<FMassEntityManager::FEntityCreationContext> TempContext = EntityManager->BatchCreateEntities(IntsArchetype, IntEntitiesToSpawnCount, Entities);
 			AITEST_EQUAL(TEXT("Creating entities should return the original context"), ObtainedContext, TempContext);
 
-			AITEST_FALSE(TEXT("CreationContext's entity collection should be invalidated at this moment")
+			AITEST_TRUE(TEXT("CreationContext's entity collection should be still valid, because we're only piling up consistent entity collections")
 				, TempContext->DebugAreEntityCollectionsUpToDate());
 		}
 
-		TConstArrayView<FMassArchetypeEntityCollection> EntityCollections = ObtainedContext->GetEntityCollections();
+		TArray<FMassArchetypeEntityCollection> EntityCollections = ObtainedContext->GetEntityCollections(*EntityManager.Get());
 		AITEST_EQUAL(TEXT("We expect the number of resulting collections to match expectations"), EntityCollections.Num(), NumDifferentArchetypesUsed);
 
 		return true;
@@ -102,7 +102,7 @@ struct FManualBuild : FEntityTestBase
 		FMassArchetypeEntityCollectionWithPayload::CreateEntityRangesWithPayload(*EntityManager, Entities, FMassArchetypeEntityCollection::NoDuplicates
 			, FMassGenericPayloadView(MakeArrayView(&PaloadView, 1)), EntityCollections);
 
-		checkf(EntityCollections.Num() <= 1, TEXT("We expect TargetEntities to only contain archetype-less entities, ones that need to be \'build\'"));
+		checkf(EntityCollections.Num() <= 1, TEXT("We expect TargetEntities to only contain archetype-less entities, ones that need to be built"));
 
 		{
 			TSharedRef<FMassEntityManager::FEntityCreationContext> CreationContext = EntityManager->BatchBuildEntities(EntityCollections[0], FMassFragmentBitSet(*FTestFragment_Float::StaticStruct()));
@@ -113,7 +113,7 @@ struct FManualBuild : FEntityTestBase
 		AITEST_TRUE(TEXT("CreationContext's entity collection should be still valid at this moment since we only added one entity collection/array")
 			, ObtainedContext->DebugAreEntityCollectionsUpToDate());
 
-		TConstArrayView<FMassArchetypeEntityCollection> ContextEntityCollections = ObtainedContext->GetEntityCollections();
+		TArray<FMassArchetypeEntityCollection> ContextEntityCollections = ObtainedContext->GetEntityCollections(*EntityManager.Get());
 		AITEST_EQUAL(TEXT("We expect the number of resulting collections to match expectations"), ContextEntityCollections.Num(), NumDifferentArchetypesUsed);
 
 		return true;
diff --git a/Engine/Source/Developer/MassEntityTestSuite/Private/MassExecutionTest.cpp b/Engine/Source/Developer/MassEntityTestSuite/Private/MassExecutionTest.cpp
--- a/Engine/Source/Developer/MassEntityTestSuite/Private/MassExecutionTest.cpp
+++ b/Engine/Source/Developer/MassEntityTestSuite/Private/MassExecutionTest.cpp
@@ -1,7 +1,7 @@
 // Copyright Epic Games, Inc. All Rights Reserved.
 
 #include "MassEntityManager.h"
-#include "MassProcessingTypes.h"
+#include "MassProcessingContext.h"
 #include "MassEntityTestTypes.h"
 #include "MassExecutor.h"
 #include "MassProcessingContext.h"
diff --git a/Engine/Source/Developer/MassEntityTestSuite/Private/MassMultiThreadTest.cpp b/Engine/Source/Developer/MassEntityTestSuite/Private/MassMultiThreadTest.cpp
--- a/Engine/Source/Developer/MassEntityTestSuite/Private/MassMultiThreadTest.cpp
+++ b/Engine/Source/Developer/MassEntityTestSuite/Private/MassMultiThreadTest.cpp
@@ -1,7 +1,7 @@
 // Copyright Epic Games, Inc. All Rights Reserved.
 
 #include "MassEntityManager.h"
-#include "MassProcessingTypes.h"
+#include "MassProcessingContext.h"
 #include "MassEntityTestTypes.h"
 #include "MassEntitySettings.h"
 #include "MassExecutor.h"
diff --git a/Engine/Source/Developer/MassEntityTestSuite/Private/MassObserverTest.cpp b/Engine/Source/Developer/MassEntityTestSuite/Private/MassObserverTest.cpp
--- a/Engine/Source/Developer/MassEntityTestSuite/Private/MassObserverTest.cpp
+++ b/Engine/Source/Developer/MassEntityTestSuite/Private/MassObserverTest.cpp
@@ -6,6 +6,7 @@
 #include "MassEntityTypes.h"
 #include "MassCommandBuffer.h"
 #include "MassExecutionContext.h"
+#include "MassObserverNotificationTypes.h"
 
 #define LOCTEXT_NAMESPACE "MassTest"
 
@@ -356,11 +357,11 @@ struct FEntityCreation_Batched : FTagBaseOperation
 		{
 			TSharedRef<FMassEntityManager::FEntityCreationContext> CreationContext = EntityManager->BatchCreateEntities(IntsArchetype, EntitiesToSpawnCount, EntitiesInt);
 
-			EntityManager->BatchChangeTagsForEntities(CreationContext->GetEntityCollections(), FMassTagBitSet(*FTagStruct::StaticStruct()), FMassTagBitSet());
+			EntityManager->BatchChangeTagsForEntities(CreationContext->GetEntityCollections(*EntityManager.Get()), FMassTagBitSet(*FTagStruct::StaticStruct()), FMassTagBitSet());
 			AITEST_TRUE(TEXT("The tag observer is not expected to run yet"), AffectedEntities.Num() == 0);
 			AITEST_FALSE(TEXT("CreationContext's entity collection should be invalidated at this moment"), CreationContext->DebugAreEntityCollectionsUpToDate());
 
-			EntityManager->BatchChangeTagsForEntities(CreationContext->GetEntityCollections(), FMassTagBitSet(*FTagStruct::StaticStruct()), FMassTagBitSet());
+			EntityManager->BatchChangeTagsForEntities(CreationContext->GetEntityCollections(*EntityManager.Get()), FMassTagBitSet(*FTagStruct::StaticStruct()), FMassTagBitSet());
 			AITEST_TRUE(TEXT("The tag observer is still not expected to run"), AffectedEntities.Num() == 0);
 		}
 		AITEST_TRUE(TEXT("The tag observer is expected to run just after FEntityCreationContext's destruction"), AffectedEntities.Num() > 0);
@@ -395,7 +396,7 @@ struct FFragmentTestBase : FEntityTestBase
 		if (FEntityTestBase::SetUp())
 		{
 			ObserverProcessor = NewTestProcessor<UMassTestProcessorBase>(EntityManager);
-			ObserverProcessor->EntityQuery.AddRequirement(FFragmentStruct::StaticStruct(), EMassFragmentAccess::ReadOnly);
+			ObserverProcessor->EntityQuery.AddRequirement(FFragmentStruct::StaticStruct(), EMassFragmentAccess::ReadWrite);
 			ObserverProcessor->ForEachEntityChunkExecutionFunction = [bCommandsFlushedPtr = &bCommandsFlushed, AffectedEntitiesPtr = &AffectedEntities](FMassExecutionContext& Context)
 				{
 					AffectedEntitiesPtr->Append(Context.GetEntities().GetData(), Context.GetEntities().Num());
@@ -724,18 +725,17 @@ struct FModificationsWhileCreationContextActive : FTagBaseOperation
 		ObserverManager.AddObserverInstance(*FTagStruct::StaticStruct(), OperationObserved, *ObserverProcessor);
 
 		{
-			TSharedRef<FMassEntityManager::FEntityCreationContext> CreationContext = EntityManager->BatchCreateEntities(IntsArchetype, EntitiesToSpawnInSecondBatch, EntitiesInt);
-			ensure(EntitiesInt.Num() == EntitiesToSpawnInFirstBatch + EntitiesToSpawnInSecondBatch);
-			// note that the observers' behavior regarding the entities just created gets tested by FEntityCreation_Batched test above
-			// we're testing only the behavior related to the previously created entities here
-
+			TSharedRef<FMassObserverManager::FObserverLock> ObserversLock = EntityManager->GetOrMakeObserversLock();
+			{
+				TSharedRef<FMassEntityManager::FEntityCreationContext> CreationContext = EntityManager->BatchCreateEntities(IntsArchetype, EntitiesToSpawnInSecondBatch, EntitiesInt);
+				ensure(EntitiesInt.Num() == EntitiesToSpawnInFirstBatch + EntitiesToSpawnInSecondBatch);
+				// note that the observers' behavior regarding the entities just created gets tested by FEntityCreation_Batched test above
+				// we're testing only the behavior related to the previously created entities here
+			}
 			EntityManager->BatchChangeTagsForEntities(MakeArrayView(&InitialEntitiesCollection, 1)
 				, FMassTagBitSet(*FTagStruct::StaticStruct()), FMassTagBitSet());
 
 			AITEST_TRUE(TEXT("The tag observer is not expected to run yet"), AffectedEntities.Num() == 0);
-			
-			// should it, though?
-			//AITEST_FALSE(TEXT("CreationContext's entity collection should be invalidated at this moment"), CreationContext->DebugAreEntityCollectionsUpToDate());
 		}
 		AITEST_TRUE(TEXT("The tag observer is expected to run just after FEntityCreationContext's destruction"), AffectedEntities.Num() > 0);
 		AITEST_EQUAL(TEXT("The tag observer is expected to process only the original entities, that had a tag added to them"), AffectedEntities.Num(), EntitiesToSpawnInFirstBatch);
@@ -743,7 +743,112 @@ struct FModificationsWhileCreationContextActive : FTagBaseOperation
 		return true;
 	}
 };
-//IMPLEMENT_AI_INSTANT_TEST(FModificationsWhileCreationContextActive, "System.Mass.Observer.Create.ModificationsToOtherEntities");
+IMPLEMENT_AI_INSTANT_TEST(FModificationsWhileCreationContextActive, "System.Mass.Observer.Create.ModificationsToOtherEntities");
+
+struct FCreationOperationOrder : FFragmentTestBase
+{
+	FCreationOperationOrder()
+	{
+		OperationObserved = EMassObservedOperation::Add;
+	}
+
+	virtual void BuildScenario(TArray<FMassEntityHandle>& PreExistingEntities, TArray<FMassEntityHandle>& NewEntities) = 0;
+
+	virtual bool InstantTest() override
+	{
+		TArray<FMassEntityHandle> PreExistingEntities;
+		TArray<FMassEntityHandle> NewEntities;
+		int32 Counter = 0;
+
+		ObserverProcessor->ForEachEntityChunkExecutionFunction = [&Counter](FMassExecutionContext& Context)
+			{
+				for (auto EntityId : Context.CreateEntityIterator())
+				{
+					Context.GetMutableFragmentView<FFragmentStruct>()[EntityId].Value = static_cast<float>(++Counter);
+				}
+			};
+
+		EntityManager->BatchCreateEntities(IntsArchetype, 2, PreExistingEntities);
+
+		FMassObserverManager& ObserverManager = EntityManager->GetObserverManager();
+		ObserverManager.AddObserverInstance(*FFragmentStruct::StaticStruct(), OperationObserved, *ObserverProcessor);
+
+		BuildScenario(PreExistingEntities, NewEntities);
+
+		// the specific order of entities handled within a single creation context doesn't
+		// need to match the assumed order
+		float FirstBatchValues[] = {
+			EntityManager->GetFragmentDataChecked<FTestFragment_Float>(NewEntities[0]).Value
+			, EntityManager->GetFragmentDataChecked<FTestFragment_Float>(NewEntities[1]).Value
+		};
+		AITEST_TRUE("First batch's values match", (FirstBatchValues[0] == 1.f && FirstBatchValues[1] == 2.f) || (FirstBatchValues[1] == 1.f && FirstBatchValues[0] == 2.f));
+
+		float PreExistingEntitiesValues[] = {
+			EntityManager->GetFragmentDataChecked<FTestFragment_Float>(PreExistingEntities[0]).Value
+			, EntityManager->GetFragmentDataChecked<FTestFragment_Float>(PreExistingEntities[1]).Value
+		};
+		AITEST_EQUAL("First preexisting entity's value", PreExistingEntitiesValues[0], 3.f);
+		AITEST_EQUAL("Second preexisting entity's value", PreExistingEntitiesValues[1], 6.f);
+
+		float SecondBatchBatchValues[] = {
+			EntityManager->GetFragmentDataChecked<FTestFragment_Float>(NewEntities[2]).Value
+			, EntityManager->GetFragmentDataChecked<FTestFragment_Float>(NewEntities[3]).Value
+		};
+		AITEST_TRUE("First batch's values match", (SecondBatchBatchValues[0] == 4.f && SecondBatchBatchValues[1] == 5.f) || (SecondBatchBatchValues[1] == 4.f && SecondBatchBatchValues[0] == 5.f));
+
+		return true;
+	}
+};
+
+struct FCreationOperationOrder_Batch : FCreationOperationOrder
+{
+	virtual void BuildScenario(TArray<FMassEntityHandle>& PreExistingEntities, TArray<FMassEntityHandle>& NewEntities) override
+	{
+		TSharedRef<FMassObserverManager::FObserverLock> ObserversLock = EntityManager->GetOrMakeObserversLock();
+		{
+			// creating two separate entities, that should end up in the same creation context
+			TSharedRef<FMassEntityManager::FEntityCreationContext> CreationContext = EntityManager->BatchCreateEntities(FloatsArchetype, 1, NewEntities);
+			EntityManager->BatchCreateEntities(FloatsArchetype, 1, NewEntities);
+		}
+		{
+			FMassArchetypeEntityCollection Collection(IntsArchetype, MakeArrayView(&PreExistingEntities[0], 1), FMassArchetypeEntityCollection::NoDuplicates);
+			EntityManager->BatchChangeFragmentCompositionForEntities(MakeArrayView(&Collection, 1), FMassFragmentBitSet(*FTestFragment_Float::StaticStruct()), {});
+		}
+		{
+			// creating two separate entities, that should end up in the same creation context
+			TSharedRef<FMassEntityManager::FEntityCreationContext> CreationContext = EntityManager->BatchCreateEntities(FloatsArchetype, 1, NewEntities);
+			EntityManager->BatchCreateEntities(FloatsArchetype, 1, NewEntities);
+		}
+		{
+			FMassArchetypeEntityCollection Collection(IntsArchetype, MakeArrayView(&PreExistingEntities[1], 1), FMassArchetypeEntityCollection::NoDuplicates);
+			EntityManager->BatchChangeFragmentCompositionForEntities(MakeArrayView(&Collection, 1), FMassFragmentBitSet(*FTestFragment_Float::StaticStruct()), {});
+		}
+	}
+};
+IMPLEMENT_AI_INSTANT_TEST(FCreationOperationOrder_Batch, "System.Mass.Observer.Create.CreationOperationOrder.Batch");
+
+struct FCreationOperationOrder_Individual : FCreationOperationOrder
+{
+	virtual void BuildScenario(TArray<FMassEntityHandle>& PreExistingEntities, TArray<FMassEntityHandle>& NewEntities) override
+	{
+		TSharedRef<FMassObserverManager::FObserverLock> ObserversLock = EntityManager->GetOrMakeObserversLock();
+		{
+			// creating two separate entities, that should end up in the same creation context
+			TSharedRef<FMassEntityManager::FEntityCreationContext> CreationContext = EntityManager->GetOrMakeCreationContext();
+			NewEntities.Add(EntityManager->CreateEntity(FloatsArchetype));
+			NewEntities.Add(EntityManager->CreateEntity(FloatsArchetype));
+		}
+		EntityManager->AddFragmentToEntity(PreExistingEntities[0], FTestFragment_Float::StaticStruct());
+		{
+			// creating two separate entities, that should end up in the same creation context
+			TSharedRef<FMassEntityManager::FEntityCreationContext> CreationContext = EntityManager->GetOrMakeCreationContext();
+			NewEntities.Add(EntityManager->CreateEntity(FloatsArchetype));
+			NewEntities.Add(EntityManager->CreateEntity(FloatsArchetype));
+		}
+		EntityManager->AddFragmentToEntity(PreExistingEntities[1], FTestFragment_Float::StaticStruct());
+	}
+};
+IMPLEMENT_AI_INSTANT_TEST(FCreationOperationOrder_Individual, "System.Mass.Observer.Create.CreationOperationOrder.Individual");
 
 // The scenario being tested:
 // 1. Create entities with Float fragment
@@ -763,7 +868,7 @@ struct FCreatedEntitiesUnobservedCompositionChange : FFragmentTestBase
 			TSharedRef<FMassEntityManager::FEntityCreationContext> CreationContext = EntityManager->BatchCreateEntities(FloatsArchetype, EntitiesToSpawnCount, EntitiesFloats);
 
 			// add unobserved tag
-			EntityManager->BatchChangeTagsForEntities(CreationContext->GetEntityCollections(), FMassTagBitSet(*FTestTag_A::StaticStruct()), FMassTagBitSet());
+			EntityManager->BatchChangeTagsForEntities(CreationContext->GetEntityCollections(*EntityManager.Get()), FMassTagBitSet(*FTestTag_A::StaticStruct()), FMassTagBitSet());
 
 			AITEST_EQUAL(TEXT("The fragment observer is not expected to run yet"), AffectedEntities.Num(), 0);
 		}
diff --git a/Engine/Source/Developer/MassEntityTestSuite/Private/MassSharedFragmentsTest.cpp b/Engine/Source/Developer/MassEntityTestSuite/Private/MassSharedFragmentsTest.cpp
--- a/Engine/Source/Developer/MassEntityTestSuite/Private/MassSharedFragmentsTest.cpp
+++ b/Engine/Source/Developer/MassEntityTestSuite/Private/MassSharedFragmentsTest.cpp
@@ -7,6 +7,7 @@
 #include "MassArchetypeData.h"
 #include "MassArchetypeTypes.h"
 #include "Algo/RandomShuffle.h"
+#include "MassObserverNotificationTypes.h"
 
 #define LOCTEXT_NAMESPACE "MassTest"
 
@@ -710,8 +711,8 @@ struct FSharedFragment_BatchSetAttempt : FEntityTestBase
 		FConstSharedStruct ConstSharedFragment = FConstSharedStruct::Make<FTestConstSharedFragment_Int>(TestIntValue);
 		SharedIntValues.Add(ConstSharedFragment);
 
-		TSharedRef<FMassEntityManager::FEntityCreationContext> CreationContext = EntityManager->BatchCreateEntities(FloatsArchetype, SharedIntValues, EntitiesToCreateNum, CreatedEntityHandles);
-		const FMassArchetypeHandle ResultingArchetype = CreationContext->GetEntityCollections()[0].GetArchetype();
+		TSharedRef<UE::Mass::ObserverManager::FCreationContext> CreationContext = EntityManager->BatchCreateEntities(FloatsArchetype, SharedIntValues, EntitiesToCreateNum, CreatedEntityHandles);
+		const FMassArchetypeHandle ResultingArchetype = CreationContext->GetEntityCollections(*EntityManager.Get())[0].GetArchetype();
 
 		FMassArchetypeEntityCollection EntityCollection(ResultingArchetype, CreatedEntityHandles, FMassArchetypeEntityCollection::EDuplicatesHandling::NoDuplicates);
 		
diff --git a/Engine/Source/Developer/MassEntityTestSuite/Private/MassTestTypes.cpp b/Engine/Source/Developer/MassEntityTestSuite/Private/MassTestTypes.cpp
--- a/Engine/Source/Developer/MassEntityTestSuite/Private/MassTestTypes.cpp
+++ b/Engine/Source/Developer/MassEntityTestSuite/Private/MassTestTypes.cpp
@@ -6,6 +6,7 @@
 #include "MassExecutionContext.h"
 #include "MassProcessingContext.h"
 #include "Engine/World.h"
+#include "MassProcessingContext.h"
 
 
 UE_DISABLE_OPTIMIZATION_SHIP
diff --git a/Engine/Source/Runtime/MassEntity/Private/MassEntityManager.cpp b/Engine/Source/Runtime/MassEntity/Private/MassEntityManager.cpp
--- a/Engine/Source/Runtime/MassEntity/Private/MassEntityManager.cpp
+++ b/Engine/Source/Runtime/MassEntity/Private/MassEntityManager.cpp
@@ -17,6 +17,7 @@
 #include "MassEntityBuilder.h"
 #include "MassTypeManager.h"
 #include "MassProcessingContext.h"
+#include "MassObserverNotificationTypes.h"
 
 #define CHECK_SYNC_API() testableCheckf(IsProcessing() == false, TEXT("Synchronous API function %hs called during mass processing. Use asynchronous API instead."), __FUNCTION__)
 #define CHECK_SYNC_API_RETURN(ReturnValue) testableCheckfReturn(IsProcessing() == false, return ReturnValue, TEXT("Synchronous API function %hs called during mass processing. Use asynchronous API instead."), __FUNCTION__)
@@ -47,106 +48,6 @@ namespace UE::Mass::Private
 	}
 }
 
-//-----------------------------------------------------------------------------
-// FMassEntityManager::FEntityCreationContext
-//-----------------------------------------------------------------------------
-FMassEntityManager::FEntityCreationContext::FEntityCreationContext()
-	: OwnerThreadId(FPlatformTLS::GetCurrentThreadId())
-{	
-}
-
-FMassEntityManager::FEntityCreationContext::FEntityCreationContext(FMassEntityManager& InManager, const TConstArrayView<FMassEntityHandle> InCreatedEntities)
-	: FEntityCreationContext()
-{
-	CreatedEntities = InCreatedEntities;
-	Manager = InManager.AsShared();
-}
-
-FMassEntityManager::FEntityCreationContext::FEntityCreationContext(FMassEntityManager& InManager, const TConstArrayView<FMassEntityHandle> InCreatedEntities
-	, FMassArchetypeEntityCollection&& InEntityCollection)
-	: FEntityCreationContext(InManager, InCreatedEntities)
-{
-	checkf(InCreatedEntities.IsEmpty() || InEntityCollection.IsEmpty() == false, TEXT("Trying to create FEntityCreationContext instance with no entities but non-empty entity collection. This is not supported."))
-	if (InCreatedEntities.IsEmpty() == false)
-	{
-		EntityCollections.Add(MoveTemp(InEntityCollection));
-	}
-}
-
-FMassEntityManager::FEntityCreationContext::~FEntityCreationContext()
-{
-	if (((EntityCollections.IsEmpty() == false) || (CreatedEntities.IsEmpty() == false)) && ensure(Manager))
-	{
-		Manager->GetObserverManager().OnPostEntitiesCreated(GetEntityCollections());
-	}
-}
-
-TConstArrayView<FMassArchetypeEntityCollection> FMassEntityManager::FEntityCreationContext::GetEntityCollections() const
-{
-	// the EntityCollection has been dirtied, we need to rebuild it
-	if (IsDirty() && ensure(Manager))
-	{
-		UE::Mass::Utils::CreateEntityCollections(*Manager.Get(), CreatedEntities, CollectionCreationDuplicatesHandling, EntityCollections);
-	}
-
-	return EntityCollections;
-}
-
-void FMassEntityManager::FEntityCreationContext::MarkDirty()
-{
-	checkf(OwnerThreadId == FPlatformTLS::GetCurrentThreadId(), TEXT("%hs: all FEntityCreationContext operations ere expected to be run in a single thread"), __FUNCTION__);
-
-	EntityCollections.Reset();
-}
-
-void FMassEntityManager::FEntityCreationContext::AppendEntities(const TConstArrayView<FMassEntityHandle> InEntitiesToAppend)
-{
-	checkf(OwnerThreadId == FPlatformTLS::GetCurrentThreadId(), TEXT("%hs: all FEntityCreationContext operations ere expected to be run in a single thread"), __FUNCTION__);
-
-	if (InEntitiesToAppend.Num())
-	{
-		LLM_SCOPE_BYNAME(TEXT("Mass/EntityManager"));
-		if (CreatedEntities.Num())
-		{
-			// since we already have entities in CreatedEntities (initially ensured to have no duplicates) we cannot 
-			// guarantee anymore that we'll have no duplicates after adding EntitiesToAppend
-			CollectionCreationDuplicatesHandling = FMassArchetypeEntityCollection::FoldDuplicates;
-			MarkDirty();
-		}
-		// else, if there are no entities the resulting state will be "dirty" by design
-		ensureMsgf(EntityCollections.IsEmpty(), TEXT("Having a non-empty array of entity collections is unexpected at this point!"));
-
-		CreatedEntities.Append(InEntitiesToAppend);
-		ensure(IsDirty());
-	}
-}
-
-void FMassEntityManager::FEntityCreationContext::AppendEntities(const TConstArrayView<FMassEntityHandle> InEntitiesToAppend, FMassArchetypeEntityCollection&& InEntityCollection)
-{
-	checkf(OwnerThreadId == FPlatformTLS::GetCurrentThreadId(), TEXT("%hs: all FEntityCreationContext operations are expected to be run in a single thread"), __FUNCTION__);
-
-	if (InEntitiesToAppend.Num() == 0)
-	{
-		return;
-	}
-
-	LLM_SCOPE_BYNAME(TEXT("Mass/EntityManager"));
-
-	AppendEntities(InEntitiesToAppend);
-
-	// this condition boils down to checking if this FEntityCreationContext instance only connects the just added EntitiesToAppend
-	if (CreatedEntities.Num() == InEntitiesToAppend.Num())
-	{
-		checkf(EntityCollections.Num() == 0, TEXT("We never expect EntityCollections to be non-empty while there are no entities in CreatedEntities."));
-		EntityCollections.Add(MoveTemp(InEntityCollection));
-	}
-}
-
-void FMassEntityManager::FEntityCreationContext::ForceUpdateCurrentThreadID()
-{
-	OwnerThreadId = FPlatformTLS::GetCurrentThreadId();
-}
-
 //-----------------------------------------------------------------------------
 // FMassEntityManager
 //-----------------------------------------------------------------------------
@@ -388,10 +289,7 @@ void FMassEntityManager::OnPostFork(EForkProcessRole Role)
 			}
 		}
 
-		if (TSharedPtr<FEntityCreationContext> ActiveContext = ActiveCreationContext.Pin())
-		{
-			ActiveContext->ForceUpdateCurrentThreadID();
-		}
+		ObserverManager.OnPostFork(Role);
 	}
 }
 
@@ -691,8 +589,7 @@ FMassEntityHandle FMassEntityManager::CreateEntity(TConstArrayView<FInstancedStr
 	const FMassEntityHandle Entity = ReserveEntity();
 
 	// Using a creation context to prevent InternalBuildEntity from notifying observers before we set fragments data
-	const TSharedRef<FEntityCreationContext> CreationContext = GetOrMakeCreationContext();
-	CreationContext->AppendEntities({Entity});
+	const TSharedRef<FEntityCreationContext> CreationContext = ObserverManager.GetOrMakeCreationContext();
 
 	InternalBuildEntity(Entity, ArchetypeHandle, SharedFragmentValues);
 
@@ -747,8 +644,7 @@ void FMassEntityManager::BuildEntity(FMassEntityHandle Entity, TConstArrayView<F
 	check(ArchetypeHandle.IsValid());
 
 	// Using a creation context to prevent InternalBuildEntity from notifying observers before we set fragments data
-	const TSharedRef<FEntityCreationContext> CreationContext = GetOrMakeCreationContext();
-	CreationContext->AppendEntities({Entity});
+	const TSharedRef<FEntityCreationContext> CreationContext = ObserverManager.GetOrMakeCreationContext();
 
 	InternalBuildEntity(Entity, ArchetypeHandle, SharedFragmentValues);
 
@@ -774,7 +670,7 @@ int32 FMassEntityManager::BatchReserveEntities(TArrayView<FMassEntityHandle> InO
 TSharedRef<FMassEntityManager::FEntityCreationContext> FMassEntityManager::BatchBuildEntities(const FMassArchetypeEntityCollectionWithPayload& EncodedEntitiesWithPayload
 	, const FMassFragmentBitSet& FragmentsAffected, const FMassArchetypeSharedFragmentValues& SharedFragmentValues, const FMassArchetypeCreationParams& CreationParams)
 {
-	CHECK_SYNC_API_RETURN(MakeShareable(new FEntityCreationContext));
+	CHECK_SYNC_API_RETURN(FMassObserverManager::FCreationContext::DebugCreateDummyCreationContext());
 	check(SharedFragmentValues.IsSorted());
 
 	FMassArchetypeCompositionDescriptor Composition(FragmentsAffected, FMassTagBitSet(), FMassChunkFragmentBitSet(), FMassSharedFragmentBitSet(), FMassConstSharedFragmentBitSet());
@@ -794,7 +690,7 @@ TSharedRef<FMassEntityManager::FEntityCreationContext> FMassEntityManager::Batch
 	, const FMassArchetypeCompositionDescriptor& Composition
 	, const FMassArchetypeSharedFragmentValues& SharedFragmentValues, const FMassArchetypeCreationParams& CreationParams)
 {
-	CHECK_SYNC_API_RETURN(MakeShareable(new FEntityCreationContext));
+	CHECK_SYNC_API_RETURN(FMassObserverManager::FCreationContext::DebugCreateDummyCreationContext());
 
 	TRACE_CPUPROFILER_EVENT_SCOPE(Mass_BatchBuildEntities);
 
@@ -844,14 +740,13 @@ TSharedRef<FMassEntityManager::FEntityCreationContext> FMassEntityManager::Batch
 	// and as such won't cause observers triggering (which usually is prevented by context's existence), and that we 
 	// strongly assume the all entity creation/building (not to be mistaken with "reserving") takes place in a single thread
 	// @todo add checks/ensures enforcing the assumption mentioned above.
-	return GetOrMakeCreationContext(EntityHandles
-		, FMassArchetypeEntityCollection(TargetArchetypeHandle, MoveTemp(TargetArchetypeEntityRanges)));
+	return ObserverManager.GetOrMakeCreationContext(EntityHandles, FMassArchetypeEntityCollection(TargetArchetypeHandle, MoveTemp(TargetArchetypeEntityRanges)));
 }
 
 TSharedRef<FMassEntityManager::FEntityCreationContext> FMassEntityManager::BatchCreateReservedEntities(const FMassArchetypeHandle& ArchetypeHandle
 	, const FMassArchetypeSharedFragmentValues& SharedFragmentValues, TConstArrayView<FMassEntityHandle> ReservedEntities)
 {
-	CHECK_SYNC_API_RETURN(MakeShareable(new FEntityCreationContext));
+	CHECK_SYNC_API_RETURN(FMassObserverManager::FCreationContext::DebugCreateDummyCreationContext());
 	checkf(!ReservedEntities.IsEmpty(), TEXT("No reserved entities given to batch create."));
 
 	TRACE_CPUPROFILER_EVENT_SCOPE(Mass_BatchCreateReservedEntities);
@@ -864,8 +759,8 @@ TSharedRef<FMassEntityManager::FEntityCreationContext> FMassEntityManager::Batch
 TSharedRef<FMassEntityManager::FEntityCreationContext> FMassEntityManager::BatchCreateEntities(const FMassArchetypeHandle& ArchetypeHandle
 	, const FMassArchetypeSharedFragmentValues& SharedFragmentValues, const int32 Count, TArray<FMassEntityHandle>& InOutEntities)
 {
-	CHECK_SYNC_API_RETURN(MakeShareable(new FEntityCreationContext));
-	testableCheckfReturn(ArchetypeHandle.IsValid(), return GetOrMakeCreationContext()
+	CHECK_SYNC_API_RETURN(FMassObserverManager::FCreationContext::DebugCreateDummyCreationContext());
+	testableCheckfReturn(ArchetypeHandle.IsValid(), return FMassObserverManager::FCreationContext::DebugCreateDummyCreationContext()
 		, TEXT("%hs expecting a valid ArchetypeHandle"), __FUNCTION__);
 
 	TRACE_CPUPROFILER_EVENT_SCOPE(Mass_BatchCreateEntities);
@@ -898,7 +793,7 @@ TSharedRef<FMassEntityManager::FEntityCreationContext> FMassEntityManager::Inter
 
 	UE_TRACE_MASS_ENTITIES_CREATED(ReservedEntities, *ArchetypeData)
 
-	return GetOrMakeCreationContext(ReservedEntities, FMassArchetypeEntityCollection(ArchetypeHandle, MoveTemp(TargetArchetypeEntityRanges)));
+	return ObserverManager.GetOrMakeCreationContext(ReservedEntities, FMassArchetypeEntityCollection(ArchetypeHandle, MoveTemp(TargetArchetypeEntityRanges)));
 }
 
 void FMassEntityManager::DestroyEntity(FMassEntityHandle Entity)
@@ -923,7 +818,7 @@ void FMassEntityManager::DestroyEntity(FMassEntityHandle Entity)
 void FMassEntityManager::BatchDestroyEntities(TConstArrayView<FMassEntityHandle> InEntities)
 {
 	CHECK_SYNC_API();
-	checkf(IsDuringEntityCreation() == false, TEXT("%hs: Trying to destroy entities while entity creation is under way. This operation is not supported."), __FUNCTION__);
+	checkf(ObserverManager.IsLocked() == false, TEXT("%hs: Trying to destroy entities while observers are locked - remove-observers won't get triggered in time to read fragments being removed."), __FUNCTION__);
 
 	TRACE_CPUPROFILER_EVENT_SCOPE(Mass_BatchDestroyEntities);
 
@@ -961,7 +856,7 @@ void FMassEntityManager::BatchDestroyEntityChunks(const FMassArchetypeEntityColl
 void FMassEntityManager::BatchDestroyEntityChunks(TConstArrayView<FMassArchetypeEntityCollection> Collections)
 {
 	CHECK_SYNC_API();
-	checkf(IsDuringEntityCreation() == false, TEXT("%hs: Trying to destroy entities while entity creation is under way. This operation is not supported."), __FUNCTION__);
+	checkf(ObserverManager.IsLocked() == false, TEXT("%hs: Trying to destroy entities while observers are locked - remove-observers won't get triggered in time to read fragments being removed."), __FUNCTION__);
 
 	TRACE_CPUPROFILER_EVENT_SCOPE(Mass_BatchDestroyEntityChunks);
 
@@ -971,10 +866,12 @@ void FMassEntityManager::BatchDestroyEntityChunks(TConstArrayView<FMassArchetype
 	for (const FMassArchetypeEntityCollection& EntityCollection : Collections)
 	{
 		EntitiesRemoved.Reset();
-		if (EntityCollection.GetArchetype().IsValid())
+		if (ensureMsgf(EntityCollection.GetArchetype().IsValid() && EntityCollection.IsUpToDate(), TEXT("Provided collection is out of data")))
 		{
 			ObserverManager.OnPreEntitiesDestroyed(ProcessingContext, EntityCollection);
 
+			checkf(EntityCollection.IsUpToDate(), TEXT("Remove-type observers resulted in additional mutating of entity composition. This is not allowed."))
+
 			FMassArchetypeData& ArchetypeData = FMassArchetypeHelper::ArchetypeDataFromHandleChecked(EntityCollection.GetArchetype());
 			ArchetypeData.BatchDestroyEntityChunks(EntityCollection.GetRanges(), EntitiesRemoved);
 		
@@ -1097,10 +994,7 @@ void FMassEntityManager::AddFragmentToEntity(FMassEntityHandle Entity, const USc
 
 	const FMassArchetypeCompositionDescriptor Descriptor(InternalAddFragmentListToEntityChecked(Entity, FMassFragmentBitSet(*FragmentType)));
 
-	if (IsAllowedToTriggerObservers())
-	{
-		ObserverManager.OnPostCompositionAdded(Entity, Descriptor);
-	}
+	ObserverManager.OnPostCompositionAdded(Entity, Descriptor);
 }
 
 void FMassEntityManager::AddFragmentToEntity(FMassEntityHandle Entity, const UScriptStruct* FragmentType, const FStructInitializationCallback& Initializer)
@@ -1118,10 +1012,7 @@ void FMassEntityManager::AddFragmentToEntity(FMassEntityHandle Entity, const USc
 
 	const FMassArchetypeCompositionDescriptor Descriptor(MoveTemp(Fragments));
 	
-	if (IsAllowedToTriggerObservers())
-	{
-		ObserverManager.OnPostCompositionAdded(Entity, Descriptor);
-	}
+	ObserverManager.OnPostCompositionAdded(Entity, Descriptor);
 }
 
 void FMassEntityManager::AddFragmentListToEntity(FMassEntityHandle Entity, TConstArrayView<const UScriptStruct*> FragmentList)
@@ -1132,10 +1023,7 @@ void FMassEntityManager::AddFragmentListToEntity(FMassEntityHandle Entity, TCons
 
 	const FMassArchetypeCompositionDescriptor Descriptor(InternalAddFragmentListToEntityChecked(Entity, FMassFragmentBitSet(FragmentList)));
 	
-	if (IsAllowedToTriggerObservers())
-	{
-		ObserverManager.OnPostCompositionAdded(Entity, Descriptor);
-	}
+	ObserverManager.OnPostCompositionAdded(Entity, Descriptor);
 }
 
 void FMassEntityManager::AddCompositionToEntity_GetDelta(FMassEntityHandle Entity, FMassArchetypeCompositionDescriptor& InOutDescriptor, const FMassArchetypeSharedFragmentValues* AddedSharedFragmentValues)
@@ -1181,10 +1069,7 @@ void FMassEntityManager::AddCompositionToEntity_GetDelta(FMassEntityHandle Entit
 
 			GetEntityStorageInterface().SetArchetypeFromShared(Entity.Index, NewArchetypeHandle.DataPtr);
 
-			if (IsAllowedToTriggerObservers())
-			{
-				ObserverManager.OnPostCompositionAdded(Entity, InOutDescriptor);
-			}
+			ObserverManager.OnPostCompositionAdded(Entity, InOutDescriptor);
 		}
 	}
 }
@@ -1207,10 +1092,7 @@ void FMassEntityManager::RemoveCompositionFromEntity(FMassEntityHandle Entity, c
 
 		if (NewDescriptor.IsEquivalent(OldArchetype->GetCompositionDescriptor()) == false)
 		{
-			if (IsAllowedToTriggerObservers())
-			{
-				ObserverManager.OnPreCompositionRemoved(Entity, InDescriptor);
-			}
+			ObserverManager.OnPreCompositionRemoved(Entity, InDescriptor);
 
 			const FMassArchetypeHandle NewArchetypeHandle = CreateArchetype(NewDescriptor, FMassArchetypeCreationParams(*OldArchetype));
 
@@ -1251,10 +1133,7 @@ void FMassEntityManager::InternalBuildEntity(FMassEntityHandle Entity, const FMa
 	
 	UE_TRACE_MASS_ENTITY_CREATED(Entity, *NewArchetype)
 
-	if (IsAllowedToTriggerObservers())
-	{
-		ObserverManager.OnPostCompositionAdded(Entity, NewArchetype->GetCompositionDescriptor());
-	}
+	ObserverManager.OnPostEntityCreated(Entity, NewArchetype->GetCompositionDescriptor());
 }
 
 void FMassEntityManager::InternalReleaseEntity(FMassEntityHandle Entity)
@@ -1312,10 +1191,7 @@ void FMassEntityManager::AddFragmentInstanceListToEntity(FMassEntityHandle Entit
 	check(CurrentArchetype);
 	CurrentArchetype->SetFragmentsData(Entity, FragmentInstanceList);
 
-	if (IsAllowedToTriggerObservers())
-	{
-		ObserverManager.OnPostCompositionAdded(Entity, Descriptor);
-	}
+	ObserverManager.OnPostCompositionAdded(Entity, Descriptor);
 }
 
 void FMassEntityManager::RemoveFragmentFromEntity(FMassEntityHandle Entity, const UScriptStruct* FragmentType)
@@ -1344,14 +1220,9 @@ void FMassEntityManager::RemoveFragmentListFromEntity(FMassEntityHandle Entity,
 			, OldArchetype->GetConstSharedFragmentBitSet());
 		const FMassArchetypeHandle NewArchetypeHandle = CreateArchetype(NewComposition, FMassArchetypeCreationParams(*OldArchetype));
 
-		FMassArchetypeCompositionDescriptor CompositionDelta;
 		// Find overlap.  It isn't guaranteed that the old archetype has all of the fragments being removed.
-		CompositionDelta.Fragments = OldArchetype->GetFragmentBitSet().GetOverlap(FragmentsToRemove);
-
-		if (IsAllowedToTriggerObservers())
-		{
-			ObserverManager.OnPreCompositionRemoved(Entity, CompositionDelta);
-		}
+		FMassArchetypeCompositionDescriptor CompositionDelta(OldArchetype->GetFragmentBitSet().GetOverlap(FragmentsToRemove));
+		ObserverManager.OnPreCompositionRemoved(Entity, CompositionDelta);
 
 		// Move the entity over
 		FMassArchetypeData& NewArchetype = FMassArchetypeHelper::ArchetypeDataFromHandleChecked(NewArchetypeHandle);
@@ -1410,16 +1281,8 @@ void FMassEntityManager::AddTagToEntity(FMassEntityHandle Entity, const UScriptS
 		// Move the entity over
 		CurrentArchetype->MoveEntityToAnotherArchetype(Entity, *NewArchetypeHandle.DataPtr.Get());
 		GetEntityStorageInterface().SetArchetypeFromShared(Entity.Index, NewArchetypeHandle.DataPtr);
-
-		FMassArchetypeCompositionDescriptor CompositionDelta;
-		FMassTagBitSet TagDelta;
-		TagDelta.Add(*TagType);
-		CompositionDelta.Tags = TagDelta;
 		
-		if (IsAllowedToTriggerObservers())
-		{
-			ObserverManager.OnPostCompositionAdded(Entity, CompositionDelta);
-		}
+		ObserverManager.OnPostCompositionAdded(Entity, FMassArchetypeCompositionDescriptor(FMassTagBitSet(*TagType)));
 	}
 }
 	
@@ -1435,17 +1298,10 @@ void FMassEntityManager::RemoveTagFromEntity(FMassEntityHandle Entity, const USc
 
 	if (CurrentArchetype->HasTagType(TagType))
 	{
-		FMassArchetypeCompositionDescriptor CompositionDelta;
-		FMassTagBitSet TagDelta;
-		TagDelta.Add(*TagType);
-		CompositionDelta.Tags = TagDelta;
-		
-		if (IsAllowedToTriggerObservers())
-		{
-			ObserverManager.OnPreCompositionRemoved(Entity, CompositionDelta);
-		}
-		
+		FMassTagBitSet TagDelta(*TagType);
 		const FMassTagBitSet NewTagComposition = CurrentArchetype->GetTagBitSet() - TagDelta;
+		ObserverManager.OnPreCompositionRemoved(Entity, FMassArchetypeCompositionDescriptor(MoveTemp(TagDelta)));
+
 		const FMassArchetypeHandle NewArchetypeHandle = InternalCreateSimilarArchetype(CurrentArchetype, NewTagComposition);
 		checkSlow(NewArchetypeHandle.IsValid());
 
@@ -1646,8 +1502,6 @@ void FMassEntityManager::BatchChangeTagsForEntities(TConstArrayView<FMassArchety
 	CHECK_SYNC_API();
 
 	TRACE_CPUPROFILER_EVENT_SCOPE(Mass_BatchChangeTagsForEntities);
-
-	const FScopedCreationContextOperations CreationContextOperations(*this);
 	
 	for (const FMassArchetypeEntityCollection& Collection : EntityCollections)
 	{
@@ -1659,16 +1513,13 @@ void FMassEntityManager::BatchChangeTagsForEntities(TConstArrayView<FMassArchety
 		if (ensure(CurrentArchetype) && CurrentArchetype->GetTagBitSet() != NewTagComposition)
 		{
 			FMassTagBitSet TagsAdded = TagsToAdd - CurrentArchetype->GetTagBitSet();
+			const bool bTagsAddedAreObserved = ObserverManager.HasObserversForBitSet(TagsAdded, EMassObservedOperation::Add);
 			FMassTagBitSet TagsRemoved = TagsToRemove.GetOverlap(CurrentArchetype->GetTagBitSet());
-
-			if (CreationContextOperations.IsAllowedToTriggerObservers()
-				&& ObserverManager.HasObserversForBitSet(TagsRemoved, EMassObservedOperation::Remove))
+			if (TagsRemoved.IsEmpty() == false)
 			{
-				// @todo should use OnPreCompositionRemoved here instead, but we're missing a FMassArchetypeEntityCollection version
-				ObserverManager.OnCompositionChanged(Collection, FMassArchetypeCompositionDescriptor(MoveTemp(TagsRemoved)), EMassObservedOperation::Remove);
+				ObserverManager.OnCompositionChanged(Collection, MoveTemp(TagsRemoved), EMassObservedOperation::Remove);
 			}
-			const bool bTagsAddedAreObserved = ObserverManager.HasObserversForBitSet(TagsAdded, EMassObservedOperation::Add);
-
+			
 			FMassArchetypeHandle NewArchetypeHandle = InternalCreateSimilarArchetype(Collection.GetArchetype().DataPtr, NewTagComposition);
 			checkSlow(NewArchetypeHandle.IsValid());
 
@@ -1685,12 +1536,11 @@ void FMassEntityManager::BatchChangeTagsForEntities(TConstArrayView<FMassArchety
 				GetEntityStorageInterface().SetArchetypeFromShared(Entity.Index, NewArchetypeHandle.DataPtr);
 			}
 
-			if (bTagsAddedAreObserved && CreationContextOperations.IsAllowedToTriggerObservers())
+			if (TagsRemoved.IsEmpty() == false)
 			{
-				// @todo should use OnPostCompositionAdded here instead, but we're missing a FMassArchetypeEntityCollection version
 				ObserverManager.OnCompositionChanged(
 					FMassArchetypeEntityCollection(NewArchetypeHandle, MoveTemp(NewArchetypeEntityRanges))
-					, FMassArchetypeCompositionDescriptor(MoveTemp(TagsAdded))
+					, MoveTemp(TagsAdded)
 					, EMassObservedOperation::Add);
 			}
 		}
@@ -1703,8 +1553,6 @@ void FMassEntityManager::BatchChangeFragmentCompositionForEntities(TConstArrayVi
 
 	TRACE_CPUPROFILER_EVENT_SCOPE(Mass_BatchChangeFragmentCompositionForEntities);
 
-	FScopedCreationContextOperations CreationContextOperations(*this);
-
 	for (const FMassArchetypeEntityCollection& Collection : EntityCollections)
 	{
 		FMassArchetypeData* CurrentArchetype = Collection.GetArchetype().DataPtr.Get();
@@ -1720,10 +1568,9 @@ void FMassEntityManager::BatchChangeFragmentCompositionForEntities(TConstArrayVi
 				const bool bFragmentsAddedAreObserved = ObserverManager.HasObserversForBitSet(FragmentsAdded, EMassObservedOperation::Add);
 				FMassFragmentBitSet FragmentsRemoved = FragmentsToRemove.GetOverlap(CurrentArchetype->GetFragmentBitSet());
 				
-				if (CreationContextOperations.IsAllowedToTriggerObservers()
-					&& ObserverManager.HasObserversForBitSet(FragmentsRemoved, EMassObservedOperation::Remove))
+				if (FragmentsRemoved.IsEmpty() == false)
 				{
-					ObserverManager.OnCompositionChanged(Collection, FMassArchetypeCompositionDescriptor(MoveTemp(FragmentsRemoved)), EMassObservedOperation::Remove);
+					ObserverManager.OnCompositionChanged(Collection, MoveTemp(FragmentsRemoved), EMassObservedOperation::Remove);
 				}
 
 				FMassArchetypeHandle NewArchetypeHandle = InternalCreateSimilarArchetype(Collection.GetArchetype().DataPtr, NewFragmentComposition);
@@ -1742,11 +1589,11 @@ void FMassEntityManager::BatchChangeFragmentCompositionForEntities(TConstArrayVi
 					GetEntityStorageInterface().SetArchetypeFromShared(Entity.Index, NewArchetypeHandle.DataPtr);
 				}
 
-				if (bFragmentsAddedAreObserved && CreationContextOperations.IsAllowedToTriggerObservers())
+				if (bFragmentsAddedAreObserved)
 				{
 					ObserverManager.OnCompositionChanged(
 						FMassArchetypeEntityCollection(NewArchetypeHandle, MoveTemp(NewArchetypeEntityRanges))
-						, FMassArchetypeCompositionDescriptor(MoveTemp(FragmentsAdded))
+						, MoveTemp(FragmentsAdded)
 						, EMassObservedOperation::Add);
 				}
 			}
@@ -1775,8 +1622,6 @@ void FMassEntityManager::BatchAddFragmentInstancesForEntities(TConstArrayView<FM
 	//		archetype. This could be utilized as a basis for optimization. To be investigated.
 	// 
 
-	FScopedCreationContextOperations CreationContextOperations(*this);
-
 	for (const FMassArchetypeEntityCollectionWithPayload& EntityRangesWithPayload : EntityCollections)
 	{
 		FMassArchetypeHandle TargetArchetypeHandle = EntityRangesWithPayload.GetEntityCollection().GetArchetype();
@@ -1823,11 +1668,11 @@ void FMassEntityManager::BatchAddFragmentInstancesForEntities(TConstArrayView<FM
 			// corresponds to the order in FMassArchetypeEntityCollectionWithPayload's payload
 			TargetArchetypeHandle.DataPtr->BatchSetFragmentValues(TargetArchetypeEntityRanges, EntityRangesWithPayload.GetPayload());
 			
-			if (bFragmentsAddedAreObserved && CreationContextOperations.IsAllowedToTriggerObservers())
+			if (bFragmentsAddedAreObserved)
 			{
 				ObserverManager.OnCompositionChanged(
 					FMassArchetypeEntityCollection(TargetArchetypeHandle, MoveTemp(TargetArchetypeEntityRanges))
-					, FMassArchetypeCompositionDescriptor(MoveTemp(FragmentsAdded))
+					, MoveTemp(FragmentsAdded)
 					, EMassObservedOperation::Add);
 			}
 		}
@@ -1889,26 +1734,14 @@ void FMassEntityManager::MoveEntityToAnotherArchetype(FMassEntityHandle Entity,
 	FMassArchetypeData* CurrentArchetype = GetEntityStorageInterface().GetArchetype(Entity.Index);
 	check(CurrentArchetype);
 
-	// we'll use this local var to cache whether we should execute OnCompositionChanged(EMassObservedOperation::Add)
-	// operation. We cannot execute it in the following block, because we need to update the entity storage before
-	// triggering the observers.
-	bool bExecutePostMoveObserver = false;
-	
-	if (IsAllowedToTriggerObservers())
-	{
-		const FMassArchetypeCompositionDescriptor CompositionRemoved = CurrentArchetype->GetCompositionDescriptor().CalculateDifference(NewArchetype.GetCompositionDescriptor());
-		ObserverManager.OnCompositionChanged(Entity, CompositionRemoved, EMassObservedOperation::Remove);
-		bExecutePostMoveObserver = true;
-	}
+	FMassArchetypeCompositionDescriptor CompositionRemoved = CurrentArchetype->GetCompositionDescriptor().CalculateDifference(NewArchetype.GetCompositionDescriptor());
+	ObserverManager.OnCompositionChanged(Entity, MoveTemp(CompositionRemoved), EMassObservedOperation::Remove);
 	
 	CurrentArchetype->MoveEntityToAnotherArchetype(Entity, NewArchetype, SharedFragmentValuesOverride);
 	GetEntityStorageInterface().SetArchetypeFromShared(Entity.Index, NewArchetypeHandle.DataPtr);
 
-	if (bExecutePostMoveObserver)
-	{
-		const FMassArchetypeCompositionDescriptor CompositionAdded = NewArchetype.GetCompositionDescriptor().CalculateDifference(CurrentArchetype->GetCompositionDescriptor());
-		ObserverManager.OnCompositionChanged(Entity, CompositionAdded, EMassObservedOperation::Add);
-	}
+	FMassArchetypeCompositionDescriptor CompositionAdded = NewArchetype.GetCompositionDescriptor().CalculateDifference(CurrentArchetype->GetCompositionDescriptor());
+	ObserverManager.OnCompositionChanged(Entity, MoveTemp(CompositionAdded), EMassObservedOperation::Add);
 }
 
 void FMassEntityManager::SetEntityFragmentValues(FMassEntityHandle Entity, TArrayView<const FInstancedStruct> FragmentInstanceList)
@@ -2140,47 +1973,7 @@ void FMassEntityManager::AppendCommands(const TSharedPtr<FMassCommandBuffer>& In
 
 TSharedRef<FMassEntityManager::FEntityCreationContext> FMassEntityManager::GetOrMakeCreationContext()
 {
-	if (ActiveCreationContext.IsValid())
-	{
-		return ActiveCreationContext.Pin().ToSharedRef();
-	}
-	else
-	{
-		FEntityCreationContext* CreationContext = new FEntityCreationContext(*this);
-		TSharedRef<FEntityCreationContext> SharedContext = MakeShareable(CreationContext);
-		ActiveCreationContext = SharedContext;
-		return SharedContext;
-	}
-}
-
-TSharedRef<FMassEntityManager::FEntityCreationContext> FMassEntityManager::GetOrMakeCreationContext(TConstArrayView<FMassEntityHandle> ReservedEntities
-	, FMassArchetypeEntityCollection&& EntityCollection)
-{
-	LLM_SCOPE_BYNAME(TEXT("Mass/EntityManager"));
-	if (ActiveCreationContext.IsValid())
-	{
-		TSharedPtr<FEntityCreationContext> SharedContext = ActiveCreationContext.Pin();
-		CA_ASSUME(SharedContext);
-		SharedContext->AppendEntities(ReservedEntities, MoveTemp(EntityCollection));
-		return SharedContext.ToSharedRef();
-	}
-	else
-	{
-		FEntityCreationContext* CreationContext = new FEntityCreationContext(*this, ReservedEntities, MoveTemp(EntityCollection));
-		TSharedRef<FEntityCreationContext> SharedContext = MakeShareable(CreationContext);
-		ActiveCreationContext = SharedContext;
-		return SharedContext;
-	}
-}
-
-bool FMassEntityManager::DirtyCreationContext()
-{
-	if (TSharedPtr<FEntityCreationContext> AsSharedPtr = ActiveCreationContext.Pin())
-	{
-		AsSharedPtr->MarkDirty();
-		return true;
-	}
-	return false;
+	return ObserverManager.GetOrMakeCreationContext();
 }
 
 UE::Mass::FEntityBuilder FMassEntityManager::MakeEntityBuilder()
@@ -2190,9 +1983,9 @@ UE::Mass::FEntityBuilder FMassEntityManager::MakeEntityBuilder()
 
 bool FMassEntityManager::DebugDoCollectionsOverlapCreationContext(TConstArrayView<FMassArchetypeEntityCollection> EntityCollections) const
 {
-	if (TSharedPtr<FEntityCreationContext> AsSharedPtr = ActiveCreationContext.Pin())
+	if (TSharedPtr<FMassObserverManager::FCreationContext> AsSharedPtr = ObserverManager.GetCreationContext())
 	{
-		TConstArrayView<FMassArchetypeEntityCollection> CreationCollections = AsSharedPtr->EntityCollections;
+		TArray<FMassArchetypeEntityCollection> CreationCollections = AsSharedPtr->GetEntityCollections(*this);
 		return CreationCollections.GetData() <= EntityCollections.GetData()
 			&& EntityCollections.GetData() <= CreationCollections.GetData() + CreationCollections.Num();
 	}
@@ -2375,13 +2168,6 @@ bool FMassEntityManager::DebugHasCommandsToFlush() const
 //-----------------------------------------------------------------------------
 // DEPRECATED
 //-----------------------------------------------------------------------------
-
-const FMassArchetypeEntityCollection& FMassEntityManager::FEntityCreationContext::GetEntityCollection() const
-{
-	static FMassArchetypeEntityCollection EmptyCollection;
-	return EntityCollections.Num() ? EntityCollections[0] : EmptyCollection;
-}
-
 FMassArchetypeHandle FMassEntityManager::InternalCreateSimilarArchetype(const FMassArchetypeData& SourceArchetypeRef, FMassArchetypeCompositionDescriptor&& NewComposition)
 {
 	return InternalCreateSimilarArchetype(SourceArchetypeRef, Forward<FMassArchetypeCompositionDescriptor>(NewComposition), {});
diff --git a/Engine/Source/Runtime/MassEntity/Private/MassObserverManager.cpp b/Engine/Source/Runtime/MassEntity/Private/MassObserverManager.cpp
--- a/Engine/Source/Runtime/MassEntity/Private/MassObserverManager.cpp
+++ b/Engine/Source/Runtime/MassEntity/Private/MassObserverManager.cpp
@@ -7,13 +7,61 @@
 #include "MassProcessingTypes.h"
 #include "MassObserverRegistry.h"
 #include "MassDebugger.h"
+#include "MassEntityCollection.h"
 #include "MassProcessingContext.h"
+#include "MassObserverNotificationTypes.h"
+#include "VisualLogger/VisualLogger.h"
 
 
 #include UE_INLINE_GENERATED_CPP_BY_NAME(MassObserverManager)
 
 namespace UE::Mass::ObserverManager
 {
+	namespace Tweakables
+	{
+		// Used as a template parameter for TInlineAllocator that we use when gathering UScriptStruct* of the observed types to process.
+		constexpr int InlineAllocatorElementsForOverlapTypes = 8;
+	} // Tweakables
+
+	bool bCoalesceBufferedNotifications = false;
+
+	namespace Private
+	{
+		FAutoConsoleVariableRef ConsoleVariables[] =
+		{
+			FAutoConsoleVariableRef(TEXT("mass.observers.CoalesceBufferedNotifications"), bCoalesceBufferedNotifications
+				, TEXT("If enabled, when buffering new notification we'll check if it's the same type as the previously stored one, and if so then merge the two."), ECVF_Default)
+		};
+
+		// a helper function to reduce code duplication in FMassObserverManager::Initialize
+		template<typename TBitSet, typename TPointerType>
+		void SetUpObservers(FMassEntityManager& EntityManager, const EProcessorExecutionFlags WorldExecutionFlags, UObject& Owner
+			, const TMap<TPointerType, FMassProcessorClassCollection>& RegisteredObserverTypes, TBitSet& ObservedBitSet, FMassObserversMap& Observers)
+		{
+			ObservedBitSet.Reset();
+
+			for (auto It : RegisteredObserverTypes)
+			{
+				if (It.Value.ClassCollection.Num() == 0)
+				{
+					continue;
+				}
+
+				ObservedBitSet.Add(*It.Key);
+				FMassRuntimePipeline& Pipeline = (*Observers).FindOrAdd(It.Key);
+
+				for (const TSubclassOf<UMassProcessor>& ProcessorClass : It.Value.ClassCollection)
+				{
+					if (ProcessorClass->GetDefaultObject<UMassProcessor>()->ShouldExecute(WorldExecutionFlags))
+					{
+						Pipeline.AppendProcessor(ProcessorClass, Owner);
+					}
+				}
+				Pipeline.Initialize(Owner, EntityManager.AsShared());
+			}
+		};
+	} // Private
+
 	struct FDeprecationHelper
 	{
 		static void HandleSingleElement(TNotNull<FMassObserverManager*> ObserverManager, const UScriptStruct& ElementType, const FMassArchetypeEntityCollection& EntityCollection, FMassObserversMap& HandlersContainer)
@@ -34,42 +82,117 @@ namespace UE::Mass::ObserverManager
 		}
 	};
 
-	namespace Tweakables
+	struct FNotificationContext
 	{
-		// Used as a template parameter for TInlineAllocator that we use when gathering UScriptStruct* of the observed types to process.
-		constexpr int InlineAllocatorElementsForOverlapTypes = 8;
-	} // Tweakables
+		FMassObserverManager& ObserverManager;
+		FProcessingContext& ProcessingContext;
+	};
 
-	namespace Private
-	{
-	// a helper function to reduce code duplication in FMassObserverManager::Initialize
-	template<typename TBitSet, typename TPointerType>
-	void SetUpObservers(FMassEntityManager& EntityManager, const EProcessorExecutionFlags WorldExecutionFlags, UObject& Owner
-		, const TMap<TPointerType, FMassProcessorClassCollection>& RegisteredObserverTypes, TBitSet& ObservedBitSet, FMassObserversMap& Observers)
+	struct FBufferedNotificationExecutioner 
 	{
-		ObservedBitSet.Reset();
+		FBufferedNotificationExecutioner(FNotificationContext& InNotificationContext, const EMassObservedOperation InOpType)
+			: NotificationContext(InNotificationContext), OpType(InOpType)
+		{}
 
-		for (auto It : RegisteredObserverTypes)
+		template<typename TEntities>
+		void operator()(const FBufferedNotification::FEmptyComposition&, TEntities)
 		{
-			if (It.Value.ClassCollection.Num() == 0)
+			// no-op
+		}
+
+		void operator()(const FMassArchetypeCompositionDescriptor& Change, const FEntityCollection& Entities)
+		{
+			if (Change.Fragments.IsEmpty() == false)
 			{
-				continue;
+				(*this)(Change.Fragments, Entities);
 			}
+			if (Change.Tags.IsEmpty() == false)
+			{
+				(*this)(Change.Tags, Entities);
+			}
+		}
 
-			ObservedBitSet.Add(*It.Key);
-			FMassRuntimePipeline& Pipeline = (*Observers).FindOrAdd(It.Key);
-
-			for (const TSubclassOf<UMassProcessor>& ProcessorClass : It.Value.ClassCollection)
+		void operator()(const FMassArchetypeCompositionDescriptor& Change, const FMassEntityHandle EntityHandle)
+		{
+			if (Change.Fragments.IsEmpty() == false)
 			{
-				if (ProcessorClass->GetDefaultObject<UMassProcessor>()->ShouldExecute(WorldExecutionFlags))
-				{
-					Pipeline.AppendProcessor(ProcessorClass, Owner);
-				}
+				(*this)(Change.Fragments, EntityHandle);
+			}
+			if (Change.Tags.IsEmpty() == false)
+			{
+				(*this)(Change.Tags, EntityHandle);
 			}
-			Pipeline.Initialize(Owner, EntityManager.AsShared());
 		}
+
+		void operator()(const FMassFragmentBitSet& Change, const FEntityCollection& Entities)
+		{
+			ObservedTypesOverlap.Reset();
+			Change.ExportTypes(ObservedTypesOverlap);
+			FMassObserverManager::HandleElementsImpl(NotificationContext.ProcessingContext
+				, Entities.GetUpToDatePerArchetypeCollections(NotificationContext.ObserverManager.EntityManager)
+				, ObservedTypesOverlap
+				, NotificationContext.ObserverManager.FragmentObservers[static_cast<uint8>(OpType)]);
+		}
+
+		void operator()(const FMassFragmentBitSet& Change, const FMassEntityHandle EntityHandle)
+		{
+			ObservedTypesOverlap.Reset();
+			Change.ExportTypes(ObservedTypesOverlap);
+
+			FMassArchetypeHandle ArchetypeHandle = NotificationContext.ObserverManager.GetEntityManager().GetArchetypeForEntity(EntityHandle);
+			NotificationContext.ObserverManager.HandleFragmentsImpl(NotificationContext.ProcessingContext
+				, FMassArchetypeEntityCollection(MoveTemp(ArchetypeHandle), EntityHandle)
+				, ObservedTypesOverlap
+				, NotificationContext.ObserverManager.FragmentObservers[static_cast<uint8>(OpType)]);
+		}
+
+		void operator()(const FMassTagBitSet& Change, const FEntityCollection& Entities)
+		{
+			ObservedTypesOverlap.Reset();
+			Change.ExportTypes(ObservedTypesOverlap);
+			FMassObserverManager::HandleElementsImpl(NotificationContext.ProcessingContext
+				, Entities.GetUpToDatePerArchetypeCollections(NotificationContext.ObserverManager.EntityManager)
+				, ObservedTypesOverlap
+				, NotificationContext.ObserverManager.TagObservers[static_cast<uint8>(OpType)]);
+		}
+
+		void operator()(const FMassTagBitSet& Change, const FMassEntityHandle EntityHandle)
+		{
+			ObservedTypesOverlap.Reset();
+			Change.ExportTypes(ObservedTypesOverlap);
+
+			FMassArchetypeHandle ArchetypeHandle = NotificationContext.ObserverManager.GetEntityManager().GetArchetypeForEntity(EntityHandle);
+			NotificationContext.ObserverManager.HandleFragmentsImpl(NotificationContext.ProcessingContext
+				, FMassArchetypeEntityCollection(MoveTemp(ArchetypeHandle), EntityHandle)
+				, ObservedTypesOverlap
+				, NotificationContext.ObserverManager.TagObservers[static_cast<uint8>(OpType)]);
+		}
+
+		TArray<const UScriptStruct*, TInlineAllocator<ObserverManager::Tweakables::InlineAllocatorElementsForOverlapTypes>> ObservedTypesOverlap;
+		FNotificationContext& NotificationContext;
+		EMassObservedOperation OpType;
+	};
+
+	struct FBufferedCreationNotificationExecutioner
+	{
+		explicit FBufferedCreationNotificationExecutioner(FNotificationContext& InNotificationContext)
+			: NotificationContext(InNotificationContext)
+		{}
+
+		void operator()(FEntityCollection&& Entities) const
+		{
+			NotificationContext.ObserverManager.OnCollectionsCreatedImpl(NotificationContext.ProcessingContext
+				, MoveTemp(Entities).ConsumeArchetypeCollections(NotificationContext.ObserverManager.EntityManager));
+		}
+
+		void operator()(const FMassEntityHandle EntityHandle) const
+		{
+			const FMassArchetypeHandle ArchetypeHandle = NotificationContext.ObserverManager.GetEntityManager().GetArchetypeForEntity(EntityHandle);
+			const FMassArchetypeCompositionDescriptor& ArchetypeComposition = NotificationContext.ProcessingContext.GetEntityManager()->GetArchetypeComposition(ArchetypeHandle);
+			NotificationContext.ObserverManager.OnCompositionChanged(EntityHandle, ArchetypeComposition, EMassObservedOperation::Add, &NotificationContext.ProcessingContext);
+		}
+		FNotificationContext& NotificationContext;
 	};
-	} // Private
 } // UE::Mass::ObserverManager
 
 //----------------------------------------------------------------------//
@@ -92,6 +215,8 @@ void FMassObserversMap::DebugAddUniqueProcessors(TArray<const UMassProcessor*>&
 //----------------------------------------------------------------------//
 // FMassObserverManager
 //----------------------------------------------------------------------//
+FMassArchetypeEntityCollection FMassObserverManager::FCollectionRefOrHandle::DummyCollection;
+
 FMassObserverManager::FMassObserverManager()
 	: EntityManager(GetMutableDefault<UMassEntitySubsystem>()->GetMutableEntityManager())
 {
@@ -156,54 +281,50 @@ void FMassObserverManager::DeInitialize()
 
 bool FMassObserverManager::OnPostEntitiesCreated(const FMassArchetypeEntityCollection& EntityCollection)
 {
-	return OnPostEntitiesCreated(MakeArrayView(&EntityCollection , 1));
-}
+	TRACE_CPUPROFILER_EVENT_SCOPE(MassObserver_OnPostEntitiesCreated_Collection);
 
-bool FMassObserverManager::OnPostEntitiesCreated(TConstArrayView<FMassArchetypeEntityCollection> EntityCollections)
-{
-	UE::Mass::FProcessingContext ProcessingContext(EntityManager, /*DeltaSeconds=*/0.f);
-	return OnPostEntitiesCreated(ProcessingContext, EntityCollections);
-}
+	if (LocksCount)
+	{
+		TSharedRef<FObserverLock> ObserverLock = ActiveObserverLock.Pin().ToSharedRef();
+		ObserverLock->AddCreatedEntitiesCollection(EntityCollection);
+		return false;
+	}
 
-// DEPRECATED
-bool FMassObserverManager::OnPostEntitiesCreated(UE::Mass::FProcessingContext& ProcessingContext, const FMassArchetypeEntityCollection& EntityCollection)
-{
-	return OnPostEntitiesCreated(ProcessingContext, MakeArrayView(&EntityCollection, 1));
+	const FMassArchetypeCompositionDescriptor& ArchetypeComposition = EntityManager.GetArchetypeComposition(EntityCollection.GetArchetype());
+	return OnCompositionChanged(EntityCollection, ArchetypeComposition, EMassObservedOperation::Add);
 }
 
-bool FMassObserverManager::OnPostEntitiesCreated(UE::Mass::FProcessingContext& ProcessingContext, TConstArrayView<FMassArchetypeEntityCollection> EntityCollections)
+bool FMassObserverManager::OnPostEntityCreated(FMassEntityHandle EntityHandle, const FMassArchetypeCompositionDescriptor& Composition)
 {
-	if (!ensureMsgf(EntityManager.IsDuringEntityCreation() == false, TEXT("Entity creation not done yet.")))
+	TRACE_CPUPROFILER_EVENT_SCOPE(MassObserver_OnPostEntitiesCreated_Collection);
+
+	if (LocksCount)
 	{
+		TSharedRef<FObserverLock> ObserverLock = ActiveObserverLock.Pin().ToSharedRef();
+		ObserverLock->AddCreatedEntity(EntityHandle);
 		return false;
 	}
 
-	TRACE_CPUPROFILER_EVENT_SCOPE_STR("OnPostEntitiesCreated")
-
-	bool bReturnValue = false;
-
-	for (const FMassArchetypeEntityCollection& Collection : EntityCollections)
+	if (Composition.IsEmpty())
 	{
-		const FMassArchetypeCompositionDescriptor& ArchetypeComposition = ProcessingContext.GetEntityManager()->GetArchetypeComposition(Collection.GetArchetype());
-		bReturnValue = OnCompositionChanged(ProcessingContext, Collection, ArchetypeComposition, EMassObservedOperation::Add) || bReturnValue;
+		const FMassArchetypeHandle ArchetypeHandle = EntityManager.GetArchetypeForEntity(EntityHandle);
+		const FMassArchetypeCompositionDescriptor& ArchetypeComposition = EntityManager.GetArchetypeComposition(ArchetypeHandle);
+		return OnCompositionChanged(EntityHandle, ArchetypeComposition, EMassObservedOperation::Add);
 	}
 
-	return bReturnValue;
+	return OnCompositionChanged(EntityHandle, Composition, EMassObservedOperation::Add);
 }
 
 bool FMassObserverManager::OnPreEntitiesDestroyed(const FMassArchetypeEntityCollection& EntityCollection)
 {
-	UE::Mass::FProcessingContext ProcessingContext(EntityManager);
-	return OnPreEntitiesDestroyed(ProcessingContext, EntityCollection);
+	const FMassArchetypeCompositionDescriptor& ArchetypeComposition = EntityManager.GetArchetypeComposition(EntityCollection.GetArchetype());
+	return OnCompositionChanged(EntityCollection, ArchetypeComposition, EMassObservedOperation::Remove);
 }
 
 bool FMassObserverManager::OnPreEntitiesDestroyed(UE::Mass::FProcessingContext& ProcessingContext, const FMassArchetypeEntityCollection& EntityCollection)
 {
-	TRACE_CPUPROFILER_EVENT_SCOPE_STR("OnPreEntitiesDestroyed")
-
-	const FMassArchetypeCompositionDescriptor& ArchetypeComposition = ProcessingContext.GetEntityManager()->GetArchetypeComposition(EntityCollection.GetArchetype());
-	
-	return OnCompositionChanged(ProcessingContext, EntityCollection, ArchetypeComposition, EMassObservedOperation::Remove);
+	const FMassArchetypeCompositionDescriptor& ArchetypeComposition = EntityManager.GetArchetypeComposition(EntityCollection.GetArchetype());	
+	return OnCompositionChanged(EntityCollection, ArchetypeComposition, EMassObservedOperation::Remove, &ProcessingContext);
 }
 
 bool FMassObserverManager::OnPreEntityDestroyed(const FMassArchetypeCompositionDescriptor& ArchetypeComposition, const FMassEntityHandle Entity)
@@ -212,88 +333,203 @@ bool FMassObserverManager::OnPreEntityDestroyed(const FMassArchetypeCompositionD
 	return OnCompositionChanged(Entity, ArchetypeComposition, EMassObservedOperation::Remove);
 }
 
-bool FMassObserverManager::OnCompositionChanged(UE::Mass::FProcessingContext& ProcessingContext, const FMassArchetypeEntityCollection& EntityCollection, const FMassArchetypeCompositionDescriptor& CompositionDelta, const EMassObservedOperation Operation)
+bool FMassObserverManager::OnCompositionChanged(FCollectionRefOrHandle&& EntityCollection, const FMassArchetypeCompositionDescriptor& CompositionDelta
+	, const EMassObservedOperation Operation, UE::Mass::FProcessingContext* ProcessingContext)
 {
 	using UE::Mass::ObserverManager::Tweakables::InlineAllocatorElementsForOverlapTypes;
+	ensureMsgf(EntityCollection.EntityHandle.IsValid() || EntityCollection.EntityCollection.IsUpToDate()
+		, TEXT("Out-of-date FMassArchetypeEntityCollection used. Stored information is unreliable."));
+
+	if (CompositionDelta.IsEmpty())
+	{
+		// nothing to do here.
+		// @todo calling this function just to bail out would be a lot cheaper if we didn't have to create
+		// FMassArchetypeCompositionDescriptor instances first - we usually just pass in tag or fragment bitsets.
+		// like in FMassEntityManager::BatchChangeTagsForEntities
+		return false;
+	}
+
+	TRACE_CPUPROFILER_EVENT_SCOPE(MassObserver_OnCompositionChanged);
+
+	if (LocksCount)
+	{
+		if (TSharedPtr<FCreationContext> CreationContext = GetCreationContext())
+		{
+			// a composition mutating operation is taking place, while creation lock is active - this operation invalidates the stored collections
+			CreationContext->MarkDirty();
+			return false;
+		}
+		UE_CVLOG_UELOG(Operation == EMassObservedOperation::Remove, EntityManager.GetOwner(), LogMass, Log
+			, TEXT("%hs: Remove operation while observers are locked - the remove-observer will be executed after the data has already been removed."), __FUNCTION__);
+	}
 
-	ensureMsgf(EntityCollection.IsUpToDate(), TEXT("Out-of-date FMassArchetypeEntityCollection used. Stored information is unreliable."));
+	const int32 OperationIndex = static_cast<int32>(Operation);
 
-	const FMassFragmentBitSet FragmentOverlap = ObservedFragments[(uint8)Operation].GetOverlap(CompositionDelta.Fragments);
+	FMassFragmentBitSet FragmentOverlap = ObservedFragments[OperationIndex].GetOverlap(CompositionDelta.Fragments);
+	FMassTagBitSet TagOverlap = ObservedTags[OperationIndex].GetOverlap(CompositionDelta.Tags);
 	const bool bHasFragmentsOverlap = !FragmentOverlap.IsEmpty();
-	const FMassTagBitSet TagOverlap = ObservedTags[(uint8)Operation].GetOverlap(CompositionDelta.Tags);
 	const bool bHasTagsOverlap = !TagOverlap.IsEmpty();
 
 	if (bHasFragmentsOverlap || bHasTagsOverlap)
 	{
-		TArray<const UScriptStruct*, TInlineAllocator<InlineAllocatorElementsForOverlapTypes>> ObservedTypesOverlap;
-
-		if (bHasFragmentsOverlap)
+		if (LocksCount)
 		{
-			FragmentOverlap.ExportTypes(ObservedTypesOverlap);
+			const UE::Mass::ObserverManager::EObservedOperationNotification NotificationType = static_cast<UE::Mass::ObserverManager::EObservedOperationNotification>(Operation);
+			TSharedRef<FObserverLock> ObserverLockRef = ActiveObserverLock.Pin().ToSharedRef();
 
-			HandleElementsImpl(ProcessingContext, {EntityCollection}, ObservedTypesOverlap, FragmentObservers[(uint8)Operation]);
-		}
+			// UE::Mass::FEntityCollection(EntityCollection) OR EntityHandle
+			if (UE::Mass::ObserverManager::bCoalesceBufferedNotifications)
+			{
+				TRACE_CPUPROFILER_EVENT_SCOPE(MassObserver_Notify_Coalesced);
 
-		if (bHasTagsOverlap)
-		{
-			ObservedTypesOverlap.Reset();
-			TagOverlap.ExportTypes(ObservedTypesOverlap);
+				if (EntityCollection.EntityHandle.IsSet())
+				{
+					ObserverLockRef->AddNotification(NotificationType, EntityCollection.EntityHandle
+						, bHasFragmentsOverlap, MoveTemp(FragmentOverlap)
+						, bHasTagsOverlap, MoveTemp(TagOverlap));
+				}
+				else
+				{
+					ObserverLockRef->AddNotification(NotificationType, EntityCollection.EntityCollection
+						, bHasFragmentsOverlap, MoveTemp(FragmentOverlap) 
+						, bHasTagsOverlap, MoveTemp(TagOverlap));
+				}
+			}
+			else
+			{
+				TRACE_CPUPROFILER_EVENT_SCOPE(MassObserver_Notify_Emplace);
+
+				FBufferedNotification::FEntitiesContainer Entities;
+				if (EntityCollection.EntityHandle.IsSet())
+				{
+					Entities.Emplace<FMassEntityHandle>(EntityCollection.EntityHandle);
+				}
+				else
+				{
+					Entities.Emplace<UE::Mass::FEntityCollection>(EntityCollection.EntityCollection);
+				}
 
-			HandleElementsImpl(ProcessingContext, {EntityCollection}, ObservedTypesOverlap, TagObservers[(uint8)Operation]);
+				if (bHasFragmentsOverlap && bHasTagsOverlap)
+				{
+					FMassArchetypeCompositionDescriptor ChangedComposition(MoveTemp(FragmentOverlap), MoveTemp(TagOverlap), {}, {}, {});
+					ObserverLockRef->BufferedNotifications.Emplace(NotificationType, MoveTemp(ChangedComposition), MoveTemp(Entities));
+				}
+				else if (bHasFragmentsOverlap)
+				{
+					ObserverLockRef->BufferedNotifications.Emplace(NotificationType, MoveTemp(FragmentOverlap), MoveTemp(Entities));
+				}
+				else // bHasTagsOverlap
+				{
+					ObserverLockRef->BufferedNotifications.Emplace(NotificationType, MoveTemp(TagOverlap), MoveTemp(Entities));
+				}
+			}
 		}
+		else
+		{
+			auto HandleElements = [&](const FMassArchetypeEntityCollection& Collection)
+			{
+				TArray<const UScriptStruct*, TInlineAllocator<InlineAllocatorElementsForOverlapTypes>> ObservedTypesOverlap;
 
-		return true;
+				uint8 LocalContextBuffer[sizeof(UE::Mass::FProcessingContext)];
+				UE::Mass::FProcessingContext* LocalProcessingContext = (ProcessingContext == nullptr)
+					? new(&LocalContextBuffer) UE::Mass::FProcessingContext(EntityManager)
+					: ProcessingContext;
+
+				if (bHasFragmentsOverlap)
+				{
+					FragmentOverlap.ExportTypes(ObservedTypesOverlap);
+
+					HandleElementsImpl(*LocalProcessingContext, {Collection}, ObservedTypesOverlap, FragmentObservers[OperationIndex]);
+				}
+
+				if (bHasTagsOverlap)
+				{
+					ObservedTypesOverlap.Reset();
+					TagOverlap.ExportTypes(ObservedTypesOverlap);
+
+					HandleElementsImpl(*LocalProcessingContext, {Collection}, ObservedTypesOverlap, TagObservers[OperationIndex]);
+				}
+
+				if (ProcessingContext == nullptr)
+				{
+					LocalProcessingContext->~FProcessingContext();
+				}
+			};
+
+			if (EntityCollection.EntityHandle.IsSet())
+			{
+				const FMassArchetypeHandle ArchetypeHandle = EntityManager.GetArchetypeForEntity(EntityCollection.EntityHandle);
+				HandleElements(FMassArchetypeEntityCollection(ArchetypeHandle, EntityCollection.EntityHandle));
+			}
+			else
+			{
+				HandleElements(EntityCollection.EntityCollection);
+			}
+
+			return true;
+		}
 	}
 
 	return false;
 }
 
-bool FMassObserverManager::OnCompositionChanged(const FMassArchetypeEntityCollection& EntityCollection, const FMassArchetypeCompositionDescriptor& Composition, const EMassObservedOperation Operation)
-{
-	UE::Mass::FProcessingContext LocalContext(EntityManager, /*DeltaSeconds=*/0.f);
-	return OnCompositionChanged(LocalContext, EntityCollection, Composition, Operation);
-}
-
-bool FMassObserverManager::OnCompositionChanged(const FMassEntityHandle Entity, const FMassArchetypeCompositionDescriptor& CompositionDelta, const EMassObservedOperation Operation)
+bool FMassObserverManager::OnCollectionsCreatedImpl(UE::Mass::FProcessingContext& ProcessingContext, TArray<FMassArchetypeEntityCollection>&& EntityCollections)
 {
 	using UE::Mass::ObserverManager::Tweakables::InlineAllocatorElementsForOverlapTypes;
 
-	const FMassFragmentBitSet FragmentOverlap = ObservedFragments[(uint8)Operation].GetOverlap(CompositionDelta.Fragments);
+	TRACE_CPUPROFILER_EVENT_SCOPE(MassObserver_OnCollectionsCreatedImpl_Collection);
+
+	check(LocksCount == 0);
+
+	constexpr int32 OperationIndex = static_cast<int32>(EMassObservedOperation::Add);
+
+	FMassFragmentBitSet FragmentOverlap;
+	FMassTagBitSet TagOverlap;
+
+	for (FMassArchetypeEntityCollection Collection : EntityCollections)
+	{
+		checkfSlow(Collection.IsUpToDate(), TEXT("Out-of-date FMassArchetypeEntityCollection used. Stored information is unreliable."));
+
+		const FMassArchetypeCompositionDescriptor& ArchetypeComposition = EntityManager.GetArchetypeComposition(Collection.GetArchetype());
+		FragmentOverlap += ArchetypeComposition.Fragments;
+		TagOverlap += ArchetypeComposition.Tags;
+	}
+	FragmentOverlap = ObservedFragments[OperationIndex].GetOverlap(FragmentOverlap);
+	TagOverlap = ObservedTags[OperationIndex].GetOverlap(TagOverlap);
+
 	const bool bHasFragmentsOverlap = !FragmentOverlap.IsEmpty();
-	const FMassTagBitSet TagOverlap = ObservedTags[(uint8)Operation].GetOverlap(CompositionDelta.Tags);
 	const bool bHasTagsOverlap = !TagOverlap.IsEmpty();
-
 	if (bHasFragmentsOverlap || bHasTagsOverlap)
 	{
 		TArray<const UScriptStruct*, TInlineAllocator<InlineAllocatorElementsForOverlapTypes>> ObservedTypesOverlap;
-		UE::Mass::FProcessingContext ProcessingContext(EntityManager, /*DeltaSeconds=*/0.f);
-		ProcessingContext.bFlushCommandBuffer = false;
-		const FMassArchetypeHandle ArchetypeHandle = EntityManager.GetArchetypeForEntity(Entity);
-		FMassArchetypeEntityCollection EntityCollection(ArchetypeHandle, {Entity}, FMassArchetypeEntityCollection::NoDuplicates);
 
 		if (bHasFragmentsOverlap)
 		{
 			FragmentOverlap.ExportTypes(ObservedTypesOverlap);
-			HandleElementsImpl(ProcessingContext, {EntityCollection}, ObservedTypesOverlap, FragmentObservers[(uint8)Operation]);
+			HandleElementsImpl(ProcessingContext, EntityCollections, ObservedTypesOverlap, FragmentObservers[OperationIndex]);
 		}
 
 		if (bHasTagsOverlap)
 		{
 			ObservedTypesOverlap.Reset();
 			TagOverlap.ExportTypes(ObservedTypesOverlap);
-			HandleElementsImpl(ProcessingContext, {EntityCollection}, ObservedTypesOverlap, TagObservers[(uint8)Operation]);
+			HandleElementsImpl(ProcessingContext, EntityCollections, ObservedTypesOverlap, TagObservers[OperationIndex]);
 		}
-	}
 
-	return bHasFragmentsOverlap || bHasTagsOverlap;
+		return true;
+	}
+	return false;
 }
 
 void FMassObserverManager::HandleElementsImpl(UE::Mass::FProcessingContext& ProcessingContext, TConstArrayView<FMassArchetypeEntityCollection> EntityCollections
-	, TArrayView<const UScriptStruct*> ObservedTypes, FMassObserversMap& HandlersContainer)
+		, TArrayView<const UScriptStruct*> ObservedTypes, FMassObserversMap& HandlersContainer)
 {	
-	TRACE_CPUPROFILER_EVENT_SCOPE(MassObserver_HandleElementsImpl);
+	TRACE_CPUPROFILER_EVENT_SCOPE(MassObserver_HandleFragmentsImpl);
+
+	check(ObservedTypes.Num());
+	ensureMsgf(EntityCollections.Num(), TEXT("Empty collections array is unexpected at this point. Nothing bad will happen, but it's a waste of perf."));
 
-	check(ObservedTypes.Num() > 0);
+	FMassEntityManager::FScopedProcessing ProcessingScope = ProcessingContext.EntityManager->NewProcessingScope();
 
 	for (const UScriptStruct* Type : ObservedTypes)
 	{		
@@ -313,13 +549,13 @@ void FMassObserverManager::AddObserverInstance(const UScriptStruct& ElementType,
 
 	if (bIsFragment)
 	{
-		Pipeline = &(*FragmentObservers[(uint8)Operation]).FindOrAdd(&ElementType);
-		ObservedFragments[(uint8)Operation].Add(ElementType);
+		Pipeline = &(*FragmentObservers[static_cast<uint8>(Operation)]).FindOrAdd(&ElementType);
+		ObservedFragments[static_cast<uint8>(Operation)].Add(ElementType);
 	}
 	else
 	{
-		Pipeline = &(*TagObservers[(uint8)Operation]).FindOrAdd(&ElementType);
-		ObservedTags[(uint8)Operation].Add(ElementType);
+		Pipeline = &(*TagObservers[static_cast<uint8>(Operation)]).FindOrAdd(&ElementType);
+		ObservedTags[static_cast<uint8>(Operation)].Add(ElementType);
 	}
 
 	// AppendUniqueProcessor will return true only if ObserverProcessor has not been a part of Pipeline yet.
@@ -367,9 +603,159 @@ void FMassObserverManager::RemoveObserverInstance(const UScriptStruct& ElementTy
 	}
 }
 
+TSharedRef<FMassObserverManager::FObserverLock> FMassObserverManager::GetOrMakeObserverLock()
+{
+	if (ActiveObserverLock.IsValid())
+	{
+		return ActiveObserverLock.Pin().ToSharedRef();
+	}
+	else
+	{
+		FObserverLock* ObserverLock = new FObserverLock(*this);
+		TSharedRef<FObserverLock> SharedContext = MakeShareable(ObserverLock);
+		ActiveObserverLock = SharedContext;
+		return SharedContext;
+	}
+}
+
+TSharedRef<FMassObserverManager::FCreationContext> FMassObserverManager::GetOrMakeCreationContext()
+{
+	if (ActiveCreationContext.IsValid())
+	{
+		return ActiveCreationContext.Pin().ToSharedRef();
+	}
+	else
+	{
+		FCreationContext* ObserverLock = new FCreationContext(GetOrMakeObserverLock());
+#if WITH_MASSENTITY_DEBUG
+		ObserverLock->CreationHandle.SerialNumber = LockedNotificationSerialNumber;
+#endif // WITH_MASSENTITY_DEBUG
+		ObserverLock->CreationHandle.OpIndex = ObserverLock->Lock->AddCreatedEntities({});
+		TSharedRef<FCreationContext> SharedContext = MakeShareable(ObserverLock);
+		ActiveCreationContext = SharedContext;
+		return SharedContext;
+	}
+}
+
+TSharedRef<FMassObserverManager::FCreationContext> FMassObserverManager::GetOrMakeCreationContext(TConstArrayView<FMassEntityHandle> ReservedEntities
+	, FMassArchetypeEntityCollection&& EntityCollection)
+{
+	if (TSharedPtr<FCreationContext> CreationContext = ActiveCreationContext.Pin())
+	{
+		CreationContext->GetObserverLock()->AddCreatedEntities(ReservedEntities, Forward<FMassArchetypeEntityCollection>(EntityCollection));
+		return CreationContext.ToSharedRef();
+	}
+	else
+	{
+		FCreationContext* ObserverLock = new FCreationContext(GetOrMakeObserverLock());
+#if WITH_MASSENTITY_DEBUG
+		ObserverLock->CreationHandle.SerialNumber = LockedNotificationSerialNumber;
+#endif // WITH_MASSENTITY_DEBUG
+		ObserverLock->CreationHandle.OpIndex = ObserverLock->Lock->AddCreatedEntities(ReservedEntities, Forward<FMassArchetypeEntityCollection>(EntityCollection));
+		TSharedRef<FCreationContext> SharedContext = MakeShareable(ObserverLock);
+		ActiveCreationContext = SharedContext;
+		return SharedContext;
+	}
+}
+
+void FMassObserverManager::OnPostFork(EForkProcessRole)
+{
+	if (TSharedPtr<FObserverLock> ActiveContext = ActiveObserverLock.Pin())
+	{
+		ActiveContext->ForceUpdateCurrentThreadID();
+	}
+}
+
+void FMassObserverManager::ResumeExecution(FObserverLock& LockBeingReleased)
+{
+	using namespace UE::Mass::ObserverManager;
+
+	ensureMsgf(LocksCount == 0, TEXT("We only expect this function to be called if all locks are released."));
+#if WITH_MASSENTITY_DEBUG
+	ensureMsgf(LockBeingReleased.LockSerialNumber == LockedNotificationSerialNumber
+		, TEXT("Lock's and ObserverManager's lock serial numbers are expected to match."));
+	++LockedNotificationSerialNumber;
+#endif // WITH_MASSENTITY_DEBUG
+
+	if (LockBeingReleased.BufferedNotifications.IsEmpty() == false)
+	{
+		TRACE_CPUPROFILER_EVENT_SCOPE(MassObserver_ResumeExecution);
+
+		UE::Mass::FProcessingContext ProcessingContext(EntityManager);
+
+		FNotificationContext NotificationContext{*this, ProcessingContext};
+		FBufferedNotificationExecutioner AddNotification(NotificationContext, EMassObservedOperation::Add);
+		FBufferedNotificationExecutioner RemoveNotification(NotificationContext, EMassObservedOperation::Remove);
+		FBufferedCreationNotificationExecutioner CreationOpExecutioner(NotificationContext);
+
+		for (FBufferedNotification& Op : LockBeingReleased.BufferedNotifications)
+		{
+			switch (Op.Type)
+			{
+				case EObservedOperationNotification::Add:
+					Visit(AddNotification, Op.CompositionChange, Op.AffectedEntities);
+					break;
+				case EObservedOperationNotification::Remove:
+					Visit(RemoveNotification, Op.CompositionChange, Op.AffectedEntities);
+					break;
+				case EObservedOperationNotification::Create:
+					Visit(CreationOpExecutioner, MoveTemp(Op.AffectedEntities));
+					break;
+				default:
+					ensureMsgf(false, TEXT("%hs: Unhandled EObservedOperationNotification value"), __FUNCTION__);
+			}
+		}
+#if WITH_MASSENTITY_DEBUG
+		++DebugNonTrivialResumeExecutionCount;
+#endif // WITH_MASSENTITY_DEBUG
+	}
+}
+
+void FMassObserverManager::ReleaseCreationHandle(FCreationNotificationHandle InCreationNotificationHandle)
+{
+	ensureMsgf(InCreationNotificationHandle.IsSet(), TEXT("Invalid creation handle passed to %hs"), __FUNCTION__);
+#if WITH_MASSENTITY_DEBUG
+	ensureMsgf(InCreationNotificationHandle.SerialNumber == LockedNotificationSerialNumber
+		, TEXT("Creation handle's serial number doesn't match the ObserverManager's data"));
+#endif // WITH_MASSENTITY_DEBUG
+
+	TSharedPtr<FObserverLock> LockInstance = ActiveObserverLock.Pin();
+	if (ensure(LockInstance))
+	{
+		ensure(LockInstance->ReleaseCreationNotification(InCreationNotificationHandle));
+		ensure(ActiveCreationContext.IsValid() == false);
+	}
+}
+
 //----------------------------------------------------------------------//
 // DEPRECATED
 //----------------------------------------------------------------------//
+bool FMassObserverManager::OnPostEntitiesCreated(UE::Mass::FProcessingContext&, const FMassArchetypeEntityCollection& EntityCollection)
+{
+	return OnPostEntitiesCreated(EntityCollection);
+}
+
+bool FMassObserverManager::OnPostEntitiesCreated(UE::Mass::FProcessingContext&, TConstArrayView<FMassArchetypeEntityCollection> EntityCollections)
+{
+	TRACE_CPUPROFILER_EVENT_SCOPE_STR("OnPostEntitiesCreated")
+
+	bool bReturnValue = false;
+
+	for (const FMassArchetypeEntityCollection& Collection : EntityCollections)
+	{
+		const FMassArchetypeCompositionDescriptor& ArchetypeComposition = EntityManager.GetArchetypeComposition(Collection.GetArchetype());
+		bReturnValue |= OnCompositionChanged(Collection, ArchetypeComposition, EMassObservedOperation::Add);
+	}
+
+	return bReturnValue;
+}
+
+bool FMassObserverManager::OnCompositionChanged(UE::Mass::FProcessingContext&, const FMassArchetypeEntityCollection& EntityCollection
+	, const FMassArchetypeCompositionDescriptor& CompositionDelta, const EMassObservedOperation InOperation)
+{
+	return OnCompositionChanged(EntityCollection, CompositionDelta, InOperation);
+}
+
 void FMassObserverManager::HandleSingleEntityImpl(const UScriptStruct& FragmentType, const FMassArchetypeEntityCollection& EntityCollection, FMassObserversMap& HandlersContainer)
 {
 	UE::Mass::ObserverManager::FDeprecationHelper::HandleSingleElement(this, FragmentType, EntityCollection, HandlersContainer);
@@ -390,8 +776,35 @@ void FMassObserverManager::OnFragmentOrTagOperation(const UScriptStruct& Fragmen
 	UE::Mass::ObserverManager::FDeprecationHelper::HandleSingleElement(this, FragmentOrTagType, EntityCollection, Operation);
 }
 
-void FMassObserverManager::HandleFragmentsImpl(UE::Mass::FProcessingContext& ProcessingContext, const FMassArchetypeEntityCollection& EntityCollection
-		, TArrayView<const UScriptStruct*> ObservedTypes, FMassObserversMap& HandlersContainer)
+TConstArrayView<FMassArchetypeEntityCollection> FMassObserverManager::FCreationContext::GetEntityCollections() const
+{
+	return {};
+}
+
+int32 FMassObserverManager::FCreationContext::GetSpawnedNum() const
+{
+	return 0;
+}
+
+bool FMassObserverManager::FCreationContext::IsDirty() const
+{
+	return true;
+}
+
+void FMassObserverManager::FCreationContext::AppendEntities(const TConstArrayView<FMassEntityHandle>)
+{
+}
+
+void FMassObserverManager::FCreationContext::AppendEntities(const TConstArrayView<FMassEntityHandle>, FMassArchetypeEntityCollection&&)
+{
+}
+
+FMassObserverManager::FCreationContext::FCreationContext(const int32)
+	: FCreationContext()
+{}
+
+const FMassArchetypeEntityCollection& FMassObserverManager::FCreationContext::GetEntityCollection() const
 {
-	HandleElementsImpl(ProcessingContext, {EntityCollection}, ObservedTypes, HandlersContainer);
+	static FMassArchetypeEntityCollection DummyInstance;
+	return DummyInstance;
 }
diff --git a/Engine/Source/Runtime/MassEntity/Private/MassObserverNotificationTypes.cpp b/Engine/Source/Runtime/MassEntity/Private/MassObserverNotificationTypes.cpp
new file mode 100644
--- /dev/null
+++ b/Engine/Source/Runtime/MassEntity/Private/MassObserverNotificationTypes.cpp
@@ -0,0 +1,116 @@
+// Copyright Epic Games, Inc. All Rights Reserved.
+
+#include "MassObserverNotificationTypes.h"
+#include "MassObserverManager.h"
+#include "MassEntityManager.h"
+#include "MassEntityUtils.h"
+
+namespace UE::Mass::ObserverManager
+{
+	namespace Private
+	{
+		TSharedRef<FObserverLock> GetDummyObserverLock()
+		{
+			static TSharedRef<FObserverLock> DummyObserverLock = MakeShareable(new FObserverLock());
+			return DummyObserverLock;
+		}
+	} // Private
+
+	//-----------------------------------------------------------------------------
+	// FObserverLock
+	//-----------------------------------------------------------------------------
+	FObserverLock::FObserverLock(FMassObserverManager& ObserverManager)
+		: OwnerThreadId(FPlatformTLS::GetCurrentThreadId())
+		, WeakEntityManager(ObserverManager.EntityManager.AsWeak())
+	#if WITH_MASSENTITY_DEBUG
+		, LockSerialNumber(ObserverManager.LockedNotificationSerialNumber)
+	#endif // WITH_MASSENTITY_DEBUG
+	{
+		++ObserverManager.LocksCount;
+	}
+
+	FObserverLock::~FObserverLock()
+	{
+		TSharedPtr<FMassEntityManager> SharedEntityManager = WeakEntityManager.Pin();
+		if (UNLIKELY(!SharedEntityManager))
+		{
+			return;
+		}
+
+		--SharedEntityManager->GetObserverManager().LocksCount;
+		checkf(SharedEntityManager->GetObserverManager().LocksCount >= 0
+			, TEXT("%hs: getting a negative locks count as a result of the operation is unexpected."), __FUNCTION__);
+		SharedEntityManager->GetObserverManager().ResumeExecution(*this);
+	}
+
+	void FObserverLock::ForceUpdateCurrentThreadID()
+	{
+		OwnerThreadId = FPlatformTLS::GetCurrentThreadId();
+	}
+
+	//-----------------------------------------------------------------------------
+	// FCreationContext
+	//-----------------------------------------------------------------------------
+	FCreationContext::FCreationContext()
+		: Lock(Private::GetDummyObserverLock())
+	{	
+	}
+
+	FCreationContext::~FCreationContext()
+	{
+		if (CreationHandle.IsSet())
+		{
+			if (TSharedPtr<FMassEntityManager> SharedEntityManager = Lock->GetWeakEntityManager().Pin())
+			{
+				SharedEntityManager->GetObserverManager().ReleaseCreationHandle(CreationHandle);
+			}
+		}
+	}
+
+	TSharedRef<FCreationContext> FCreationContext::DebugCreateDummyCreationContext()
+	{
+		return MakeShareable(new FCreationContext());
+	}
+
+	TArray<FMassArchetypeEntityCollection> FCreationContext::GetEntityCollections(const FMassEntityManager& InEntityManager) const
+	{
+		TArray<FMassArchetypeEntityCollection> OutCollections;
+
+		// if the creation handle isn't set there are no creation ops we know about
+		if (CreationHandle.IsSet())
+		{
+			const FBufferedNotification& Notification = Lock->GetCreationNotification(CreationHandle);
+
+			if (const UE::Mass::FEntityCollection* CreatedEntities = Notification.AffectedEntities.TryGet<FEntityCollection>())
+			{
+				OutCollections.Append(CreatedEntities->GetUpToDatePerArchetypeCollections(InEntityManager));
+			}
+			else
+			{
+				const FMassEntityHandle EntityHandle = Notification.AffectedEntities.Get<FMassEntityHandle>();
+
+				UE::Mass::Utils::CreateEntityCollections(InEntityManager, MakeArrayView(&EntityHandle, 1)
+					, FMassArchetypeEntityCollection::NoDuplicates
+					, OutCollections);
+			}
+		}
+
+		return OutCollections;
+	}
+
+	bool FCreationContext::DebugAreEntityCollectionsUpToDate() const 
+	{ 
+		if (CreationHandle.IsSet())
+		{
+			const FBufferedNotification& Notification = Lock->GetCreationNotification(CreationHandle);
+
+			// collections can be not up to date only if we're storing multiple entities (i.e. not a single handle)
+			if (const UE::Mass::FEntityCollection* CreatedEntities = Notification.AffectedEntities.TryGet<FEntityCollection>())
+			{
+				return CreatedEntities->IsUpToDate();
+			}
+		}
+
+		return true;
+	}
+} // UE::Mass::ObserverManager
diff --git a/Engine/Source/Runtime/MassEntity/Public/MassEntityManager.h b/Engine/Source/Runtime/MassEntity/Public/MassEntityManager.h
--- a/Engine/Source/Runtime/MassEntity/Public/MassEntityManager.h
+++ b/Engine/Source/Runtime/MassEntity/Public/MassEntityManager.h
@@ -41,6 +41,11 @@ namespace UE::Mass
 	{
 		struct FEntityStorageInitializer;
 	}
+	namespace ObserverManager
+	{
+		struct FObserverLock;
+		struct FCreationContext;
+	}
 }
 
 // use REQUESTED_MASS_CONCURRENT_RESERVE=0 in your project's Build.cs file to disable concurrent storage
@@ -236,60 +241,38 @@ public:
 	 * @return FMassEntityHandle id of the newly created entity */
 	FMassEntityHandle CreateEntity(TConstArrayView<FInstancedStruct> FragmentInstanceList, const FMassArchetypeSharedFragmentValues& SharedFragmentValues = {}, const FMassArchetypeCreationParams& CreationParams = FMassArchetypeCreationParams());
 
-	/**
-	 * A dedicated structure for ensuring the "on entities creation" observers get notified only once all other 
-	 * initialization operations are done and this creation context instance gets released. 
-	 */
-	struct MASSENTITY_API FEntityCreationContext
-	{
-	private:
-		FEntityCreationContext();
-		explicit FEntityCreationContext(FMassEntityManager& InManager, const TConstArrayView<FMassEntityHandle> InCreatedEntities = {});
-		FEntityCreationContext(FMassEntityManager& InManager, const TConstArrayView<FMassEntityHandle> InCreatedEntities, FMassArchetypeEntityCollection&& InEntityCollection);
-		
-	public:
-		~FEntityCreationContext();
-
-		/** Returns EntityCollections, reconstructing them if needed (empty or dirtied). */
-		TConstArrayView<FMassArchetypeEntityCollection> GetEntityCollections() const;
-		int32 GetSpawnedNum() const { return CreatedEntities.Num(); }
-		void MarkDirty();
-		FORCEINLINE bool IsDirty() const { return EntityCollections.IsEmpty() && (CreatedEntities.IsEmpty() == false); }
-		void AppendEntities(const TConstArrayView<FMassEntityHandle> InEntitiesToAppend);
-		void AppendEntities(const TConstArrayView<FMassEntityHandle> InEntitiesToAppend, FMassArchetypeEntityCollection&& InEntityCollection);
-
-		/** Function for debugging/testing purposes. We don't expect users to ever call it, always get collections via GetEntityCollections */
-		bool DebugAreEntityCollectionsUpToDate() const { return EntityCollections.IsEmpty() == CreatedEntities.IsEmpty(); }
-
-		UE_DEPRECATED(5.5, "This constructor is now deprecated and defunct. Use one of the others instead.")
-		explicit FEntityCreationContext(const int32 InNumSpawned) : FEntityCreationContext() {}
-		UE_DEPRECATED(5.5, "This function is now deprecated since FEntityCreationContext can contain more than a single collection now. Use GetEntityCollections instead.")
-		const FMassArchetypeEntityCollection& GetEntityCollection() const;
+	using FEntityCreationContext = UE::Mass::ObserverManager::FCreationContext;
 
-	private:
-		friend FMassEntityManager;
-		/** To be called in case of processor forking. */
-		void ForceUpdateCurrentThreadID();
-
-		/**
-		 * Identifies the thread where given FEntityCreationContext instance was created. All subsequent operations are 
-		 * expected to be run in the same thread.
-		 */
-		uint32 OwnerThreadId;
-		mutable TArray<FMassArchetypeEntityCollection> EntityCollections;
-		TArray<FMassEntityHandle> CreatedEntities;
-		FMassArchetypeEntityCollection::EDuplicatesHandling CollectionCreationDuplicatesHandling = FMassArchetypeEntityCollection::EDuplicatesHandling::NoDuplicates;
-		TSharedPtr<FMassEntityManager> Manager;
-	};
 	/**
 	 * The main use-case for this function is to create a blank FEntityCreationContext and hold on to it while creating 
 	 * a bunch of entities (with multiple calls to BatchCreate* and/or BatchBuild*) and modifying them (with mutating batched API)
 	 * while not causing multiple Observers to trigger. All the observers will be triggered at one go, once the FEntityCreationContext 
-	 * instance gets destroyed. 
-	 * @return ActiveCreationContext. If it's valid FEntityCreationContext instance will be created and assigned to ActiveCreationContext first.
+	 * instance gets destroyed.
+	 * 
+	 * !Important note: the "Creation Context" is a specialized wrapper for an "Observers Lock" (@see GetOrMakeObserversLock).
+	 * As long as the creation context is alive all the operations will be assumed to affect the newly created entities.
+	 * The consequence of that is operations performed on already existing entities won't be tracked, as long
+	 * as the creation context is alive.
+	 * Note that you can hold a FMassObserverManager::FObserverLock instance while the creation lock gets destroyed, the
+	 * observers lock is a lower-level concept than the creation context.
+	 * 
+	 * @return the existing (if valid) or a newly created creation context
 	 */
 	TSharedRef<FEntityCreationContext> GetOrMakeCreationContext();
 
+	/**
+	 * Fetches the observers lock (as hosted by FMassObserverManager). If one is not currently active,
+	 * one will be created. While the lock is active all the observers notifications are suspended, and
+	 * will be sent out when FMassObserverManager::FObserverLock instance gets destroyed.
+	 * Locking observers needs to be used when entities are being configured with multiple operations,
+	 * and we want observers to be triggered only once all the operations are executed.
+	 *
+	 * Note that while the observers are locked we're unable to send "Remove" notifications, so once
+	 * the lock is released and the observers get notified, the data being removed won't be available anymore
+	 * (which is a difference in behavior as compared to removal notifications while the observers are not locked).
+	 */
+	TSharedRef<UE::Mass::ObserverManager::FObserverLock> GetOrMakeObserversLock();
+
 	/**
 	 * A version of CreateEntity that's creating a number of entities (Count) in one go
 	 * @param ArchetypeHandle you want this entity to be
@@ -1137,48 +1120,6 @@ private:
 #endif
 	}
 
-	/**
-	 * If ActiveCreationContext is not valid the function creates a new shared FEntityCreationContext instance and returns that.
-	 * Otherwise ActiveCreationContext will get extended with ReservedEntities and EntityCollection, and returned by the function.
-	 */
-	TSharedRef<FEntityCreationContext> GetOrMakeCreationContext(TConstArrayView<FMassEntityHandle> ReservedEntities, FMassArchetypeEntityCollection&& EntityCollection);
-
-	/** 
-	 * This type is used in entity mutating batched API to ensure the active FEntityCreationContext gets dirtied 
-	 * upon function's end (since the mutating operations render FEntityCreationContext.EntityCollections invalid).
-	 * It also serves as a cached IsDuringEntityCreation value.
-	 */
-	struct FScopedCreationContextOperations
-	{
-		FScopedCreationContextOperations(FMassEntityManager& InManager)
-			: bIsDuringEntityCreation(InManager.IsDuringEntityCreation())
-			, Manager(InManager)
-		{}
-		~FScopedCreationContextOperations()
-		{
-			if (bIsDuringEntityCreation)
-			{
-				Manager.DirtyCreationContext();
-			}
-			// else, there's nothing to do, there's no creation context to call functions for
-		}
-
-		bool IsAllowedToTriggerObservers() const { return (bIsDuringEntityCreation == false); }
-
-	private:
-		const bool bIsDuringEntityCreation = false;
-		FMassEntityManager& Manager;
-	};
-	friend FScopedCreationContextOperations;
-
-	/** @return whether actual context dirtying took place which is equivalent to IsDuringEntityCreation */
-	bool DirtyCreationContext();
-
-	/** 
-	 * @return whether it's allowed for observers to get triggered. If not then the active creation context will be dirtied
-	 *	to cause observers triggering upon context's destruction
-	 */
-	bool IsAllowedToTriggerObservers() { return DirtyCreationContext() == false; }
 	bool DebugDoCollectionsOverlapCreationContext(TConstArrayView<FMassArchetypeEntityCollection> EntityCollections) const;
 		
 private:
@@ -1193,9 +1134,6 @@ private:
 	UE::Mass::FStorageType* DebugEntityStoragePtr = nullptr;
 #endif // WITH_MASSENTITY_DEBUG
 
-	/** Never access directly, use GetOrMakeCreationContext instead. */
-	TWeakPtr<FEntityCreationContext> ActiveCreationContext;
-
 	std::atomic<int32> ProcessingScopeCount = 0;
 
 	// the "version" number increased every time an archetype gets added
@@ -1422,5 +1360,10 @@ inline UE::Mass::FTypeManager& FMassEntityManager::GetTypeManager()
 
 inline bool FMassEntityManager::IsDuringEntityCreation() const
 {
-	return ActiveCreationContext.IsValid();
+	return ObserverManager.GetCreationContext().IsValid();
+}
+
+inline TSharedRef<FMassObserverManager::FObserverLock> FMassEntityManager::GetOrMakeObserversLock()
+{
+	return ObserverManager.GetOrMakeObserverLock();
 }
diff --git a/Engine/Source/Runtime/MassEntity/Public/MassObserverManager.h b/Engine/Source/Runtime/MassEntity/Public/MassObserverManager.h
--- a/Engine/Source/Runtime/MassEntity/Public/MassObserverManager.h
+++ b/Engine/Source/Runtime/MassEntity/Public/MassObserverManager.h
@@ -5,6 +5,13 @@
 #include "MassEntityHandle.h"
 #include "MassProcessingTypes.h"
 #include "MassEntityTypes.h"
+#include "MassArchetypeTypes.h"
+#include "Misc/Fork.h"
+#if UE_ENABLE_INCLUDE_ORDER_DEPRECATED_IN_5_6
+// the following is a new header, but it contains some types moved over from other places, and including
+// said header ties everything together. Engine code has been updated to not require it
+#include "MassObserverNotificationTypes.h"
+#endif // UE_ENABLE_INCLUDE_ORDER_DEPRECATED_IN_5_6
 #include "MassObserverManager.generated.h"
 
 
@@ -17,6 +24,13 @@ namespace UE::Mass
 	namespace ObserverManager
 	{
 		struct FDeprecationHelper;
+		struct FBufferedNotificationExecutioner;
+		struct FBufferedCreationNotificationExecutioner;
+
+		struct FBufferedNotification;
+		struct FCreationNotificationHandle;
+		struct FObserverLock;
+		struct FCreationContext;
 	}
 }
 
@@ -51,8 +65,15 @@ struct MASSENTITY_API FMassObserverManager
 {
 	GENERATED_BODY()
 
-public:
-	FMassObserverManager();	
+	/** convenience aliases */
+	using FObserverLock = UE::Mass::ObserverManager::FObserverLock;
+	using FBufferedNotification = UE::Mass::ObserverManager::FBufferedNotification;
+	using FCreationNotificationHandle = UE::Mass::ObserverManager::FCreationNotificationHandle;
+	using FCreationContext = UE::Mass::ObserverManager::FCreationContext;
+
+	FMassObserverManager();
+
+	FMassEntityManager& GetEntityManager();
 
 	const FMassFragmentBitSet* GetObservedFragmentBitSets() const { return ObservedFragments; }
 	const FMassFragmentBitSet& GetObservedFragmentsBitSet(const EMassObservedOperation Operation) const 
@@ -81,47 +102,95 @@ public:
 		return HasObserversForBitSet(Composition.Fragments, Operation) || HasObserversForBitSet(Composition.Tags, Operation);
 	}
 
-	/** @return whether any observers have been executed */
+	/** @return whether there are observers watching affected elements */
 	bool OnPostEntitiesCreated(const FMassArchetypeEntityCollection& EntityCollection);
 
-	/** @return whether any observers have been executed */
-	bool OnPostEntitiesCreated(TConstArrayView<FMassArchetypeEntityCollection> EntityCollections);
+	/** @return whether there are observers watching affected elements */
+	bool OnPostEntityCreated(const FMassEntityHandle Entity, const FMassArchetypeCompositionDescriptor& Composition);
 
-	/** @return whether any observers have been executed */
-	bool OnPostEntitiesCreated(UE::Mass::FProcessingContext& InProcessingContext, TConstArrayView<FMassArchetypeEntityCollection> EntityCollections);
-
-	/** @return whether any observers have been executed */
+	/** @return whether there are observers watching affected elements */
 	bool OnPreEntitiesDestroyed(const FMassArchetypeEntityCollection& EntityCollection);
 
-	/** @return whether any observers have been executed */
+	/** @return whether there are observers watching affected elements */
 	bool OnPreEntitiesDestroyed(UE::Mass::FProcessingContext& ProcessingContext, const FMassArchetypeEntityCollection& EntityCollection);
 
-	/** @return whether any observers have been executed */
+	/** @return whether there are observers watching affected elements */
 	bool OnPreEntityDestroyed(const FMassArchetypeCompositionDescriptor& ArchetypeComposition, const FMassEntityHandle Entity);
 
-	/** @return whether any observers have been executed */
-	bool OnCompositionChanged(const FMassEntityHandle Entity, const FMassArchetypeCompositionDescriptor& CompositionDelta, const EMassObservedOperation Operation);
-
-	/** @return whether any observers have been executed */
+	/** @return whether there are observers watching affected elements */
 	bool OnPostCompositionAdded(const FMassEntityHandle Entity, const FMassArchetypeCompositionDescriptor& Composition)
 	{
 		return OnCompositionChanged(Entity, Composition, EMassObservedOperation::Add);
 	}
-	/** @return whether any observers have been executed */
+	/** @return whether there are observers watching affected elements */
 	bool OnPreCompositionRemoved(const FMassEntityHandle Entity, const FMassArchetypeCompositionDescriptor& Composition)
 	{
 		return OnCompositionChanged(Entity, Composition, EMassObservedOperation::Remove);
 	}
 
-	/** @return whether any observers have been executed */
-	bool OnCompositionChanged(UE::Mass::FProcessingContext& InProcessingContext, const FMassArchetypeEntityCollection& EntityCollection, const FMassArchetypeCompositionDescriptor& Composition, const EMassObservedOperation Operation);
+protected:
+	/**
+	 * helper struct for holding either a single enity handle or an archetype collection reference.
+	 * The type is highly transient, do not store instances of it. Its only function is to 
+	 * allow us to have a single OnCompositionChange implementation rather than having two separate
+	 * implementations, one for entity handle, the other for archetype collection.
+	 */
+	struct FCollectionRefOrHandle
+	{
+		UE_NONCOPYABLE(FCollectionRefOrHandle);
+
+		static FMassArchetypeEntityCollection DummyCollection;
+explicit FCollectionRefOrHandle(FMassEntityHandle InEntityHandle)
+			: EntityHandle(InEntityHandle)
+			, EntityCollection(DummyCollection)
+		{	
+		}
+
+		explicit FCollectionRefOrHandle(const FMassArchetypeEntityCollection& InEntityCollection)
+			: EntityCollection(InEntityCollection)
+		{	
+		}
+
+		FMassEntityHandle EntityHandle;
+		const FMassArchetypeEntityCollection& EntityCollection;
+	};
+	bool OnCompositionChanged(FCollectionRefOrHandle&& EntityCollection, const FMassArchetypeCompositionDescriptor& Composition
+		, const EMassObservedOperation Operation, UE::Mass::FProcessingContext* ProcessingContext = nullptr);
+
+public:
+	bool OnCompositionChanged(const FMassArchetypeEntityCollection& EntityCollection, const FMassArchetypeCompositionDescriptor& Composition
+		, const EMassObservedOperation Operation, UE::Mass::FProcessingContext* ProcessingContext = nullptr)
+	{
+		TRACE_CPUPROFILER_EVENT_SCOPE(MassObserver_OnCompositionChanged_Collection);
+		return OnCompositionChanged(FCollectionRefOrHandle(EntityCollection), Composition, Operation, ProcessingContext);
+	}
+
+	bool OnCompositionChanged(FMassEntityHandle EntityHandle, const FMassArchetypeCompositionDescriptor& Composition
+		, const EMassObservedOperation Operation, UE::Mass::FProcessingContext* ProcessingContext = nullptr)
+	{
+		TRACE_CPUPROFILER_EVENT_SCOPE(MassObserver_OnCompositionChanged_Entity);
+		return OnCompositionChanged(FCollectionRefOrHandle(EntityHandle), Composition, Operation, ProcessingContext);
+	}
 
-	/** @return whether any observers have been executed */
-	bool OnCompositionChanged(const FMassArchetypeEntityCollection& EntityCollection, const FMassArchetypeCompositionDescriptor& Composition, const EMassObservedOperation Operation);
+	template<typename T, typename U = std::decay_t<T>> requires (std::is_same_v<U, FMassFragmentBitSet> || std::is_same_v<U, FMassTagBitSet>)
+	bool OnCompositionChanged(const FMassArchetypeEntityCollection& EntityCollection, T&& BitSet
+		, const EMassObservedOperation Operation, UE::Mass::FProcessingContext* ProcessingContext = nullptr)
+	{
+		return OnCompositionChanged(EntityCollection, FMassArchetypeCompositionDescriptor(MoveTemp(BitSet)), Operation, ProcessingContext);
+	}
+
+	template<typename T, typename U = std::decay_t<T>> requires (std::is_same_v<U, FMassFragmentBitSet> || std::is_same_v<U, FMassTagBitSet>)
+	bool OnCompositionChanged(const FMassEntityHandle Entity, T&& BitSet
+		, const EMassObservedOperation Operation, UE::Mass::FProcessingContext* ProcessingContext = nullptr)
+	{
+		return OnCompositionChanged(Entity, FMassArchetypeCompositionDescriptor(MoveTemp(BitSet)), Operation, ProcessingContext);
+	}
 
 	void AddObserverInstance(const UScriptStruct& ElementType, const EMassObservedOperation Operation, UMassProcessor& ObserverProcessor);
 	void RemoveObserverInstance(const UScriptStruct& ElementType, const EMassObservedOperation Operation, UMassProcessor& ObserverProcessor);
 
+	void ReleaseCreationHandle(UE::Mass::ObserverManager::FCreationNotificationHandle InCreationNotificationHandle);
+
 	void DebugGatherUniqueProcessors(TArray<const UMassProcessor*>& OutProcessors) const;
 
 protected:
@@ -130,9 +199,56 @@ protected:
 
 	void Initialize();
 	void DeInitialize();
+
+	friend UE::Mass::ObserverManager::FBufferedNotificationExecutioner;
+	friend UE::Mass::ObserverManager::FBufferedCreationNotificationExecutioner;
+
 	static void HandleElementsImpl(UE::Mass::FProcessingContext& ProcessingContext, TConstArrayView<FMassArchetypeEntityCollection> EntityCollections
 		, TArrayView<const UScriptStruct*> ObservedTypes, FMassObserversMap& HandlersContainer);
 
+	static void HandleFragmentsImpl(UE::Mass::FProcessingContext& ProcessingContext, const FMassArchetypeEntityCollection& EntityCollection
+		, TArrayView<const UScriptStruct*> ObservedTypes, FMassObserversMap& HandlersContainer)
+	{
+		HandleElementsImpl(ProcessingContext, {EntityCollection}, ObservedTypes, HandlersContainer);
+	}
+
+	/** Coalesces all the elements observed in all the collections and executes all the observers at once */
+	bool OnCollectionsCreatedImpl(UE::Mass::FProcessingContext& ProcessingContext, TArray<FMassArchetypeEntityCollection>&& EntityCollections);
+
+	void OnPostFork(EForkProcessRole);
+
+	//TSharedRef<FObserverLock> GetOrMakeObserverLock(TConstArrayView<FMassEntityHandle> ReservedEntities, FMassArchetypeEntityCollection&& EntityCollection);
+	TSharedRef<FObserverLock> GetOrMakeObserverLock();
+	TSharedPtr<FObserverLock> GetObserverLock() const;
+	bool IsLocked() const { return ActiveObserverLock.IsValid(); }
+	TSharedRef<FCreationContext> GetOrMakeCreationContext();
+	TSharedRef<FCreationContext> GetOrMakeCreationContext(TConstArrayView<FMassEntityHandle> ReservedEntities, FMassArchetypeEntityCollection&& EntityCollection);
+	TSharedPtr<FCreationContext> GetCreationContext() const;
+	
+	//-----------------------------------------------------------------------------
+	// Observers locking logic
+	//-----------------------------------------------------------------------------
+	friend FObserverLock;
+
+	/**
+	 * Resumes observer triggering. All notifications collected in lock's BufferedNotifications will be processed at this point.
+	 *
+	 * Note that due to all the notifications being sent our are being sent post-factum the "OnPreRemove" 
+	 * observers won't be able to access the data being removed, since the remove operation has already been performed.
+	 * All the instances of removal-observers being triggered will be logged.
+	 * 
+	 * Intended to be called automatically by ~FObserverLock
+	 */
+	void ResumeExecution(FObserverLock& LockBeingReleased);
+
+	/**
+	 * Never access directly, use GetOrMakeObserverLock or GetOrMakeCreationContext instead.
+	 * Note: current lock is single-threaded. There's a path towards making it multithreaded, we'll work on it once we have a use-case
+	 */
+	TWeakPtr<FObserverLock> ActiveObserverLock;
+	std::atomic<int32> LocksCount = 0;
+	TWeakPtr<FCreationContext> ActiveCreationContext;
+
 	FMassFragmentBitSet ObservedFragments[(uint8)EMassObservedOperation::MAX];
 	FMassTagBitSet ObservedTags[(uint8)EMassObservedOperation::MAX];
 
@@ -144,10 +260,15 @@ protected:
 
 	/** 
 	 * The owning EntityManager. No need for it to be a UPROPERTY since by design we don't support creation of 
-	 * FMassObserverManager outside of an FMassEntityManager instance 
+	 * FMassObserverManager outside a FMassEntityManager instance 
 	 */
 	FMassEntityManager& EntityManager;
 
+#if WITH_MASSENTITY_DEBUG
+	std::atomic<uint32> LockedNotificationSerialNumber = 1;
+	uint32 DebugNonTrivialResumeExecutionCount = 0;
+#endif // WITH_MASSENTITY_DEBUG
+
 public:
 	UE_DEPRECATED(5.5, "This flavor of OnPostEntitiesCreated is deprecated. Please use the one taking a TConstArrayView<FMassArchetypeEntityCollection> parameter instead")
 	bool OnPostEntitiesCreated(UE::Mass::FProcessingContext& InProcessingContext, const FMassArchetypeEntityCollection& EntityCollection);
@@ -157,13 +278,15 @@ public:
 	void OnPreFragmentOrTagRemoved(const UScriptStruct& FragmentOrTagType, const FMassArchetypeEntityCollection& EntityCollection);
 	UE_DEPRECATED(5.6, "*FragmentOrTag* functions have been deprecated. Use OnCompositionChanged instead.")
 	void OnFragmentOrTagOperation(const UScriptStruct& FragmentOrTagType, const FMassArchetypeEntityCollection& EntityCollection, const EMassObservedOperation Operation);
+	UE_DEPRECATED(5.6, "Use the other OnPostEntitiesCreated implementation.")
+	bool OnPostEntitiesCreated(UE::Mass::FProcessingContext& InProcessingContext, TConstArrayView<FMassArchetypeEntityCollection> EntityCollections);
+	UE_DEPRECATED(5.6, "Use the other OnCompositionChanged implementation.")
+	bool OnCompositionChanged(UE::Mass::FProcessingContext& InProcessingContext, const FMassArchetypeEntityCollection& EntityCollection
+		, const FMassArchetypeCompositionDescriptor& Composition, const EMassObservedOperation Operation);
 protected:
 	UE_DEPRECATED(5.6, "Use HandleElementsImpl instead.")
 	void HandleSingleEntityImpl(const UScriptStruct& FragmentType, const FMassArchetypeEntityCollection& EntityCollection, FMassObserversMap& HandlersContainer);
-	UE_DEPRECATED(5.6, "Use HandleElementsImpl instead.")
-	void HandleFragmentsImpl(UE::Mass::FProcessingContext& ProcessingContext, const FMassArchetypeEntityCollection& EntityCollection
-		, TArrayView<const UScriptStruct*> ObservedTypes, FMassObserversMap& HandlersContainer);
-
+	
 	friend UE::Mass::ObserverManager::FDeprecationHelper;
 };
 
@@ -175,3 +298,21 @@ struct TStructOpsTypeTraits<FMassObserverManager> : public TStructOpsTypeTraitsB
 		WithCopy = false,
 	};
 };
+
+//----------------------------------------------------------------------//
+// inlines
+//----------------------------------------------------------------------//
+inline FMassEntityManager& FMassObserverManager::GetEntityManager()
+{
+	return EntityManager;
+}
+
+inline TSharedPtr<FMassObserverManager::FObserverLock> FMassObserverManager::GetObserverLock() const
+{
+	return ActiveObserverLock.Pin();
+}
+
+inline TSharedPtr<FMassObserverManager::FCreationContext> FMassObserverManager::GetCreationContext() const
+{
+	return ActiveCreationContext.Pin();
+}
diff --git a/Engine/Source/Runtime/MassEntity/Public/MassObserverNotificationTypes.h b/Engine/Source/Runtime/MassEntity/Public/MassObserverNotificationTypes.h
new file mode 100644
--- /dev/null
+++ b/Engine/Source/Runtime/MassEntity/Public/MassObserverNotificationTypes.h
@@ -0,0 +1,462 @@
+// Copyright Epic Games, Inc. All Rights Reserved.
+
+#pragma once
+
+#include "MassEntityHandle.h"
+#include "Misc/TVariant.h"
+#include "MassEntityCollection.h"
+#include "MassEntityTypes.h"
+
+#define UE_CHECK_OWNER_THREADID() checkf(OwnerThreadId == FPlatformTLS::GetCurrentThreadId(), TEXT("%hs: all FObserverLock operations are expected to be run in a single thread"), __FUNCTION__)
+
+struct FMassObserverManager;
+namespace UE::Mass
+{
+	struct FProcessingContext;
+	namespace ObserverManager
+	{
+		enum class EObservedOperationNotification : uint8
+		{
+			Add = static_cast<uint8>(EMassObservedOperation::Add),
+			Remove = static_cast<uint8>(EMassObservedOperation::Remove),
+			Create
+		};
+
+
+		/**
+		 * The type represents a single "operation", as observed by the registered observers, an operation that has been performed
+		 * while the FObserverLock was active. Instances of FBufferedNotification contain all the information needed
+		 * to send out the necessary notification once the observers lock gets released.
+		 *
+		 * Note that the type contains information necessary to sent our notification. In case of "Remove" notifications the
+		 * operations has already been performed, and the data being removed is no longer available to the observers, and
+		 * instances of FBufferedNotification do not host this information either. 
+		 */
+		struct FBufferedNotification
+		{
+			struct FEmptyComposition
+			{
+			};
+			EObservedOperationNotification Type;
+			using FCompositionDescription = TVariant<FEmptyComposition, FMassArchetypeCompositionDescriptor, FMassFragmentBitSet, FMassTagBitSet>;
+			FCompositionDescription CompositionChange;
+			using FEntitiesContainer = TVariant<FEntityCollection, FMassEntityHandle>;
+			FEntitiesContainer AffectedEntities;
+
+			template<typename TComposition>
+			FBufferedNotification(const EObservedOperationNotification InType, TComposition&& Composition, FEntitiesContainer&& Entities)
+				: Type(InType)
+				, CompositionChange(TInPlaceType<typename TDecay<TComposition>::Type>(), Forward<TComposition>(Composition))
+				, AffectedEntities(MoveTemp(Entities))
+			{
+			}
+
+			template<typename TComposition, typename TEntities>
+			FBufferedNotification(const EObservedOperationNotification InType, TComposition&& Composition, TEntities&& Entities)
+				: Type(InType)
+				, CompositionChange(TInPlaceType<typename TDecay<TComposition>::Type>(), Forward<TComposition>(Composition))
+				, AffectedEntities(TInPlaceType<typename TDecay<TEntities>::Type>(), Forward<TEntities>(Entities))
+			{
+			}
+
+			template<typename TEntities>
+			FBufferedNotification(const EObservedOperationNotification InType, FCompositionDescription&& Composition, TEntities&& Entities)
+				: Type(InType)
+				, CompositionChange(MoveTemp(Composition))
+				, AffectedEntities(TInPlaceType<typename TDecay<TEntities>::Type>(), Forward<TEntities>(Entities))
+			{
+			}
+
+			FBufferedNotification(const EObservedOperationNotification InType, FCompositionDescription&& Composition, const FMassArchetypeEntityCollection& Entities)
+				: Type(InType)
+				, CompositionChange(MoveTemp(Composition))
+				, AffectedEntities(TInPlaceType<typename TDecay<FEntityCollection>::Type>(), Entities)
+			{
+			}
+
+			bool IsCreationNotification() const
+			{
+				return Type == EObservedOperationNotification::Create;
+			}
+
+			void AddHandle(const FMassEntityHandle EntityHandle)
+			{
+				FEntityCollection* StoredCollection = AffectedEntities.TryGet<FEntityCollection>();
+				if (StoredCollection == nullptr)
+				{
+					StoredCollection = &ConvertStoredHandleToCollection(TConstArrayView<FMassEntityHandle>());
+				}
+				StoredCollection->AddHandle(EntityHandle);
+			}
+
+			FORCEINLINE void AppendEntities(const FMassEntityHandle EntityHandle)
+			{
+				AddHandle(EntityHandle);
+			}
+
+			void AppendEntities(const TConstArrayView<FMassEntityHandle> InEntityHandles)
+			{
+				if (FEntityCollection* StoredCollection = AffectedEntities.TryGet<FEntityCollection>())
+				{
+					StoredCollection->AppendHandles(InEntityHandles);
+				}
+				else
+				{
+					ConvertStoredHandleToCollection(InEntityHandles);
+				}
+			}
+
+			void AppendEntities(const TConstArrayView<FMassEntityHandle> InEntityHandles, FMassArchetypeEntityCollection&& EntityCollection)
+			{
+				if (FEntityCollection* StoredCollection = AffectedEntities.TryGet<FEntityCollection>())
+				{
+					StoredCollection->AppendHandles(InEntityHandles, Forward<FMassArchetypeEntityCollection>(EntityCollection));
+				}
+				else
+				{
+					ConvertStoredHandleToCollection(InEntityHandles);
+					// we're ignoring EntityCollection since the collections will need to be rebuilt anyway,
+					// due to AffectedEntities already containing some data before this call
+				}
+			}
+
+			template<typename T> requires std::is_same_v<typename TDecay<T>::Type, FMassArchetypeEntityCollection>
+			void AppendEntities(T&& InEntityCollection)
+			{
+				if (FEntityCollection* StoredCollection = AffectedEntities.TryGet<FEntityCollection>())
+				{
+					StoredCollection->AppendCollection(Forward<T>(InEntityCollection));
+				}
+				else
+				{
+					ConvertStoredHandleToCollection(Forward<T>(InEntityCollection));
+				}
+			}
+
+			void DirtyAffectedEntities()
+			{
+				if (FEntityCollection* StoredCollection = AffectedEntities.TryGet<FEntityCollection>())
+				{
+					StoredCollection->MarkDirty();
+				}
+			}
+
+			static bool AreCompositionsEqual(const FCompositionDescription& A, const FCompositionDescription& B)
+			{
+				if (A.GetIndex() == B.GetIndex())
+				{
+					switch (A.GetIndex())
+					{
+						case FCompositionDescription::IndexOfType<FEmptyComposition>():
+							return true;
+						case FCompositionDescription::IndexOfType<FMassArchetypeCompositionDescriptor>():
+							return A.Get<FMassArchetypeCompositionDescriptor>().IsIdentical(B.Get<FMassArchetypeCompositionDescriptor>());
+						case FCompositionDescription::IndexOfType<FMassFragmentBitSet>():
+							return A.Get<FMassFragmentBitSet>() == B.Get<FMassFragmentBitSet>();
+						case FCompositionDescription::IndexOfType<FMassTagBitSet>():
+							return A.Get<FMassTagBitSet>() == B.Get<FMassTagBitSet>();
+						default:
+							return false;
+					}
+				}
+				return false;
+			}
+
+		private:
+			template<typename TEntities> 
+			FEntityCollection& ConvertStoredHandleToCollection(TEntities&& InEntities)
+			{
+				// AffectedEntities holds a single handle. We need to extract it and emplace a FEntityCollection instance 
+				const FMassEntityHandle StoredHandle = AffectedEntities.Get<FMassEntityHandle>();
+				AffectedEntities.Emplace<FEntityCollection>(Forward<TEntities>(InEntities));
+				FEntityCollection& StoredEntities = AffectedEntities.Get<FEntityCollection>();
+				StoredEntities.AddHandle(StoredHandle);
+				return StoredEntities;
+			}
+		};
+
+		/** Simple handle type representing a entity creation notification as stored by FObserverLock */
+		struct FCreationNotificationHandle
+		{
+			bool IsSet() const
+			{
+				return OpIndex != INDEX_NONE;
+			}
+
+			operator int() const
+			{
+				return OpIndex;
+			}
+
+		private:
+			friend FMassObserverManager;
+
+			/**
+			 * set upon creation to the value of LockedNotificationSerialNumber.
+			 * This property's value is checked when the creation handle gets "released" via
+			 * FMassObserverManager::ReleaseCreationHandle
+			 */ 
+			uint32 SerialNumber = 0;
+
+			int32 OpIndex = INDEX_NONE;
+		};
+
+		/**
+		 * Once created with MassObserverManager.GetOrMakeObserverLock will prevent triggering
+		 * observers and instead buffer all the notifications to be sent.
+		 * Once the FObserverLock gets released it will call MassObserverManager.ResumeExecution
+		 * that will send out all the buffered notifications.
+		 * 
+		 * @note that due to the buffering, all the "Remove" operation observers will be sent out later
+		 * than usually - without locking those observers get triggered before the removal operation is
+		 * performed, and as such have access to the data "about to be removed". Removal observers sent out
+		 * after lock release won't have access to that information.
+		 * 
+		 * There's a special path for freshly created entities, see FCreationContext for more details.
+		 */
+		struct MASSENTITY_API FObserverLock
+		{
+			FObserverLock() = default;
+			~FObserverLock();
+
+			TWeakPtr<FMassEntityManager> GetWeakEntityManager() const
+			{
+				return WeakEntityManager;
+			}
+
+			void MarkCreationNotificationDirty(FCreationNotificationHandle CreationHandle)
+			{
+				ensureMsgf(CreationHandle == CreationNotificationIndex, TEXT("Given creation handle doesn't match this Lock's data"));
+				checkf(BufferedNotifications.IsValidIndex(CreationHandle), TEXT("Given CreationHandle doesn't match stored notifications"));
+				BufferedNotifications[CreationHandle].DirtyAffectedEntities();
+			}
+
+			const FBufferedNotification& GetCreationNotification(FCreationNotificationHandle CreationHandle) const
+			{
+				ensureMsgf(CreationHandle == CreationNotificationIndex, TEXT("Given creation handle doesn't match this Lock's data"));
+				checkf(BufferedNotifications.IsValidIndex(CreationHandle), TEXT("Given CreationHandle doesn't match stored notifications"));
+
+				return BufferedNotifications[CreationHandle];
+			}
+
+		private:
+			explicit FObserverLock(FMassObserverManager& ObserverManager);
+
+			int32 GetOrCreateCreationNotification()
+			{
+				UE_CHECK_OWNER_THREADID();
+				int32 LocalCreationNotificationIndex = INDEX_NONE;
+				if (CreationNotificationIndex.compare_exchange_weak(LocalCreationNotificationIndex, BufferedNotifications.Num()))
+				{
+					BufferedNotifications.Emplace(EObservedOperationNotification::Create
+						, FBufferedNotification::FEmptyComposition{}
+						, UE::Mass::FEntityCollection());
+				}
+				return LocalCreationNotificationIndex;
+			}
+
+			bool ReleaseCreationNotification(FCreationNotificationHandle CreationHandle)
+			{
+				UE_CHECK_OWNER_THREADID();
+				checkf(BufferedNotifications.IsValidIndex(CreationHandle), TEXT("Given CreationHandle doesn't match stored notifications"));
+
+				int32 CreationHandleOpIndex = CreationHandle;
+				if (CreationNotificationIndex.compare_exchange_weak(CreationHandleOpIndex, INDEX_NONE))
+				{
+					return true;
+				}
+				// else 
+				ensureMsgf(CreationHandle == CreationNotificationIndex, TEXT("Given creation handle doesn't match this Lock's data"));
+				return false;
+			}
+
+			int32 AddCreatedEntity(FMassEntityHandle CreatedEntity)
+			{
+				UE_CHECK_OWNER_THREADID();
+				int32 LocalCreationNotificationIndex = INDEX_NONE;
+				if (CreationNotificationIndex.compare_exchange_weak(LocalCreationNotificationIndex, BufferedNotifications.Num()))
+				{
+					BufferedNotifications.Emplace(EObservedOperationNotification::Create
+						, FBufferedNotification::FEmptyComposition{}
+						, CreatedEntity);
+				}
+				else
+				{
+					BufferedNotifications[LocalCreationNotificationIndex].AddHandle(CreatedEntity);
+				}
+
+				return CreationNotificationIndex;
+			}
+
+			int32 AddCreatedEntities(const TConstArrayView<FMassEntityHandle> InCreatedEntities)
+			{
+				UE_CHECK_OWNER_THREADID();
+				int32 LocalCreationNotificationIndex = INDEX_NONE;
+				if (CreationNotificationIndex.compare_exchange_weak(LocalCreationNotificationIndex, BufferedNotifications.Num()))
+				{
+					BufferedNotifications.Emplace(EObservedOperationNotification::Create
+						, FBufferedNotification::FEmptyComposition{}
+						, UE::Mass::FEntityCollection(InCreatedEntities));
+				}
+				else
+				{
+					BufferedNotifications[LocalCreationNotificationIndex].AppendEntities(InCreatedEntities);
+				}
+
+				return CreationNotificationIndex;
+			}
+
+			int32 AddCreatedEntities(const TConstArrayView<FMassEntityHandle> InCreatedEntities, FMassArchetypeEntityCollection&& InEntityCollection)
+			{
+				UE_CHECK_OWNER_THREADID();
+				int32 LocalCreationNotificationIndex = INDEX_NONE;
+				if (CreationNotificationIndex.compare_exchange_weak(LocalCreationNotificationIndex, BufferedNotifications.Num()))
+				{
+					BufferedNotifications.Emplace(EObservedOperationNotification::Create, FBufferedNotification::FEmptyComposition{}
+						, UE::Mass::FEntityCollection(InCreatedEntities, Forward<FMassArchetypeEntityCollection>(InEntityCollection)));
+				}
+				else
+				{
+					BufferedNotifications[LocalCreationNotificationIndex].AppendEntities(InCreatedEntities, Forward<FMassArchetypeEntityCollection>(InEntityCollection));
+				}
+
+				return CreationNotificationIndex;
+			}
+
+			template<typename T> requires std::is_same_v<typename TDecay<T>::Type, FMassArchetypeEntityCollection>
+			int32 AddCreatedEntitiesCollection(T&& InEntityCollection)
+			{
+				UE_CHECK_OWNER_THREADID();
+				int32 LocalCreationNotificationIndex = INDEX_NONE;
+				if (CreationNotificationIndex.compare_exchange_weak(LocalCreationNotificationIndex, BufferedNotifications.Num()))
+				{
+					BufferedNotifications.Emplace(EObservedOperationNotification::Create
+						, FBufferedNotification::FEmptyComposition{}
+						, UE::Mass::FEntityCollection(Forward<T>(InEntityCollection)));
+				}
+				else
+				{
+					BufferedNotifications[LocalCreationNotificationIndex].AppendEntities<T>(InEntityCollection);
+				}
+
+				return CreationNotificationIndex;
+			}
+
+			template<typename TEntities>
+			void AddNotification(const EObservedOperationNotification NotificationType
+				, TEntities&& Entities
+				, const bool bHasFragmentsOverlap, FMassFragmentBitSet&& FragmentOverlap
+				, const bool bHasTagsOverlap, FMassTagBitSet&& TagOverlap)
+			{
+				checkSlow(bHasFragmentsOverlap || bHasTagsOverlap);
+				FBufferedNotification::FCompositionDescription CompositionChange = (bHasFragmentsOverlap != bHasTagsOverlap)
+					? (bHasFragmentsOverlap
+						? FBufferedNotification::FCompositionDescription(TInPlaceType<FMassFragmentBitSet>(), MoveTemp(FragmentOverlap))
+						: FBufferedNotification::FCompositionDescription(TInPlaceType<FMassTagBitSet>(), MoveTemp(TagOverlap)))
+					: FBufferedNotification::FCompositionDescription(TInPlaceType<FMassArchetypeCompositionDescriptor>(), FMassArchetypeCompositionDescriptor(MoveTemp(FragmentOverlap), MoveTemp(TagOverlap), {}, {}, {}));
+
+				if (BufferedNotifications.Num() 
+					&& BufferedNotifications.Last().Type == NotificationType
+					&& FBufferedNotification::AreCompositionsEqual(BufferedNotifications.Last().CompositionChange, CompositionChange))
+				{
+					BufferedNotifications.Last().AppendEntities(Forward<TEntities>(Entities));
+				}
+				else
+				{
+					BufferedNotifications.Emplace(NotificationType, MoveTemp(CompositionChange), Forward<TEntities>(Entities));
+				}
+			}
+
+			friend FMassObserverManager;
+			/** To be called in case of processor forking. */
+			void ForceUpdateCurrentThreadID();
+
+			/**
+			 * Identifies the thread where given FObserverLock instance was created. All subsequent operations are 
+			 * expected to be run in the same thread.
+			 */
+			uint32 OwnerThreadId;
+
+			std::atomic<int32> CreationNotificationIndex = INDEX_NONE;
+
+			TArray<FBufferedNotification> BufferedNotifications;
+
+			FMassArchetypeEntityCollection::EDuplicatesHandling CollectionCreationDuplicatesHandling = FMassArchetypeEntityCollection::EDuplicatesHandling::NoDuplicates;
+
+			/** Point to outer EntityManager. Used to obtain ObserverManager in type's destructor*/
+			TWeakPtr<FMassEntityManager> WeakEntityManager;
+
+	#if WITH_MASSENTITY_DEBUG
+			uint32 LockSerialNumber = 0;
+	#endif // WITH_MASSENTITY_DEBUG
+		};
+
+		/**
+		 * A dedicated structure for ensuring the "on entities creation" observers get notified only once all other 
+		 * initialization operations are done and this creation context instance gets released. 
+		 */
+		struct MASSENTITY_API FCreationContext
+		{
+			TArray<FMassArchetypeEntityCollection> GetEntityCollections(const FMassEntityManager& EntityManager) const;
+
+			/** Function for debugging/testing purposes. We don't expect users to ever call it, always get collections via GetEntityCollections */
+			bool DebugAreEntityCollectionsUpToDate() const;
+
+			~FCreationContext();
+
+			static TSharedRef<FCreationContext> DebugCreateDummyCreationContext();
+
+			UE_DEPRECATED(5.6, "Use the other GetEntityCollections flavor insteand")
+			TConstArrayView<FMassArchetypeEntityCollection> GetEntityCollections() const;
+			UE_DEPRECATED(5.6, "Functionality no longer available")
+			int32 GetSpawnedNum() const;
+			UE_DEPRECATED(5.6, "Do not use, internal use only")
+			bool IsDirty() const;
+			UE_DEPRECATED(5.6, "Manually adding entities directly to the creation context is not longer supported and is not taking place automatically")
+			void AppendEntities(const TConstArrayView<FMassEntityHandle>);
+			UE_DEPRECATED(5.6, "Manually adding entities directly to the creation context is not longer supported and is not taking place automatically")
+			void AppendEntities(const TConstArrayView<FMassEntityHandle>, FMassArchetypeEntityCollection&&);
+			UE_DEPRECATED(5.5, "This constructor is now deprecated and defunct. Use one of the others instead.")
+			explicit FCreationContext(const int32);
+			UE_DEPRECATED(5.5, "This function is now deprecated since FEntityCreationContext can contain more than a single collection now. Use GetEntityCollections instead.")
+			const FMassArchetypeEntityCollection& GetEntityCollection() const;
+			
+			/**
+			 *	Called in response to composition mutating operation - these operations invalidate stored collections
+			 */
+			UE_DEPRECATED_FORGAME(5.6, "Do not use, internal use only")
+			void MarkDirty()
+			{
+				Lock->MarkCreationNotificationDirty(CreationHandle);
+			}
+
+		private:
+			FCreationContext();
+
+			FCreationContext(TSharedRef<FObserverLock>&& InLock)
+				: Lock(Forward<TSharedRef<FObserverLock>>(InLock))
+			{	
+			}
+			FCreationContext(const TSharedRef<FObserverLock>& InLock)
+				: Lock(InLock)
+			{
+			}
+
+			TSharedRef<FObserverLock> GetObserverLock() const
+			{
+				return Lock;
+			}
+
+			bool IsValid() const
+			{
+				return CreationHandle.IsSet();
+			}
+
+			friend FMassObserverManager;
+			TSharedRef<FObserverLock> Lock;
+			FCreationNotificationHandle CreationHandle;
+		};
+	} // namespace UE::Mass::ObserverManager
+} // namespace UE::Mass
+
+#undef UE_CHECK_OWNER_THREADID
\ No newline at end of file
