[Mass] switched UMassObserverProcessor.ObservedType to TObjectPtr<const UScriptStruct>

#rnx

diff --git a/Engine/Source/Runtime/MassEntity/Public/MassObserverProcessor.h b/Engine/Source/Runtime/MassEntity/Public/MassObserverProcessor.h
--- a/Engine/Source/Runtime/MassEntity/Public/MassObserverProcessor.h
+++ b/Engine/Source/Runtime/MassEntity/Public/MassObserverProcessor.h
@@ -29,7 +29,7 @@ protected:
 
 	/** Determines which Fragment or Tag type this given UMassObserverProcessor will be observing */
 	UPROPERTY()
-	TObjectPtr<UScriptStruct> ObservedType = nullptr;
+	TObjectPtr<const UScriptStruct> ObservedType = nullptr;
 
 	EMassObservedOperation Operation = EMassObservedOperation::MAX;
 };
