[Mass] Improved MassObserverProcessor's auto-registration issues reporting

#rnx

diff --git a/Engine/Source/Runtime/MassEntity/Private/MassObserverProcessor.cpp b/Engine/Source/Runtime/MassEntity/Private/MassObserverProcessor.cpp
--- a/Engine/Source/Runtime/MassEntity/Private/MassObserverProcessor.cpp
+++ b/Engine/Source/Runtime/MassEntity/Private/MassObserverProcessor.cpp
@@ -25,11 +25,11 @@ void UMassObserverProcessor::PostInitProperties()
 
 	if (HasAnyFlags(RF_ClassDefaultObject) && MyClass->HasAnyClassFlags(CLASS_Abstract) == false)
 	{
-		if (ensure(ObservedType != nullptr && Operation != EMassObservedOperation::MAX))
+		if (ObservedType != nullptr && Operation != EMassObservedOperation::MAX)
 		{
 			Register();
 		}
-		else
+		else if (bAutoRegisterWithObserverRegistry)
 		{
 			UE_LOG(LogMass, Error, TEXT("%s attempting to register %s while it\'s misconfigured, Type: %s, Operation: %s")
 				, ANSI_TO_TCHAR(__FUNCTION__), *MyClass->GetName(), *GetNameSafe(ObservedType), *UEnum::GetValueAsString(Operation));
