[Mass] fixed FMassObserverManager::Initialize stomping any observer processor instances registered beforehand via FMassObserverManager::AddObserverInstance

diff --git a/Engine/Source/Runtime/MassEntity/Private/MassObserverManager.cpp b/Engine/Source/Runtime/MassEntity/Private/MassObserverManager.cpp
--- a/Engine/Source/Runtime/MassEntity/Private/MassObserverManager.cpp
+++ b/Engine/Source/Runtime/MassEntity/Private/MassObserverManager.cpp
@@ -35,11 +35,9 @@ namespace UE::Mass::ObserverManager
 
 		// a helper function to reduce code duplication in FMassObserverManager::Initialize
 		template<typename TBitSet, typename TPointerType>
-		void SetUpObservers(FMassEntityManager& EntityManager, const EProcessorExecutionFlags WorldExecutionFlags, UObject& Owner
-			, const TMap<TPointerType, FMassProcessorClassCollection>& RegisteredObserverTypes, TBitSet& ObservedBitSet, FMassObserversMap& Observers)
+		void AddRegisteredObserverProcessorInstances(FMassEntityManager& EntityManager, const EProcessorExecutionFlags WorldExecutionFlags, UObject& Owner
+			, const TMap<TPointerType, FMassProcessorClassCollection>& RegisteredObserverTypes, TBitSet& InOutObservedBitSet, FMassObserversMap& Observers)
 		{
-			ObservedBitSet.Reset();
-
 			for (auto It : RegisteredObserverTypes)
 			{
 				if (It.Value.ClassCollection.Num() == 0)
@@ -47,7 +45,7 @@ namespace UE::Mass::ObserverManager
 					continue;
 				}
 
-				ObservedBitSet.Add(*It.Key);
+				InOutObservedBitSet.Add(*It.Key);
 				FMassRuntimePipeline& Pipeline = (*Observers).FindOrAdd(It.Key);
 
 				for (const TSubclassOf<UMassProcessor>& ProcessorClass : It.Value.ClassCollection)
@@ -250,11 +248,11 @@ void FMassObserverManager::Initialize()
 	const UWorld* World = Owner->GetWorld();
 	const EProcessorExecutionFlags WorldExecutionFlags = UE::Mass::Utils::DetermineProcessorExecutionFlags(World);
 
-	using UE::Mass::ObserverManager::Private::SetUpObservers;
+	using UE::Mass::ObserverManager::Private::AddRegisteredObserverProcessorInstances;
 	for (int i = 0; i < (int)EMassObservedOperation::MAX; ++i)
 	{
-		SetUpObservers(EntityManager, WorldExecutionFlags, *Owner, *Registry.FragmentObservers[i], ObservedFragments[i], FragmentObservers[i]);
-		SetUpObservers(EntityManager, WorldExecutionFlags, *Owner, *Registry.TagObservers[i], ObservedTags[i], TagObservers[i]);
+		AddRegisteredObserverProcessorInstances(EntityManager, WorldExecutionFlags, *Owner, *Registry.FragmentObservers[i], ObservedFragments[i], FragmentObservers[i]);
+		AddRegisteredObserverProcessorInstances(EntityManager, WorldExecutionFlags, *Owner, *Registry.TagObservers[i], ObservedTags[i], TagObservers[i]);
 	}
 
 #if WITH_MASSENTITY_DEBUG
