[AutoRTFM] Update LocksCount member variable for AutoRTFM safety.

Like other locking structures in AutoRTFM, the lock will remain held until the transaction either succeeds or aborts.

Note that this is not a complete fix for the issue on Horde as there is a second atomic as well.

diff --git a/Engine/Source/Runtime/MassEntity/Private/MassObserverManager.cpp b/Engine/Source/Runtime/MassEntity/Private/MassObserverManager.cpp
--- a/Engine/Source/Runtime/MassEntity/Private/MassObserverManager.cpp
+++ b/Engine/Source/Runtime/MassEntity/Private/MassObserverManager.cpp
@@ -281,7 +281,7 @@ bool FMassObserverManager::OnPostEntitiesCreated(const FMassArchetypeEntityColle
 {
 	TRACE_CPUPROFILER_EVENT_SCOPE(MassObserver_OnPostEntitiesCreated_Collection);
 
-	if (LocksCount)
+	if (HasAnyObserverLocks())
 	{
 		TSharedRef<FObserverLock> ObserverLock = ActiveObserverLock.Pin().ToSharedRef();
 		ObserverLock->AddCreatedEntitiesCollection(EntityCollection);
@@ -296,7 +296,7 @@ bool FMassObserverManager::OnPostEntityCreated(FMassEntityHandle EntityHandle, c
 {
 	TRACE_CPUPROFILER_EVENT_SCOPE(MassObserver_OnPostEntitiesCreated_Collection);
 
-	if (LocksCount)
+	if (HasAnyObserverLocks())
 	{
 		TSharedRef<FObserverLock> ObserverLock = ActiveObserverLock.Pin().ToSharedRef();
 		ObserverLock->AddCreatedEntity(EntityHandle);
@@ -349,7 +349,7 @@ bool FMassObserverManager::OnCompositionChanged(FCollectionRefOrHandle&& EntityC
 
 	TRACE_CPUPROFILER_EVENT_SCOPE(MassObserver_OnCompositionChanged);
 
-	if (LocksCount)
+	if (HasAnyObserverLocks())
 	{
 		if (TSharedPtr<FCreationContext> CreationContext = GetCreationContext())
 		{
@@ -370,7 +370,7 @@ bool FMassObserverManager::OnCompositionChanged(FCollectionRefOrHandle&& EntityC
 
 	if (bHasFragmentsOverlap || bHasTagsOverlap)
 	{
-		if (LocksCount)
+		if (HasAnyObserverLocks())
 		{
 			const UE::Mass::ObserverManager::EObservedOperationNotification NotificationType = static_cast<UE::Mass::ObserverManager::EObservedOperationNotification>(Operation);
 			TSharedRef<FObserverLock> ObserverLockRef = ActiveObserverLock.Pin().ToSharedRef();
@@ -477,7 +477,7 @@ bool FMassObserverManager::OnCollectionsCreatedImpl(UE::Mass::FProcessingContext
 
 	TRACE_CPUPROFILER_EVENT_SCOPE(MassObserver_OnCollectionsCreatedImpl_Collection);
 
-	check(LocksCount == 0);
+	check(!HasAnyObserverLocks());
 
 	constexpr int32 OperationIndex = static_cast<int32>(EMassObservedOperation::Add);
 
@@ -668,7 +668,7 @@ void FMassObserverManager::ResumeExecution(FObserverLock& LockBeingReleased)
 {
 	using namespace UE::Mass::ObserverManager;
 
-	ensureMsgf(LocksCount == 0, TEXT("We only expect this function to be called if all locks are released."));
+	ensureMsgf(!HasAnyObserverLocks(), TEXT("We only expect this function to be called if all locks are released."));
 #if WITH_MASSENTITY_DEBUG
 	ensureMsgf(LockBeingReleased.LockSerialNumber == LockedNotificationSerialNumber
 		, TEXT("Lock's and ObserverManager's lock serial numbers are expected to match."));
diff --git a/Engine/Source/Runtime/MassEntity/Private/MassObserverNotificationTypes.cpp b/Engine/Source/Runtime/MassEntity/Private/MassObserverNotificationTypes.cpp
--- a/Engine/Source/Runtime/MassEntity/Private/MassObserverNotificationTypes.cpp
+++ b/Engine/Source/Runtime/MassEntity/Private/MassObserverNotificationTypes.cpp
@@ -4,6 +4,7 @@
 #include "MassObserverManager.h"
 #include "MassEntityManager.h"
 #include "MassEntityUtils.h"
+#include "AutoRTFM.h"
 
 namespace UE::Mass::ObserverManager
 {
@@ -26,7 +27,15 @@ namespace UE::Mass::ObserverManager
 		, LockSerialNumber(ObserverManager.LockedNotificationSerialNumber)
 	#endif // WITH_MASSENTITY_DEBUG
 	{
-		++ObserverManager.LocksCount;
+		FMassObserverManager* ObserverManagerPtr = &ObserverManager;
+		UE_AUTORTFM_OPEN
+		{
+			++ObserverManagerPtr->LocksCount;
+		};
+		UE_AUTORTFM_ONABORT(ObserverManagerPtr)
+		{
+			--ObserverManagerPtr->LocksCount;
+		};
 	}
 
 	FObserverLock::~FObserverLock()
@@ -37,10 +46,16 @@ namespace UE::Mass::ObserverManager
 			return;
 		}
 
-		--SharedEntityManager->GetObserverManager().LocksCount;
-		checkf(SharedEntityManager->GetObserverManager().LocksCount >= 0
-			, TEXT("%hs: getting a negative locks count as a result of the operation is unexpected."), __FUNCTION__);
-		SharedEntityManager->GetObserverManager().ResumeExecution(*this);
+		FMassObserverManager& ObserverManager = SharedEntityManager->GetObserverManager();
+
+		UE_AUTORTFM_ONCOMMIT(SharedEntityManager = MoveTemp(SharedEntityManager))
+		{
+			--SharedEntityManager->GetObserverManager().LocksCount;
+			checkf(SharedEntityManager->GetObserverManager().LocksCount >= 0
+				, TEXT("%hs: the lock count has become unbalanced."), __FUNCTION__);
+		};
+
+		ObserverManager.ResumeExecution(*this);
 	}
 
 	void FObserverLock::ForceUpdateCurrentThreadID()
diff --git a/Engine/Source/Runtime/MassEntity/Public/MassObserverManager.h b/Engine/Source/Runtime/MassEntity/Public/MassObserverManager.h
--- a/Engine/Source/Runtime/MassEntity/Public/MassObserverManager.h
+++ b/Engine/Source/Runtime/MassEntity/Public/MassObserverManager.h
@@ -6,6 +6,7 @@
 #include "MassProcessingTypes.h"
 #include "MassEntityTypes.h"
 #include "MassArchetypeTypes.h"
+#include "AutoRTFM.h"
 #include "Misc/Fork.h"
 #if UE_ENABLE_INCLUDE_ORDER_DEPRECATED_IN_5_6
 // the following is a new header, but it contains some types moved over from other places, and including
@@ -249,6 +250,12 @@ protected:
 	std::atomic<int32> LocksCount = 0;
 	TWeakPtr<FCreationContext> ActiveCreationContext;
 
+	UE_AUTORTFM_ALWAYS_OPEN
+	bool HasAnyObserverLocks()
+	{
+		return LocksCount.load(std::memory_order_relaxed) > 0;
+	}
+
 	FMassFragmentBitSet ObservedFragments[(uint8)EMassObservedOperation::MAX];
 	FMassTagBitSet ObservedTags[(uint8)EMassObservedOperation::MAX];
 
