[Mass] More tightening of access control to MassProcessingContext's internals, detal time and bFlushCommandBuffer can not only be set via the constructor, and read via getters.

#rnx

diff --git a/Engine/Source/Developer/MassEntityTestSuite/Private/MassEntityBuilderTest.cpp b/Engine/Source/Developer/MassEntityTestSuite/Private/MassEntityBuilderTest.cpp
--- a/Engine/Source/Developer/MassEntityTestSuite/Private/MassEntityBuilderTest.cpp
+++ b/Engine/Source/Developer/MassEntityTestSuite/Private/MassEntityBuilderTest.cpp
@@ -407,8 +407,7 @@ namespace UE::Mass::Test::EntityBuilder
 			};
 			Processor->EntityQuery.AddRequirement<FTestFragment_Int>(EMassFragmentAccess::ReadOnly);
 
-			FMassProcessingContext ProcessingContext(EntityManager, /*DeltaSeconds=*/0.f);
-			ProcessingContext.bFlushCommandBuffer = false;
+			FMassProcessingContext ProcessingContext(EntityManager, /*DeltaSeconds=*/0.f, /*bFlushCommandBuffer=*/false);
 
 			for (int32 Iteration = 0; Iteration < NumIterations; ++Iteration)
 			{
@@ -460,8 +459,7 @@ namespace UE::Mass::Test::EntityBuilder
 			};
 			Processor->EntityQuery.AddRequirement<FTestFragment_Int>(EMassFragmentAccess::ReadOnly);
 
-			FMassProcessingContext ProcessingContext(EntityManager, /*DeltaSeconds=*/0.f);
-			ProcessingContext.bFlushCommandBuffer = false;
+			FMassProcessingContext ProcessingContext(EntityManager, /*DeltaSeconds=*/0.f, /*bFlushCommandBuffer=*/false);
 			Executor::RunProcessorsView(MakeArrayView(reinterpret_cast<UMassProcessor**>(&Processor), 1), ProcessingContext);
 
 			AITEST_EQUAL("Number of fully-formed entities expected", ProcessedEntitiesCount, 1);
diff --git a/Engine/Source/Runtime/MassEntity/Private/MassExecutor.cpp b/Engine/Source/Runtime/MassEntity/Private/MassExecutor.cpp
--- a/Engine/Source/Runtime/MassEntity/Private/MassExecutor.cpp
+++ b/Engine/Source/Runtime/MassEntity/Private/MassExecutor.cpp
@@ -24,7 +24,7 @@ FORCEINLINE void ExecuteProcessors(FMassEntityManager& EntityManager, TArrayView
 
 void Run(FMassRuntimePipeline& RuntimePipeline, FProcessingContext& ProcessingContext)
 {
-	if (!ensure(ProcessingContext.DeltaSeconds >= 0.f) 
+	if (!ensure(ProcessingContext.GetDeltaSeconds() >= 0.f) 
 		|| !ensure(RuntimePipeline.GetProcessors().Find(nullptr) == INDEX_NONE))
 	{
 		return;
@@ -65,7 +65,7 @@ void RunSparse(FMassRuntimePipeline& RuntimePipeline, FProcessingContext& Proces
 
 void Run(UMassProcessor& Processor, FProcessingContext& ProcessingContext)
 {
-	if (!ensure(ProcessingContext.DeltaSeconds >= 0.f))
+	if (!ensure(ProcessingContext.GetDeltaSeconds() >= 0.f))
 	{
 		return;
 	}
diff --git a/Engine/Source/Runtime/MassEntity/Public/MassProcessingContext.h b/Engine/Source/Runtime/MassEntity/Public/MassProcessingContext.h
--- a/Engine/Source/Runtime/MassEntity/Public/MassProcessingContext.h
+++ b/Engine/Source/Runtime/MassEntity/Public/MassProcessingContext.h
@@ -19,14 +19,16 @@ namespace UE::Mass
 
 		FProcessingContext(FProcessingContext&&) = default;
 
-		explicit FProcessingContext(FMassEntityManager& InEntityManager, const float InDeltaSeconds = 0.f);
-		explicit FProcessingContext(const TSharedRef<FMassEntityManager>& InEntityManager, const float InDeltaSeconds = 0.f);
-		explicit FProcessingContext(TSharedRef<FMassEntityManager>&& InEntityManager, const float InDeltaSeconds = 0.f);
-		explicit FProcessingContext(const TSharedPtr<FMassEntityManager>& InEntityManager, const float InDeltaSeconds = 0.f);
+		explicit FProcessingContext(FMassEntityManager& InEntityManager, const float InDeltaSeconds = 0.f, const bool bInFlushCommandBuffer = true);
+		explicit FProcessingContext(const TSharedRef<FMassEntityManager>& InEntityManager, const float InDeltaSeconds = 0.f, const bool bInFlushCommandBuffer = true);
+		explicit FProcessingContext(TSharedRef<FMassEntityManager>&& InEntityManager, const float InDeltaSeconds = 0.f, const bool bInFlushCommandBuffer = true);
+		explicit FProcessingContext(const TSharedPtr<FMassEntityManager>& InEntityManager, const float InDeltaSeconds = 0.f, const bool bInFlushCommandBuffer = true);
 		~FProcessingContext();
 		
 		[[nodiscard]] FMassExecutionContext& GetExecutionContext() &;
 		[[nodiscard]] FMassExecutionContext&& GetExecutionContext() &&;
+		bool GetWillFlushCommands() const;
+		float GetDeltaSeconds() const;
 
 		void SetCommandBuffer(TSharedPtr<FMassCommandBuffer>&& InCommandBuffer);
 		void SetCommandBuffer(const TSharedPtr<FMassCommandBuffer>& InCommandBuffer);
@@ -36,6 +38,7 @@ namespace UE::Mass
 		UE_DEPRECATED_FORGAME(5.6, "Direct access to FProcessingContext.EntityManager has been deprecated. Use GetEntityManager instead.")
 		TSharedRef<FMassEntityManager> EntityManager;
 
+		UE_DEPRECATED_FORGAME(5.6, "Direct access to FProcessingContex.DeltaSeconds has been deprecated. Set it via the constructor instead, read via the getter.")
 		float DeltaSeconds = 0.f;
 
 		FInstancedStruct AuxData;
@@ -44,6 +47,7 @@ namespace UE::Mass
 		 * If set to "true" the MassExecutor will flush commands at the end of given execution function. 
 		 * If "false" the caller is responsible for manually flushing the commands.
 		 */
+		UE_DEPRECATED_FORGAME(5.6, "Direct access to FProcessingContex.bFlushCommandBuffer has been deprecated. Set it via the constructor instead, read via the GetWillFlushCommands getter.")
 		bool bFlushCommandBuffer = true; 
 
 		UE_DEPRECATED_FORGAME(5.6, "Direct access to FProcessingContext's CommandBuffer has been deprecated. Use SetCommandBuffer()")
@@ -59,28 +63,30 @@ namespace UE::Mass
 	//----------------------------------------------------------------------//
 	// @todo remove the depracation disabling once CommandBuffer and EntityManager are moved to `protected` and un-deprecated (around UE5.8)
 	PRAGMA_DISABLE_DEPRECATION_WARNINGS
-	inline FProcessingContext::FProcessingContext(FMassEntityManager& InEntityManager, const float InDeltaSeconds)
-		: FProcessingContext(/*MoveTemp*/(InEntityManager.AsShared()), InDeltaSeconds)
+	inline FProcessingContext::FProcessingContext(FMassEntityManager& InEntityManager, const float InDeltaSeconds, const bool bInFlushCommandBuffer)
+		: FProcessingContext(/*MoveTemp*/(InEntityManager.AsShared()), InDeltaSeconds, bInFlushCommandBuffer)
 	{
 
 	}
 
-	inline FProcessingContext::FProcessingContext(const TSharedRef<FMassEntityManager>& InEntityManager, const float InDeltaSeconds)
+	inline FProcessingContext::FProcessingContext(const TSharedRef<FMassEntityManager>& InEntityManager, const float InDeltaSeconds, const bool bInFlushCommandBuffer)
 		: EntityManager(InEntityManager)
 		, DeltaSeconds(InDeltaSeconds)
+		, bFlushCommandBuffer(bInFlushCommandBuffer)
 	{
 		
 	}
 
-	inline FProcessingContext::FProcessingContext(TSharedRef<FMassEntityManager>&& InEntityManager, const float InDeltaSeconds)
+	inline FProcessingContext::FProcessingContext(TSharedRef<FMassEntityManager>&& InEntityManager, const float InDeltaSeconds, const bool bInFlushCommandBuffer)
 		: EntityManager(MoveTemp(InEntityManager))
 		, DeltaSeconds(InDeltaSeconds)
+		, bFlushCommandBuffer(bInFlushCommandBuffer)
 	{
 		
 	}
 
-	inline FProcessingContext::FProcessingContext(const TSharedPtr<FMassEntityManager>& InEntityManager, const float InDeltaSeconds)
-		: FProcessingContext(InEntityManager.ToSharedRef(), InDeltaSeconds)
+	inline FProcessingContext::FProcessingContext(const TSharedPtr<FMassEntityManager>& InEntityManager, const float InDeltaSeconds, const bool bInFlushCommandBuffer)
+		: FProcessingContext(InEntityManager.ToSharedRef(), InDeltaSeconds, bInFlushCommandBuffer)
 	{
 		
 	}
@@ -142,5 +148,15 @@ namespace UE::Mass
 		return EntityManager;
 	}
 
+	inline bool FProcessingContext::GetWillFlushCommands() const
+	{
+		return bFlushCommandBuffer;
+	}
+
+	inline float FProcessingContext::GetDeltaSeconds() const
+	{
+		return DeltaSeconds;
+	}
+
 	PRAGMA_ENABLE_DEPRECATION_WARNINGS
 } // namespace UE::Mass
