[Mass]

Removed remaining atomic declaration from MassObservers.

diff --git a/Engine/Source/Runtime/MassEntity/Public/MassObserverNotificationTypes.h b/Engine/Source/Runtime/MassEntity/Public/MassObserverNotificationTypes.h
--- a/Engine/Source/Runtime/MassEntity/Public/MassObserverNotificationTypes.h
+++ b/Engine/Source/Runtime/MassEntity/Public/MassObserverNotificationTypes.h
@@ -379,7 +379,7 @@ namespace UE::Mass
 			 */
 			uint32 OwnerThreadId;
 
-			std::atomic<int32> CreationNotificationIndex = INDEX_NONE;
+			int32 CreationNotificationIndex = INDEX_NONE;
 
 			TArray<FBufferedNotification> BufferedNotifications;
 
