[Mass] DebugGetEntityIndexHandle now only returns a valid handle if the entity is fully created preventing debug tools from using a reserved entity in operations that require a fully built entity.

diff --git a/Engine/Source/Runtime/MassEntity/Private/MassEntityManager.cpp b/Engine/Source/Runtime/MassEntity/Private/MassEntityManager.cpp
--- a/Engine/Source/Runtime/MassEntity/Private/MassEntityManager.cpp
+++ b/Engine/Source/Runtime/MassEntity/Private/MassEntityManager.cpp
@@ -2120,7 +2120,10 @@ void FMassEntityManager::DebugGetArchetypeStrings(const FMassArchetypeHandle& Ar
 
 FMassEntityHandle FMassEntityManager::DebugGetEntityIndexHandle(const int32 EntityIndex) const
 {
-	return GetEntityStorageInterface().IsValidIndex(EntityIndex) ? FMassEntityHandle(EntityIndex, GetEntityStorageInterface().GetSerialNumber(EntityIndex)) : FMassEntityHandle();
+	return (GetEntityStorageInterface().IsValidIndex(EntityIndex)
+		&& GetEntityStorageInterface().GetEntityState(EntityIndex) == UE::Mass::IEntityStorageInterface::EEntityState::Created)
+		? FMassEntityHandle(EntityIndex, GetEntityStorageInterface().GetSerialNumber(EntityIndex))
+		: FMassEntityHandle();
 }
 
 const FString& FMassEntityManager::DebugGetName() const
