[Mass] Doubled-down on UE::Mass::FTypeHandle use in UE::Mass::FTypeManager operations, it's replacing direct use of ObjectKey types.

#rnx

diff --git a/Engine/Source/Runtime/MassEntity/Private/MassProcessorDependencySolver.cpp b/Engine/Source/Runtime/MassEntity/Private/MassProcessorDependencySolver.cpp
--- a/Engine/Source/Runtime/MassEntity/Private/MassProcessorDependencySolver.cpp
+++ b/Engine/Source/Runtime/MassEntity/Private/MassProcessorDependencySolver.cpp
@@ -1016,7 +1016,7 @@ void FMassProcessorDependencySolver::GatherSubsystemInformation(const UE::Mass::
 			check(SubsystemTraits);
 			if (SubsystemTraits->bThreadSafeWrite)
 			{
-				const UClass* SubsystemClass = Cast<UClass>(SubsystemTypeIterator->ResolveObjectPtr());
+				const UClass* SubsystemClass = SubsystemTypeIterator->GetClass();
 				check(SubsystemClass);
 				MultiThreadedSystemsBitSet.Add(*SubsystemClass);
 			}
diff --git a/Engine/Source/Runtime/MassEntity/Private/MassTypeManager.cpp b/Engine/Source/Runtime/MassEntity/Private/MassTypeManager.cpp
--- a/Engine/Source/Runtime/MassEntity/Private/MassTypeManager.cpp
+++ b/Engine/Source/Runtime/MassEntity/Private/MassTypeManager.cpp
@@ -25,18 +25,18 @@ namespace UE::Mass
 	{
 	}
 
-	FTypeHandle FTypeManager::RegisterType(TNotNull<const UStruct*> InType, FTypeInfo&& TypeInfo)
+	FTypeHandle FTypeManager::RegisterTypeInternal(TNotNull<const UStruct*> InType, FTypeInfo&& TypeInfo)
 	{
 		FTypeHandle TypeHandle(InType);
-		FTypeInfo* ExistingData = TypeDataMap.Find(TypeHandle.TypeKey);
+		FTypeInfo* ExistingData = TypeDataMap.Find(FTypeHandle(TypeHandle.TypeKey));
 		if (LIKELY(ExistingData == nullptr))
 		{
 			if (TypeInfo.Traits.IsType<FSubsystemTypeTraits>())
 			{
-				SubsystemTypes.Add(TypeHandle.TypeKey);
+				SubsystemTypes.Add(FTypeHandle(TypeHandle.TypeKey));
 			}
 
-			TypeDataMap.Add(TypeHandle.TypeKey, MoveTemp(TypeInfo));
+			TypeDataMap.Add(FTypeHandle(TypeHandle.TypeKey), MoveTemp(TypeInfo));
 		}
 		else
 		{
@@ -55,7 +55,7 @@ namespace UE::Mass
 		FTypeInfo TypeInfo;
 		TypeInfo.TypeName = InType->GetFName();
 		TypeInfo.Traits.Set<FSubsystemTypeTraits>(MoveTemp(TypeTraits));
-		return RegisterType(InType, MoveTemp(TypeInfo));
+		return RegisterTypeInternal(InType, MoveTemp(TypeInfo));
 	}
 
 	FTypeHandle FTypeManager::RegisterType(TNotNull<const UStruct*> InType, FSharedFragmentTypeTraits&& TypeTraits)
@@ -66,11 +66,6 @@ namespace UE::Mass
 		FTypeInfo TypeInfo;
 		TypeInfo.TypeName = InType->GetFName();
 		TypeInfo.Traits.Set<FSharedFragmentTypeTraits>(MoveTemp(TypeTraits));
-		return RegisterType(InType, MoveTemp(TypeInfo));
-	}
-
-	const FTypeInfo* FTypeManager::GetTypeInfo(TObjectKey<const UStruct> TypeKey) const
-	{
-		return TypeDataMap.Find(TypeKey);
+		return RegisterTypeInternal(InType, MoveTemp(TypeInfo));
 	}
 }
\ No newline at end of file
diff --git a/Engine/Source/Runtime/MassEntity/Public/MassTypeManager.h b/Engine/Source/Runtime/MassEntity/Public/MassTypeManager.h
--- a/Engine/Source/Runtime/MassEntity/Public/MassTypeManager.h
+++ b/Engine/Source/Runtime/MassEntity/Public/MassTypeManager.h
@@ -17,12 +17,29 @@ struct FMassEntityManager;
 namespace UE::Mass
 {
 	/** Handle for identifying and managing types in the type manager */
-	struct FTypeHandle 
+	struct FTypeHandle
 	{
 		FTypeHandle() = default;
 
 		bool operator==(const FTypeHandle&) const = default;
 
+		bool IsValid() const;
+
+		const UClass* GetClass() const
+		{
+			return Cast<const UClass>(TypeKey.ResolveObjectPtr());
+		}
+
+		const UScriptStruct* GetScriptStruct() const
+		{
+			return Cast<UScriptStruct>(TypeKey.ResolveObjectPtr());
+		}
+
+		friend FORCEINLINE uint32 GetTypeHash(const FTypeHandle& InHandle)
+		{
+			return GetTypeHash(InHandle.TypeKey);
+		}
+
 	private:
 		friend struct FTypeManager;
 		explicit FTypeHandle(TObjectKey<const UStruct> InTypeKey);
@@ -101,11 +118,9 @@ namespace UE::Mass
 		 */
 		bool IsEmpty() const;
 
-		/** Register traits for given type */
-		FTypeHandle RegisterType(TNotNull<const UStruct*> InType, FTypeInfo&&);
-		/** Register traits for given type */
+		/** Register traits for given subsystem type */
 		FTypeHandle RegisterType(TNotNull<const UStruct*> InType, FSubsystemTypeTraits&&);
-		/** Register traits for given type */
+		/** Register traits for given shared fragments type */
 		FTypeHandle RegisterType(TNotNull<const UStruct*> InType, FSharedFragmentTypeTraits&&);
 
 		/** Registration helper for shared fragments */
@@ -127,30 +142,43 @@ namespace UE::Mass
 		const FTypeInfo* GetTypeInfo(TObjectKey<const UStruct> TypeKey) const;
 
 		/** Alias for the iterator type that can be used to iterate all of stored type traits information */
-		using FTypeInfoConstIterator = TMap<TObjectKey<const UStruct>, FTypeInfo>::TConstIterator;
+		using FTypeInfoConstIterator = TMap<FTypeHandle, FTypeInfo>::TConstIterator;
+
 		/**
 		 * Alias for the iterator type that can be used to iterate over stored subsystem types.
 		 * Note that this iterator will only point to types, not type trait data. Contents of this iterator
 		 * need to be used with GetTypeInfo to get actual traits data. 
 		 */
-		using FSubsystemTypeConstIterator = TSet<TObjectKey<const UStruct>>::TConstIterator;
+		using FSubsystemTypeConstIterator = TSet<FTypeHandle>::TConstIterator;
 
 		FTypeInfoConstIterator MakeIterator() const;
 		FSubsystemTypeConstIterator MakeSubsystemIterator() const;
 
+		FMassEntityManager& GetEntityManager();
+
 	private:
+		/** Register traits for given type */
+		FTypeHandle RegisterTypeInternal(TNotNull<const UStruct*> InType, FTypeInfo&&);
+
 		FMassEntityManager& OuterEntityManager; 
 
 		/** Mapping of types to their info */
-		TMap<TObjectKey<const UStruct>, FTypeInfo> TypeDataMap;
+		TMap<FTypeHandle, FTypeInfo> TypeDataMap;
 
 		/** Contains all registered subsystem types. Can be used to filter access to TypeDataMap */
-		TSet<TObjectKey<const UStruct>> SubsystemTypes; 
+		TSet<FTypeHandle> SubsystemTypes;
 	};
 
 	//----------------------------------------------------------------------------
 	// INLINES
 	//----------------------------------------------------------------------------
+	inline bool FTypeHandle::IsValid() const
+	{
+		// using this strange contraption since TObjectKey doesn't supply a way to check
+		// whether it's set, while comparison and construction is trivial.
+		// Note: we don't care to what the key's been set, we don't expect types to go away
+		return TypeKey != TObjectKey<const UStruct>();
+	}
 
 	inline const FSubsystemTypeTraits* FTypeInfo::GetAsSystemTraits() const
 	{
@@ -192,6 +220,16 @@ namespace UE::Mass
 
 	inline const FTypeInfo* FTypeManager::GetTypeInfo(FTypeHandle TypeHandle) const
 	{
-		return GetTypeInfo(TypeHandle.TypeKey);
+		return TypeDataMap.Find(TypeHandle);
+	}
+
+	inline const FTypeInfo* FTypeManager::GetTypeInfo(TObjectKey<const UStruct> TypeKey) const
+	{
+		return TypeDataMap.Find(FTypeHandle(TypeKey));
+	}
+
+	inline FMassEntityManager& FTypeManager::GetEntityManager()
+	{
+		return OuterEntityManager;
 	}
 }
