[MassEntity]
* Moved dll storage to methods/variables instead of on types. This is needed to get merged modules working

diff --git a/Engine/Source/Runtime/MassEntity/Public/MassArchetypeTypes.h b/Engine/Source/Runtime/MassEntity/Public/MassArchetypeTypes.h
--- a/Engine/Source/Runtime/MassEntity/Public/MassArchetypeTypes.h
+++ b/Engine/Source/Runtime/MassEntity/Public/MassArchetypeTypes.h
@@ -12,6 +12,8 @@
 #endif // UE_ENABLE_INCLUDE_ORDER_DEPRECATED_IN_5_6
 #include "MassEntityTypes.h"
 
+#define UE_API MASSENTITY_API
+
 struct FMassEntityManager;
 struct FMassArchetypeData;
 struct FMassExecutionContext;
@@ -85,7 +87,7 @@ private:
  *  entity chunks. The goal is to have the user create an instance of this struct once and run through a bunch of
  *  systems. The runtime code usually uses FMassArchetypeChunkIterator to iterate on the chunk collection.
  */
-struct MASSENTITY_API FMassArchetypeEntityCollection 
+struct FMassArchetypeEntityCollection 
 {
 public:
 	struct FArchetypeEntityRange
@@ -163,12 +165,12 @@ private:
 
 public:
 	FMassArchetypeEntityCollection() = default;
-	FMassArchetypeEntityCollection(const FMassArchetypeHandle& InArchetype, TConstArrayView<FMassEntityHandle> InEntities, EDuplicatesHandling DuplicatesHandling);
+	UE_API FMassArchetypeEntityCollection(const FMassArchetypeHandle& InArchetype, TConstArrayView<FMassEntityHandle> InEntities, EDuplicatesHandling DuplicatesHandling);
 	/** optimized, special case for a single-entity */
-	FMassArchetypeEntityCollection(const FMassArchetypeHandle& InArchetype, const FMassEntityHandle EntityHandle);
-	FMassArchetypeEntityCollection(FMassArchetypeHandle&& InArchetype, const FMassEntityHandle EntityHandle);
-	explicit FMassArchetypeEntityCollection(const FMassArchetypeHandle& InArchetypeHandle, const EInitializationType Initialization = EInitializationType::GatherAll);
-	explicit FMassArchetypeEntityCollection(TSharedPtr<FMassArchetypeData>& InArchetype, const EInitializationType Initialization = EInitializationType::GatherAll);
+	UE_API FMassArchetypeEntityCollection(const FMassArchetypeHandle& InArchetype, const FMassEntityHandle EntityHandle);
+	UE_API FMassArchetypeEntityCollection(FMassArchetypeHandle&& InArchetype, const FMassEntityHandle EntityHandle);
+	UE_API explicit FMassArchetypeEntityCollection(const FMassArchetypeHandle& InArchetypeHandle, const EInitializationType Initialization = EInitializationType::GatherAll);
+	UE_API explicit FMassArchetypeEntityCollection(TSharedPtr<FMassArchetypeData>& InArchetype, const EInitializationType Initialization = EInitializationType::GatherAll);
 	FMassArchetypeEntityCollection(const FMassArchetypeHandle& InArchetypeHandle, FEntityRangeArray&& InEntityRanges)
 		: Ranges(Forward<FEntityRangeArray>(InEntityRanges))
 		, Archetype(InArchetypeHandle)
@@ -189,7 +191,7 @@ public:
 	}
 
 	/** The comparison function that checks if Other is identical to this. Intended for diagnostics/debugging. */
-	bool IsSame(const FMassArchetypeEntityCollection& Other) const;
+	UE_API bool IsSame(const FMassArchetypeEntityCollection& Other) const;
 	bool IsSameArchetype(const FMassArchetypeEntityCollection& Other) const;
 
 	/**
@@ -208,24 +210,24 @@ public:
 	 * Note that the operation is only supported for already created entities (i.e. not "reserved")
 	 * @return whether any entity handles have been actually exported
 	 */
-	bool ExportEntityHandles(TArray<FMassEntityHandle>& InOutHandles) const;
+	UE_API bool ExportEntityHandles(TArray<FMassEntityHandle>& InOutHandles) const;
 
-	static bool DoesContainOverlappingRanges(FConstEntityRangeArrayView Ranges);
+	static UE_API bool DoesContainOverlappingRanges(FConstEntityRangeArrayView Ranges);
 
 #if WITH_MASSENTITY_DEBUG
-	int32 DebugCountEntities() const;
+	UE_API int32 DebugCountEntities() const;
 #endif // WITH_MASSENTITY_DEBUG
 
 protected:
 	friend struct FMassArchetypeEntityCollectionWithPayload;
-	void BuildEntityRanges(TStridedView<const int32> TrueIndices);
-	static FArchetypeEntityRange CreateRangeForEntity(const FMassArchetypeHandle& InArchetype, const FMassEntityHandle EntityHandle);
+	UE_API void BuildEntityRanges(TStridedView<const int32> TrueIndices);
+	static UE_API FArchetypeEntityRange CreateRangeForEntity(const FMassArchetypeHandle& InArchetype, const FMassEntityHandle EntityHandle);
 
 private:
 	void GatherChunksFromArchetype();
 };
 
-struct MASSENTITY_API FMassArchetypeEntityCollectionWithPayload
+struct FMassArchetypeEntityCollectionWithPayload
 {
 	explicit FMassArchetypeEntityCollectionWithPayload(const FMassArchetypeEntityCollection& InEntityCollection)
 		: Entities(InEntityCollection)
@@ -237,7 +239,7 @@ struct MASSENTITY_API FMassArchetypeEntityCollectionWithPayload
 	{
 	}
 
-	static void CreateEntityRangesWithPayload(const FMassEntityManager& EntityManager, const TConstArrayView<FMassEntityHandle> Entities
+	static UE_API void CreateEntityRangesWithPayload(const FMassEntityManager& EntityManager, const TConstArrayView<FMassEntityHandle> Entities
 		, const FMassArchetypeEntityCollection::EDuplicatesHandling DuplicatesHandling, FMassGenericPayloadView Payload
 		, TArray<FMassArchetypeEntityCollectionWithPayload>& OutEntityCollections);
 
@@ -258,7 +260,7 @@ private:
  *  The type used to iterate over given archetype's chunks, be it full, continuous chunks or sparse subchunks. It hides
  *  this details from the rest of the system.
  */
-struct MASSENTITY_API FMassArchetypeChunkIterator
+struct FMassArchetypeChunkIterator
 {
 private:
 	FMassArchetypeEntityCollection::FConstEntityRangeArrayView EntityRanges;
@@ -441,3 +443,5 @@ inline bool FMassEntityInChunkDataHandle::operator==(const FMassEntityInChunkDat
 		&& ChunkIndex == Other.ChunkIndex
 		&& ChunkSerialNumber == Other.ChunkSerialNumber;
 }
+
+#undef UE_API 
\ No newline at end of file
diff --git a/Engine/Source/Runtime/MassEntity/Public/MassCommandBuffer.h b/Engine/Source/Runtime/MassEntity/Public/MassCommandBuffer.h
--- a/Engine/Source/Runtime/MassEntity/Public/MassCommandBuffer.h
+++ b/Engine/Source/Runtime/MassEntity/Public/MassCommandBuffer.h
@@ -20,11 +20,11 @@ struct FMassEntitymanager;
 checkf(IsFlushing() == false, TEXT("Trying to push commands is not supported while the given buffer is being flushed")); \
 checkf(OwnerThreadId == FPlatformTLS::GetCurrentThreadId(), TEXT("Commands can be pushed only in the same thread where the command buffer was created."))
 
-struct MASSENTITY_API FMassCommandBuffer
+struct FMassCommandBuffer
 {
 public:
-	FMassCommandBuffer();
-	~FMassCommandBuffer();
+	MASSENTITY_API FMassCommandBuffer();
+	MASSENTITY_API ~FMassCommandBuffer();
 
 	/** Adds a new entry to a given TCommand batch command instance */
 	template< template<typename... TArgs> typename TCommand, typename... TArgs >
@@ -155,14 +155,14 @@ public:
 		PushCommand<FMassCommandDestroyEntities>(InEntitiesToDestroy);
 	}
 
-	SIZE_T GetAllocatedSize() const;
+	MASSENTITY_API SIZE_T GetAllocatedSize() const;
 
 	/** 
 	 * Appends the commands from the passed buffer into this one
 	 * @param InOutOther the source buffer to copy the commands from. Note that after the call the InOutOther will be 
 	 *	emptied due to the function using Move semantics
 	 */
-	void MoveAppend(FMassCommandBuffer& InOutOther);
+	MASSENTITY_API void MoveAppend(FMassCommandBuffer& InOutOther);
 
 	bool HasPendingCommands() const 
 	{
@@ -210,7 +210,7 @@ private:
 	 * @return whether any commands have actually been executed
 	 */
 	bool Flush(FMassEntityManager& EntityManager);
-	void CleanUp();
+	MASSENTITY_API void CleanUp();
 
 	FCriticalSection AppendingCommandsCS;
 
diff --git a/Engine/Source/Runtime/MassEntity/Public/MassCommands.h b/Engine/Source/Runtime/MassEntity/Public/MassCommands.h
--- a/Engine/Source/Runtime/MassEntity/Public/MassCommands.h
+++ b/Engine/Source/Runtime/MassEntity/Public/MassCommands.h
@@ -75,7 +75,7 @@ namespace UE::Mass::Utils
 	}
 } // namespace UE::Mass::Utils
 
-struct MASSENTITY_API FMassBatchedCommand
+struct FMassBatchedCommand
 {
 	FMassBatchedCommand() = default;
 	explicit FMassBatchedCommand(EMassCommandOperationType OperationType)
@@ -124,7 +124,7 @@ protected:
 #endif // CSV_PROFILER_STATS || WITH_MASSENTITY_DEBUG
 
 private:
-	static std::atomic<uint32> CommandsCounter;
+	static MASSENTITY_API std::atomic<uint32> CommandsCounter;
 };
 
 struct FMassBatchedEntityCommand : public FMassBatchedCommand
diff --git a/Engine/Source/Runtime/MassEntity/Public/MassEntityBuilder.h b/Engine/Source/Runtime/MassEntity/Public/MassEntityBuilder.h
--- a/Engine/Source/Runtime/MassEntity/Public/MassEntityBuilder.h
+++ b/Engine/Source/Runtime/MassEntity/Public/MassEntityBuilder.h
@@ -7,6 +7,7 @@
 #include "MassEntityHandle.h"
 #include "MassEntityTypes.h"
 
+#define UE_API MASSENTITY_API
 
 struct FMassEntityManager;
 
@@ -53,35 +54,35 @@ namespace UE::Mass
  * - Committing entities while Mass's processing is in progress is not yet supported; this functionality will be implemented in the near future.
  * - no support for entity grouping
  */
-struct MASSENTITY_API FEntityBuilder
+struct FEntityBuilder
 {
 	/** Constructs a FEntityBuilder using a reference to a FMassEntityManager. */
-	explicit FEntityBuilder(FMassEntityManager& InEntityManager);
+	UE_API explicit FEntityBuilder(FMassEntityManager& InEntityManager);
 
 	/** Constructs a FEntityBuilder using a shared reference to a FMassEntityManager. */
-	explicit FEntityBuilder(const TSharedRef<FMassEntityManager>& InEntityManager);
+	UE_API explicit FEntityBuilder(const TSharedRef<FMassEntityManager>& InEntityManager);
 
 	/** Copy constructor - copies-create a new instance that represents a new entity and does not carry over reserved handle. */
-	FEntityBuilder(const FEntityBuilder& Other);
+	UE_API FEntityBuilder(const FEntityBuilder& Other);
 
 	/** Assignment operator - copies represent new entities, with no carryover of reserved handle from the original. */
-	FEntityBuilder& operator=(const FEntityBuilder& Other);
+	UE_API FEntityBuilder& operator=(const FEntityBuilder& Other);
 
 	/** Move assignment operator - moves over all the data from Other, including the internal state (like whether the entity handle has already been reserved) */
-	FEntityBuilder& operator=(FEntityBuilder&& Other);
+	UE_API FEntityBuilder& operator=(FEntityBuilder&& Other);
 
 	/** Destructor - automatically commits entity creation if not explicitly aborted or committed beforehand. */
-	~FEntityBuilder();
+	UE_API ~FEntityBuilder();
 
 	/** Creates an instance of FEntityBuilder and populates it with provided data */
-	static FEntityBuilder Make(const TSharedRef<FMassEntityManager>& InEntityManager
+	static UE_API FEntityBuilder Make(const TSharedRef<FMassEntityManager>& InEntityManager
 		, const FMassArchetypeCompositionDescriptor& Composition
 		, TConstArrayView<FInstancedStruct> InitialFragmentValues = {}
 		, TConstArrayView<FConstSharedStruct> ConstSharedFragments = {}
 		, TConstArrayView<FSharedStruct> SharedFragments = {});
 
 	/** Creates an instance of FEntityBuilder and populates it with provided data, using move-semantics on said data */
-	static FEntityBuilder Make(const TSharedRef<FMassEntityManager>& InEntityManager
+	static UE_API FEntityBuilder Make(const TSharedRef<FMassEntityManager>& InEntityManager
 		, const FMassArchetypeCompositionDescriptor& Composition
 		, TArray<FInstancedStruct>&& InitialFragmentValues
 		, TArray<FConstSharedStruct>&& ConstSharedFragments
@@ -91,23 +92,23 @@ struct MASSENTITY_API FEntityBuilder
 	 * Finalizes the creation of the entity with the specified fragments and configurations.
 	 * Note that this function needs to be called manually, no automated entity creation will take place upon builder's destruction. 
 	 */
-	FMassEntityHandle Commit();
+	UE_API FMassEntityHandle Commit();
 
 	/**
 	 * A wrapper for "Commit" call that, once that's done, prepares the builder for another commit, forgetting the
 	 * handle for the entity just created, and reverting the state back to "ReadyToCommit"
 	 * @see Commit
 	 */
-	FMassEntityHandle CommitAndReprepare();
+	UE_API FMassEntityHandle CommitAndReprepare();
 
 	/** if the builder is in "Committed" state it will roll back to ReadyToSubmit and reset the stored entity handle */
-	void Reprepare();
+	UE_API void Reprepare();
 
 	/**
 	 * Resets the builder to its initial state, discarding all previous entity configurations.
 	 * @param bReleaseEntityHandleIfReserved configures what to do with the reserved entity handle, if it's valid.
 	 */
-	void Reset(const bool bReleaseEntityHandleIfReserved = true);
+	UE_API void Reset(const bool bReleaseEntityHandleIfReserved = true);
 
 	/**
 	 * Stores ReservedEntityHandle as the cached EntityHandle. The ReservedEntityHandle is expected to be valid
@@ -115,21 +116,21 @@ struct MASSENTITY_API FEntityBuilder
 	 * If the existing EntityHandle also represents a valid, reserved entity, that handle will be released.
 	 * @return whether the ReservedEntityHandle has been stored.
 	 */
-	bool SetReservedEntityHandle(const FMassEntityHandle ReservedEntityHandle);
+	UE_API bool SetReservedEntityHandle(const FMassEntityHandle ReservedEntityHandle);
 
 	/**
 	 * Appends all element types and values stored by the entity indicated by SourceEntityHandle.
 	 * @param SourceEntityHandle valid handle for a fully constructed, built entity.
 	 * @return whether the operation was successful
 	 */
-	bool AppendDataFromEntity(const FMassEntityHandle SourceEntityHandle);
+	UE_API bool AppendDataFromEntity(const FMassEntityHandle SourceEntityHandle);
 
 	/**
 	 * Copies all element types and values stored by the entity indicated by SourceEntityHandle. Any existing builder data will be overridden
 	 * @param SourceEntityHandle valid handle for a fully constructed, built entity.
 	 * @return whether the operation was successful
 	 */
-	bool CopyDataFromEntity(const FMassEntityHandle SourceEntityHandle);
+	UE_API bool CopyDataFromEntity(const FMassEntityHandle SourceEntityHandle);
 
 	/**
 	 * Adds a fragment of type T to the entity and returns a reference to it, constructing it with the provided arguments.
@@ -173,8 +174,8 @@ struct MASSENTITY_API FEntityBuilder
 	/**
 	 * Adds a fragment instance to the Entity Builder, treating the contents according to its type
 	 */
-	FEntityBuilder& Add(const FInstancedStruct& ElementInstance);
-	FEntityBuilder& Add(FInstancedStruct&& ElementInstance);
+	UE_API FEntityBuilder& Add(const FInstancedStruct& ElementInstance);
+	UE_API FEntityBuilder& Add(FInstancedStruct&& ElementInstance);
 
 	/**
 	 * Finds and retrieves a pointer to a fragment of type T if it exists.
@@ -192,9 +193,9 @@ struct MASSENTITY_API FEntityBuilder
 	void ConfigureArchetypeCreation(const FMassArchetypeCreationParams& InCreationParams);
 
 	/** Converts the builder to a FMassEntityHandle, reserving the entity handle if not already committed. */
-	[[nodiscard]] FMassEntityHandle GetEntityHandle() const;
+	[[nodiscard]] UE_API FMassEntityHandle GetEntityHandle() const;
 
-	[[nodiscard]] FMassArchetypeHandle GetArchetypeHandle();
+	[[nodiscard]] UE_API FMassArchetypeHandle GetArchetypeHandle();
 
 	/** Checks whether the builder is in a valid, expected state */
 	bool IsValid() const;
@@ -211,9 +212,9 @@ struct MASSENTITY_API FEntityBuilder
 protected:
 	friend Private::FEntityBuilderHelper;
 
-	void CacheSharedFragmentValue();
-	void CacheArchetypeHandle();
-	void InvalidateCachedData();
+	UE_API void CacheSharedFragmentValue();
+	UE_API void CacheArchetypeHandle();
+	UE_API void InvalidateCachedData();
 
 private:
 	template<typename T>
@@ -418,3 +419,5 @@ inline void FEntityBuilder::ConfigureArchetypeCreation(const FMassArchetypeCreat
 }
 
 } // namespace UE::Mass
+
+#undef UE_API
\ No newline at end of file
diff --git a/Engine/Source/Runtime/MassEntity/Public/MassEntityCollection.h b/Engine/Source/Runtime/MassEntity/Public/MassEntityCollection.h
--- a/Engine/Source/Runtime/MassEntity/Public/MassEntityCollection.h
+++ b/Engine/Source/Runtime/MassEntity/Public/MassEntityCollection.h
@@ -24,7 +24,7 @@ namespace UE::Mass
 	 */
 	struct FEntityCollection
 	{
-		MASSENTITY_API FEntityCollection() = default;
+		FEntityCollection() = default;
 
 		/**
 		 * The following constructor are equivalent to using the default constructor and subsequently
@@ -137,7 +137,7 @@ namespace UE::Mass
 		 * Fetches up-to-date FMassArchetypeEntityCollection instances matching stored
 		 * entity handles.
 		 */
-		MASSENTITY_API TConstArrayView<FMassArchetypeEntityCollection> GetUpToDatePerArchetypeCollections(const FMassEntityManager& EntityManager) const
+		TConstArrayView<FMassArchetypeEntityCollection> GetUpToDatePerArchetypeCollections(const FMassEntityManager& EntityManager) const
 		{
 			ConditionallyUpdate(EntityManager);
 			return CachedCollections;
@@ -146,7 +146,7 @@ namespace UE::Mass
 		/**
 		 * Updates cached archetype collections and returns the container with move semantics
 		 */
-		MASSENTITY_API TArray<FMassArchetypeEntityCollection> ConsumeArchetypeCollections(const FMassEntityManager& EntityManager) &&
+		TArray<FMassArchetypeEntityCollection> ConsumeArchetypeCollections(const FMassEntityManager& EntityManager) &&
 		{
 			ConditionallyUpdate(EntityManager);
 			return MoveTemp(CachedCollections);
diff --git a/Engine/Source/Runtime/MassEntity/Public/MassEntityManager.h b/Engine/Source/Runtime/MassEntity/Public/MassEntityManager.h
--- a/Engine/Source/Runtime/MassEntity/Public/MassEntityManager.h
+++ b/Engine/Source/Runtime/MassEntity/Public/MassEntityManager.h
@@ -22,6 +22,7 @@
 #include "Misc/EnumClassFlags.h"
 #endif
 
+#define UE_API MASSENTITY_API
 
 struct FInstancedStruct;
 struct FMassEntityQuery;
@@ -80,7 +81,7 @@ namespace UE::Mass
  * (as supplied by deriving from TSharedFromThis<FMassEntityManager>).
  * IMPORTANT: if you create your own FMassEntityManager instance remember to call Initialize() before using it.
  */
-struct MASSENTITY_API FMassEntityManager : public TSharedFromThis<FMassEntityManager>, public FGCObject
+struct FMassEntityManager : public TSharedFromThis<FMassEntityManager>, public FGCObject
 {
 	friend FMassEntityQuery;
 	friend FMassDebugger;
@@ -107,32 +108,32 @@ public:
 	};
 	using FStructInitializationCallback = TFunctionRef<void(void* Fragment, const UScriptStruct& FragmentType)>;
 
-	const static FMassEntityHandle InvalidEntity;
+	const UE_API static FMassEntityHandle InvalidEntity;
 
-	explicit FMassEntityManager(UObject* InOwner = nullptr);
+	UE_API explicit FMassEntityManager(UObject* InOwner = nullptr);
 	FMassEntityManager(const FMassEntityManager& Other) = delete;
-	virtual ~FMassEntityManager();
+	UE_API virtual ~FMassEntityManager();
 
 	// FGCObject interface
-	virtual void AddReferencedObjects(FReferenceCollector& Collector) override;
+	UE_API virtual void AddReferencedObjects(FReferenceCollector& Collector) override;
 	virtual FString GetReferencerName() const override
 	{
 		return TEXT("FMassEntityManager");
 	}
 	// End of FGCObject interface
-	void GetResourceSizeEx(FResourceSizeEx& CumulativeResourceSize);
+	UE_API void GetResourceSizeEx(FResourceSizeEx& CumulativeResourceSize);
 
 	// Default to use single threaded implementation
-	void Initialize();
-	void Initialize(const FMassEntityManagerStorageInitParams& InitializationParams);
-	void PostInitialize();
-	void Deinitialize();
+	UE_API void Initialize();
+	UE_API void Initialize(const FMassEntityManagerStorageInitParams& InitializationParams);
+	UE_API void PostInitialize();
+	UE_API void Deinitialize();
 
 	/** 
 	 * A special, relaxed but slower version of CreateArchetype functions that allows FragmentAngTagsList to contain 
 	 * both fragments and tags. 
 	 */
-	FMassArchetypeHandle CreateArchetype(TConstArrayView<const UScriptStruct*> FragmentsAndTagsList, const FMassArchetypeCreationParams& CreationParams = FMassArchetypeCreationParams());
+	UE_API FMassArchetypeHandle CreateArchetype(TConstArrayView<const UScriptStruct*> FragmentsAndTagsList, const FMassArchetypeCreationParams& CreationParams = FMassArchetypeCreationParams());
 
 	/**
 	 * A special, relaxed but slower version of CreateArchetype functions that allows FragmentAngTagsList to contain
@@ -142,7 +143,7 @@ public:
 	 * @param SourceArchetype The archetype where the composition will be copied from.
 	 * @param FragmentsAndTagsList The list of fragments and tags to add to the copied composition.
 	 */
-	FMassArchetypeHandle CreateArchetype(FMassArchetypeHandle SourceArchetype, TConstArrayView<const UScriptStruct*> FragmentsAndTagsList);
+	UE_API FMassArchetypeHandle CreateArchetype(FMassArchetypeHandle SourceArchetype, TConstArrayView<const UScriptStruct*> FragmentsAndTagsList);
 	
 	/**
 	 * A special, relaxed but slower version of CreateArchetype functions that allows FragmentAngTagsList to contain
@@ -153,7 +154,7 @@ public:
 	 * @param FragmentsAndTagsList The list of fragments and tags to add to the copied composition.
 	 * @param CreationParams Additional arguments used to create the new archetype.
 	 */
-	FMassArchetypeHandle CreateArchetype(FMassArchetypeHandle SourceArchetype, TConstArrayView<const UScriptStruct*> FragmentsAndTagsList, 
+	UE_API FMassArchetypeHandle CreateArchetype(FMassArchetypeHandle SourceArchetype, TConstArrayView<const UScriptStruct*> FragmentsAndTagsList, 
 		const FMassArchetypeCreationParams& CreationParams);
 
 	/**
@@ -163,7 +164,7 @@ public:
 	 * @param CreationParams Parameters used during archetype construction
 	 * @return a handle of a new archetype 
 	 */
-	FMassArchetypeHandle CreateArchetype(const FMassArchetypeCompositionDescriptor& Composition, const FMassArchetypeCreationParams& CreationParams = FMassArchetypeCreationParams());
+	UE_API FMassArchetypeHandle CreateArchetype(const FMassArchetypeCompositionDescriptor& Composition, const FMassArchetypeCreationParams& CreationParams = FMassArchetypeCreationParams());
 
 	/**
 	 *  Creates an archetype like SourceArchetype + InFragments.
@@ -175,7 +176,7 @@ public:
 	 *   types that SourceArchetype doesn't already have. If the caller cannot guarantee it use of AddFragment functions
 	 *   family is recommended.
 	 */
-	FMassArchetypeHandle CreateArchetype(const TSharedPtr<FMassArchetypeData>& SourceArchetype, const FMassFragmentBitSet& InFragments);
+	UE_API FMassArchetypeHandle CreateArchetype(const TSharedPtr<FMassArchetypeData>& SourceArchetype, const FMassFragmentBitSet& InFragments);
 	
 	/** 
 	 *  Creates an archetype like SourceArchetype + InFragments. 
@@ -188,7 +189,7 @@ public:
 	 *   types that SourceArchetype doesn't already have. If the caller cannot guarantee it use of AddFragment functions
 	 *   family is recommended.
 	 */
-	FMassArchetypeHandle CreateArchetype(const TSharedPtr<FMassArchetypeData>& SourceArchetype, const FMassFragmentBitSet& InFragments, 
+	UE_API FMassArchetypeHandle CreateArchetype(const TSharedPtr<FMassArchetypeData>& SourceArchetype, const FMassFragmentBitSet& InFragments, 
 		const FMassArchetypeCreationParams& CreationParams);
 
 	/** 
@@ -198,48 +199,48 @@ public:
 	 * @param SharedFragmentBitSet indicates which shared fragments we want the target archetype to have. If ArchetypeHandle 
 	 *	doesn't have these a new archetype will be created.
 	 */
-	FMassArchetypeHandle GetOrCreateSuitableArchetype(const FMassArchetypeHandle& ArchetypeHandle
+	UE_API FMassArchetypeHandle GetOrCreateSuitableArchetype(const FMassArchetypeHandle& ArchetypeHandle
 		, const FMassSharedFragmentBitSet& SharedFragmentBitSet
 		, const FMassConstSharedFragmentBitSet& ConstSharedFragmentBitSet
 		, const FMassArchetypeCreationParams& CreationParams = FMassArchetypeCreationParams());
 
 	/** Fetches the archetype for a given EntityHandle. If EntityHandle is not valid it will still return a handle, just with an invalid archetype */
-	FMassArchetypeHandle GetArchetypeForEntity(FMassEntityHandle EntityHandle) const;
+	UE_API FMassArchetypeHandle GetArchetypeForEntity(FMassEntityHandle EntityHandle) const;
 	/**
 	 * Fetches the archetype for a given EntityHandle. Note that it's callers responsibility the given EntityHandle handle is valid.
 	 * If you can't ensure that call GetArchetypeForEntity.
 	 */
-	FMassArchetypeHandle GetArchetypeForEntityUnsafe(FMassEntityHandle EntityHandle) const;
+	UE_API FMassArchetypeHandle GetArchetypeForEntityUnsafe(FMassEntityHandle EntityHandle) const;
 
 	/**
 	 * Searches through all known archetypes and matches them to the provided requirements. All archetypes that pass the requirement check are returned.
 	 * @param Requirements The set of fragments and tags that need to be on available in the request form on an archetype before it's added.
 	 * @param OutValidArchetypes Archetypes that pass the requirements test are added here.
 	 */
-	void GetMatchingArchetypes(const FMassFragmentRequirements& Requirements, TArray<FMassArchetypeHandle>& OutValidArchetypes) const;
+	UE_API void GetMatchingArchetypes(const FMassFragmentRequirements& Requirements, TArray<FMassArchetypeHandle>& OutValidArchetypes) const;
 
 	/** Method to iterate on all the fragment types of an archetype */
-	static void ForEachArchetypeFragmentType(const FMassArchetypeHandle& ArchetypeHandle, TFunction< void(const UScriptStruct* /*FragmentType*/)> Function);
+	static UE_API void ForEachArchetypeFragmentType(const FMassArchetypeHandle& ArchetypeHandle, TFunction< void(const UScriptStruct* /*FragmentType*/)> Function);
 
 	/**
 	 * Go through all archetypes and compact entities
 	 * @param TimeAllowed to do entity compaction, once it reach that time it will stop and return
 	 */
-	void DoEntityCompaction(const double TimeAllowed);
+	UE_API void DoEntityCompaction(const double TimeAllowed);
 
 	/**
 	 * Creates fully built entity ready to be used by the subsystem
 	 * @param ArchetypeHandle you want this entity to be
 	 * @param SharedFragmentValues to be associated with the entity
 	 * @return FMassEntityHandle id of the newly created entity */
-	FMassEntityHandle CreateEntity(const FMassArchetypeHandle& ArchetypeHandle, const FMassArchetypeSharedFragmentValues& SharedFragmentValues = {});
+	UE_API FMassEntityHandle CreateEntity(const FMassArchetypeHandle& ArchetypeHandle, const FMassArchetypeSharedFragmentValues& SharedFragmentValues = {});
 
 	/**
 	 * Creates fully built entity ready to be used by the subsystem
 	 * @param FragmentInstanceList is the fragments to create the entity from and initialize values
 	 * @param SharedFragmentValues to be associated with the entity
 	 * @return FMassEntityHandle id of the newly created entity */
-	FMassEntityHandle CreateEntity(TConstArrayView<FInstancedStruct> FragmentInstanceList, const FMassArchetypeSharedFragmentValues& SharedFragmentValues = {}, const FMassArchetypeCreationParams& CreationParams = FMassArchetypeCreationParams());
+	UE_API FMassEntityHandle CreateEntity(TConstArrayView<FInstancedStruct> FragmentInstanceList, const FMassArchetypeSharedFragmentValues& SharedFragmentValues = {}, const FMassArchetypeCreationParams& CreationParams = FMassArchetypeCreationParams());
 
 	using FEntityCreationContext = UE::Mass::ObserverManager::FCreationContext;
 
@@ -258,7 +259,7 @@ public:
 	 * 
 	 * @return the existing (if valid) or a newly created creation context
 	 */
-	TSharedRef<FEntityCreationContext> GetOrMakeCreationContext();
+	UE_API TSharedRef<FEntityCreationContext> GetOrMakeCreationContext();
 
 	/**
 	 * Fetches the observers lock (as hosted by FMassObserverManager). If one is not currently active,
@@ -279,7 +280,7 @@ public:
 	 * @param SharedFragmentValues to be associated with the entities
 	 * @param ReservedEntities a list of reserved entities that have not yet been assigned to an archetype.
 	 * @return a creation context that will notify all the interested observers about newly created fragments once the context is released */
-	TSharedRef<FEntityCreationContext> BatchCreateReservedEntities(const FMassArchetypeHandle& ArchetypeHandle,
+	UE_API TSharedRef<FEntityCreationContext> BatchCreateReservedEntities(const FMassArchetypeHandle& ArchetypeHandle,
 		const FMassArchetypeSharedFragmentValues& SharedFragmentValues, TConstArrayView<FMassEntityHandle> ReservedEntities);
 	FORCEINLINE TSharedRef<FEntityCreationContext> BatchCreateReservedEntities(const FMassArchetypeHandle& ArchetypeHandle,
 		TConstArrayView<FMassEntityHandle> OutEntities)
@@ -293,7 +294,7 @@ public:
 	 * @param Count number of entities to create
 	 * @param InOutEntities the newly created entities are appended to given array, i.e. the pre-existing content of OutEntities won't be affected by the call
 	 * @return a creation context that will notify all the interested observers about newly created fragments once the context is released */
-	TSharedRef<FEntityCreationContext> BatchCreateEntities(const FMassArchetypeHandle& ArchetypeHandle, const FMassArchetypeSharedFragmentValues& SharedFragmentValues, const int32 Count, TArray<FMassEntityHandle>& InOutEntities);
+	UE_API TSharedRef<FEntityCreationContext> BatchCreateEntities(const FMassArchetypeHandle& ArchetypeHandle, const FMassArchetypeSharedFragmentValues& SharedFragmentValues, const int32 Count, TArray<FMassEntityHandle>& InOutEntities);
 	FORCEINLINE TSharedRef<FEntityCreationContext> BatchCreateEntities(const FMassArchetypeHandle& ArchetypeHandle, const int32 Count, TArray<FMassEntityHandle>& InOutEntities)
 	{
 		return BatchCreateEntities(ArchetypeHandle, FMassArchetypeSharedFragmentValues(), Count, InOutEntities);
@@ -302,12 +303,12 @@ public:
 	/**
 	 * Destroys a fully built entity, use ReleaseReservedEntity if entity was not yet built.
 	 * @param EntityHandle identifying the entity to destroy */
-	void DestroyEntity(FMassEntityHandle EntityHandle);
+	UE_API void DestroyEntity(FMassEntityHandle EntityHandle);
 
 	/**
 	 * Reserves an entity in the subsystem, the entity is still not ready to be used by the subsystem, need to call BuildEntity()
 	 * @return FMassEntityHandle id of the reserved entity */
-	FMassEntityHandle ReserveEntity();
+	UE_API FMassEntityHandle ReserveEntity();
 
 	/**
 	 * Builds an entity for it to be ready to be used by the subsystem
@@ -315,7 +316,7 @@ public:
 	 * @param ArchetypeHandle you want this entity to be
 	 * @param SharedFragmentValues to be associated with the entity
 	 */
-	void BuildEntity(FMassEntityHandle EntityHandle, const FMassArchetypeHandle& ArchetypeHandle, const FMassArchetypeSharedFragmentValues& SharedFragmentValues = {});
+	UE_API void BuildEntity(FMassEntityHandle EntityHandle, const FMassArchetypeHandle& ArchetypeHandle, const FMassArchetypeSharedFragmentValues& SharedFragmentValues = {});
 
 	/**
 	 * Builds an entity for it to be ready to be used by the subsystem
@@ -323,12 +324,12 @@ public:
 	 * @param FragmentInstanceList is the fragments to create the entity from and initialize values
 	 * @param SharedFragmentValues to be associated with the entity
 	 */
-	void BuildEntity(FMassEntityHandle EntityHandle, TConstArrayView<FInstancedStruct> FragmentInstanceList, const FMassArchetypeSharedFragmentValues& SharedFragmentValues = {});
+	UE_API void BuildEntity(FMassEntityHandle EntityHandle, TConstArrayView<FInstancedStruct> FragmentInstanceList, const FMassArchetypeSharedFragmentValues& SharedFragmentValues = {});
 
 	/*
 	 * Releases a previously reserved entity handle that was not yet built, otherwise call DestroyEntity
 	 * @param EntityHandle to release */
-	void ReleaseReservedEntity(FMassEntityHandle EntityHandle);
+	UE_API void ReleaseReservedEntity(FMassEntityHandle EntityHandle);
 
 	/**
 	 * Destroys all the entities in the provided array of entities. The function will also gracefully handle entities
@@ -336,54 +337,54 @@ public:
 	 * @note the function doesn't handle duplicates in InEntities.
 	 * @param InEntities to destroy
 	 */
-	void BatchDestroyEntities(TConstArrayView<FMassEntityHandle> InEntities);
+	UE_API void BatchDestroyEntities(TConstArrayView<FMassEntityHandle> InEntities);
 
 	/**
 	 * Destroys all the entities provided via the Collection. The function will also gracefully handle entities
 	 * that have been reserved but not created yet.
 	 * @param Collection to destroy
 	 */
-	void BatchDestroyEntityChunks(const FMassArchetypeEntityCollection& Collection);
-	void BatchDestroyEntityChunks(TConstArrayView<FMassArchetypeEntityCollection> Collections);
+	UE_API void BatchDestroyEntityChunks(const FMassArchetypeEntityCollection& Collection);
+	UE_API void BatchDestroyEntityChunks(TConstArrayView<FMassArchetypeEntityCollection> Collections);
 
 	/**
 	 * Assigns all entities indicated by Collections to a given archetype group.
 	 * Note that depending on their individual composition each entity can end up in a different archetype.
 	 * @paramm GroupHandle indicates the target group. Passing an invalid group handle will get logged as warning and ignored.
 	 */
-	void BatchGroupEntities(const UE::Mass::FArchetypeGroupHandle GroupHandle, TConstArrayView<FMassArchetypeEntityCollection> Collections);
+	UE_API void BatchGroupEntities(const UE::Mass::FArchetypeGroupHandle GroupHandle, TConstArrayView<FMassArchetypeEntityCollection> Collections);
 
 	/**
 	 * Assigns all entities indicated by InEntities to a given archetype group.
 	 * Note that depending on their individual composition each entity can end up in a different archetype.
 	 * @paramm GroupHandle indicates the target group. Passing an invalid group handle will get logged as warning and ignored.
 	 */
-	void BatchGroupEntities(const UE::Mass::FArchetypeGroupHandle GroupHandle, TConstArrayView<FMassEntityHandle> InEntities);
+	UE_API void BatchGroupEntities(const UE::Mass::FArchetypeGroupHandle GroupHandle, TConstArrayView<FMassEntityHandle> InEntities);
 
 	/**
 	 * Fetches FArchetypeGroupType instance (copy) associated with the given GroupName. A new group type
 	 * is created if GroupName has not been used in the past.
 	 */
-	UE::Mass::FArchetypeGroupType FinOrAddArchetypeGroupType(const FName GroupName);
+	UE_API UE::Mass::FArchetypeGroupType FinOrAddArchetypeGroupType(const FName GroupName);
 
-	const UE::Mass::FArchetypeGroups& GetGroupsForArchetype(const FMassArchetypeHandle& ArchetypeHandle) const;
+	UE_API const UE::Mass::FArchetypeGroups& GetGroupsForArchetype(const FMassArchetypeHandle& ArchetypeHandle) const;
 
-	void AddFragmentToEntity(FMassEntityHandle EntityHandle, const UScriptStruct* FragmentType);
-	void AddFragmentToEntity(FMassEntityHandle EntityHandle, const UScriptStruct* FragmentType, const FStructInitializationCallback& Initializer);
+	UE_API void AddFragmentToEntity(FMassEntityHandle EntityHandle, const UScriptStruct* FragmentType);
+	UE_API void AddFragmentToEntity(FMassEntityHandle EntityHandle, const UScriptStruct* FragmentType, const FStructInitializationCallback& Initializer);
 
 	/** 
 	 *  Ensures that only unique fragments are added. 
 	 *  @note It's caller's responsibility to ensure EntityHandle's and FragmentList's validity. 
 	 */
-	void AddFragmentListToEntity(FMassEntityHandle EntityHandle, TConstArrayView<const UScriptStruct*> FragmentList);
+	UE_API void AddFragmentListToEntity(FMassEntityHandle EntityHandle, TConstArrayView<const UScriptStruct*> FragmentList);
 
-	void AddFragmentInstanceListToEntity(FMassEntityHandle EntityHandle, TConstArrayView<FInstancedStruct> FragmentInstanceList);
-	void RemoveFragmentFromEntity(FMassEntityHandle EntityHandle, const UScriptStruct* FragmentType);
-	void RemoveFragmentListFromEntity(FMassEntityHandle EntityHandle, TConstArrayView<const UScriptStruct*> FragmentList);
+	UE_API void AddFragmentInstanceListToEntity(FMassEntityHandle EntityHandle, TConstArrayView<FInstancedStruct> FragmentInstanceList);
+	UE_API void RemoveFragmentFromEntity(FMassEntityHandle EntityHandle, const UScriptStruct* FragmentType);
+	UE_API void RemoveFragmentListFromEntity(FMassEntityHandle EntityHandle, TConstArrayView<const UScriptStruct*> FragmentList);
 
-	void AddTagToEntity(FMassEntityHandle EntityHandle, const UScriptStruct* TagType);
-	void RemoveTagFromEntity(FMassEntityHandle EntityHandle, const UScriptStruct* TagType);
-	void SwapTagsForEntity(FMassEntityHandle EntityHandle, const UScriptStruct* FromFragmentType, const UScriptStruct* ToFragmentType);
+	UE_API void AddTagToEntity(FMassEntityHandle EntityHandle, const UScriptStruct* TagType);
+	UE_API void RemoveTagFromEntity(FMassEntityHandle EntityHandle, const UScriptStruct* TagType);
+	UE_API void SwapTagsForEntity(FMassEntityHandle EntityHandle, const UScriptStruct* FromFragmentType, const UScriptStruct* ToFragmentType);
 
 	/** 
 	 * Adds a new const shared fragment to the given entity. Note that it only works if the given entity doesn't have
@@ -393,14 +394,14 @@ public:
 	 * @return whether the entity has the Fragment value assigned to it, regardless of its original state (i.e. the function will
 	 *	return true also if the entity already had the same values associated with it)
 	 */
-	bool AddConstSharedFragmentToEntity(const FMassEntityHandle EntityHandle, const FConstSharedStruct& InConstSharedFragment);
+	UE_API bool AddConstSharedFragmentToEntity(const FMassEntityHandle EntityHandle, const FConstSharedStruct& InConstSharedFragment);
 
 	/**
 	 * Removes a const shared fragment of the given type from the entity.
 	 * Will do nothing if entity did not have the shared fragment.
 	 * @return True if fragment removed from entity, false otherwise.
 	 */
-	bool RemoveConstSharedFragmentFromEntity(const FMassEntityHandle EntityHandle, const UScriptStruct& ConstSharedFragmentType);
+	UE_API bool RemoveConstSharedFragmentFromEntity(const FMassEntityHandle EntityHandle, const UScriptStruct& ConstSharedFragmentType);
 
 	/**
 	 * Adds a new shared fragment to the given entity. Note that it only works if the given entity doesn't have
@@ -410,50 +411,50 @@ public:
 	 * @return whether the entity has the Fragment value assigned to it, regardless of its original state (i.e. the function will
 	 *	return true also if the entity already had the same values associated with it)
 	 */
-	bool AddSharedFragmentToEntity(const FMassEntityHandle EntityHandle, const FSharedStruct& InSharedFragment);
+	UE_API bool AddSharedFragmentToEntity(const FMassEntityHandle EntityHandle, const FSharedStruct& InSharedFragment);
 
 	/**
 	 * Removes a shared fragment of the given type from the entity.
 	 * Will do nothing if entity did not have the shared fragment.
 	 * @return True if fragment removed from entity, false otherwise.
 	 */
-	bool RemoveSharedFragmentFromEntity(const FMassEntityHandle EntityHandle, const UScriptStruct& SharedFragmentType);
+	UE_API bool RemoveSharedFragmentFromEntity(const FMassEntityHandle EntityHandle, const UScriptStruct& SharedFragmentType);
 
 	/**
 	 * Removes EntityHandle from any-and-all groups of given type - i.e. the entity will be moved to an archetype
 	 * not in any of the groups of the given type.
 	 */
-	void RemoveEntityFromGroupType(FMassEntityHandle EntityHandle, UE::Mass::FArchetypeGroupType GroupType);
+	UE_API void RemoveEntityFromGroupType(FMassEntityHandle EntityHandle, UE::Mass::FArchetypeGroupType GroupType);
 
 	/**
 	 * @return the group handle of the specific group of type GroupType that the entity belongs to
 	 */
-	UE::Mass::FArchetypeGroupHandle GetGroupForEntity(FMassEntityHandle EntityHandle, UE::Mass::FArchetypeGroupType GroupType) const;
+	UE_API UE::Mass::FArchetypeGroupHandle GetGroupForEntity(FMassEntityHandle EntityHandle, UE::Mass::FArchetypeGroupType GroupType) const;
 
 	/**
 	 * Reserves Count number of entities and appends them to InOutEntities
 	 * @return a view into InOutEntities containing only the freshly reserved entities
 	 */
-	TConstArrayView<FMassEntityHandle> BatchReserveEntities(const int32 Count, TArray<FMassEntityHandle>& InOutEntities);
+	UE_API TConstArrayView<FMassEntityHandle> BatchReserveEntities(const int32 Count, TArray<FMassEntityHandle>& InOutEntities);
 	
 	/**
 	 * Reserves number of entities corresponding to number of entries in the provided array view InOutEntities.
 	 * As a result InOutEntities gets filled with handles of reserved entities
 	 * @return the number of entities reserved
 	 */
-	int32 BatchReserveEntities(TArrayView<FMassEntityHandle> InOutEntities);
+	UE_API int32 BatchReserveEntities(TArrayView<FMassEntityHandle> InOutEntities);
 
-	TSharedRef<FEntityCreationContext> BatchBuildEntities(const FMassArchetypeEntityCollectionWithPayload& EncodedEntitiesWithPayload, const FMassFragmentBitSet& FragmentsAffected
+	UE_API TSharedRef<FEntityCreationContext> BatchBuildEntities(const FMassArchetypeEntityCollectionWithPayload& EncodedEntitiesWithPayload, const FMassFragmentBitSet& FragmentsAffected
 		, const FMassArchetypeSharedFragmentValues& SharedFragmentValues = {}, const FMassArchetypeCreationParams& CreationParams = FMassArchetypeCreationParams());
-	TSharedRef<FEntityCreationContext> BatchBuildEntities(const FMassArchetypeEntityCollectionWithPayload& EncodedEntitiesWithPayload, const FMassArchetypeCompositionDescriptor& Composition
+	UE_API TSharedRef<FEntityCreationContext> BatchBuildEntities(const FMassArchetypeEntityCollectionWithPayload& EncodedEntitiesWithPayload, const FMassArchetypeCompositionDescriptor& Composition
 		, const FMassArchetypeSharedFragmentValues& SharedFragmentValues = {}, const FMassArchetypeCreationParams& CreationParams = FMassArchetypeCreationParams());
-	void BatchChangeTagsForEntities(TConstArrayView<FMassArchetypeEntityCollection> EntityCollections, const FMassTagBitSet& TagsToAdd, const FMassTagBitSet& TagsToRemove);
-	void BatchChangeFragmentCompositionForEntities(TConstArrayView<FMassArchetypeEntityCollection> EntityCollections, const FMassFragmentBitSet& FragmentsToAdd, const FMassFragmentBitSet& FragmentsToRemove);
-	void BatchAddFragmentInstancesForEntities(TConstArrayView<FMassArchetypeEntityCollectionWithPayload> EntityCollections, const FMassFragmentBitSet& FragmentsAffected);
+	UE_API void BatchChangeTagsForEntities(TConstArrayView<FMassArchetypeEntityCollection> EntityCollections, const FMassTagBitSet& TagsToAdd, const FMassTagBitSet& TagsToRemove);
+	UE_API void BatchChangeFragmentCompositionForEntities(TConstArrayView<FMassArchetypeEntityCollection> EntityCollections, const FMassFragmentBitSet& FragmentsToAdd, const FMassFragmentBitSet& FragmentsToRemove);
+	UE_API void BatchAddFragmentInstancesForEntities(TConstArrayView<FMassArchetypeEntityCollectionWithPayload> EntityCollections, const FMassFragmentBitSet& FragmentsAffected);
 	/** 
 	 * Adds a new const and non-const shared fragments to all entities provided via EntityCollections 
 	 */
-	void BatchAddSharedFragmentsForEntities(TConstArrayView<FMassArchetypeEntityCollection> EntityCollections, const FMassArchetypeSharedFragmentValues& AddedFragmentValues);
+	UE_API void BatchAddSharedFragmentsForEntities(TConstArrayView<FMassArchetypeEntityCollection> EntityCollections, const FMassArchetypeSharedFragmentValues& AddedFragmentValues);
 
 	/**
 	 * Adds elements indicated by InOutDescriptor to the entity indicated by EntityHandle. The function also figures out which elements
@@ -462,10 +463,10 @@ public:
 	 * If InOutDescriptor indicates shared fragments to be added the caller is required to provide matching values for the indicated
 	 * shared fragment types, via AddedSharedFragmentValues.
 	 */
-	void AddCompositionToEntity_GetDelta(FMassEntityHandle EntityHandle, FMassArchetypeCompositionDescriptor& InOutDescriptor, const FMassArchetypeSharedFragmentValues* AddedSharedFragmentValues = nullptr);
-	void RemoveCompositionFromEntity(FMassEntityHandle EntityHandle, const FMassArchetypeCompositionDescriptor& InDescriptor);
+	UE_API void AddCompositionToEntity_GetDelta(FMassEntityHandle EntityHandle, FMassArchetypeCompositionDescriptor& InOutDescriptor, const FMassArchetypeSharedFragmentValues* AddedSharedFragmentValues = nullptr);
+	UE_API void RemoveCompositionFromEntity(FMassEntityHandle EntityHandle, const FMassArchetypeCompositionDescriptor& InDescriptor);
 
-	const FMassArchetypeCompositionDescriptor& GetArchetypeComposition(const FMassArchetypeHandle& ArchetypeHandle) const;
+	UE_API const FMassArchetypeCompositionDescriptor& GetArchetypeComposition(const FMassArchetypeHandle& ArchetypeHandle) const;
 
 	/** 
 	 * Moves an entity over to a new archetype by copying over fragments common to both archetypes
@@ -473,21 +474,21 @@ public:
 	 * @param NewArchetypeHandle the handle to the new archetype
 	 * @param SharedFragmentValuesOverride if provided will override all given entity's shared fragment values
 	 */
-	void MoveEntityToAnotherArchetype(FMassEntityHandle EntityHandle, FMassArchetypeHandle NewArchetypeHandle, const FMassArchetypeSharedFragmentValues* SharedFragmentValuesOverride = nullptr);
+	UE_API void MoveEntityToAnotherArchetype(FMassEntityHandle EntityHandle, FMassArchetypeHandle NewArchetypeHandle, const FMassArchetypeSharedFragmentValues* SharedFragmentValuesOverride = nullptr);
 
 	/** 
 	 *  Copies values from FragmentInstanceList over to target entity's fragment. Caller is responsible for ensuring that 
 	 *  the given entity does have given fragments. Failing this assumption will cause a check-fail.
 	 *  @param EntityHandle idicates the target entity
 	 */
-	void SetEntityFragmentValues(FMassEntityHandle EntityHandle, TArrayView<const FInstancedStruct> FragmentInstanceList);
+	UE_API void SetEntityFragmentValues(FMassEntityHandle EntityHandle, TArrayView<const FInstancedStruct> FragmentInstanceList);
 
 	/** Copies values from FragmentInstanceList over to fragments of given entities collection. The caller is responsible 
 	 *  for ensuring that the given entity archetype (FMassArchetypeEntityCollection .Archetype) does have given fragments. 
 	 *  Failing this assumption will cause a check-fail. */
-	void BatchSetEntityFragmentValues(const FMassArchetypeEntityCollection& SparseEntities, TArrayView<const FInstancedStruct> FragmentInstanceList);
+	UE_API void BatchSetEntityFragmentValues(const FMassArchetypeEntityCollection& SparseEntities, TArrayView<const FInstancedStruct> FragmentInstanceList);
 
-	void BatchSetEntityFragmentValues(TConstArrayView<FMassArchetypeEntityCollection> EntityCollections, TArrayView<const FInstancedStruct> FragmentInstanceList);
+	UE_API void BatchSetEntityFragmentValues(TConstArrayView<FMassArchetypeEntityCollection> EntityCollections, TArrayView<const FInstancedStruct> FragmentInstanceList);
 
 	/**
 	 * @return whether the given handle represents a valid and built entity
@@ -502,22 +503,22 @@ public:
 	 * @return whether the given entity handle is valid, i.e. it points
 	 *	 to a valid spot in the entity storage and the handle's serial number is up to date
 	 */
-	bool IsEntityValid(FMassEntityHandle EntityHandle) const;
+	UE_API bool IsEntityValid(FMassEntityHandle EntityHandle) const;
 
 	/** whether the entity handle represents an entity that has been fully built (expecting a valid EntityHandle) */
-	bool IsEntityBuilt(FMassEntityHandle EntityHandle) const;
+	UE_API bool IsEntityBuilt(FMassEntityHandle EntityHandle) const;
 
 	/**
 	 * @return whether the given EntityHandle is valid and the entity it represents is in `Reserved` state
 	 *	 (i.e. it will also fail if the entity has already been `Created`)
 	 */
-	bool IsEntityReserved(FMassEntityHandle EntityHandle) const;
+	UE_API bool IsEntityReserved(FMassEntityHandle EntityHandle) const;
 
 	/** Asserts that IsEntityValid */
-	void CheckIfEntityIsValid(FMassEntityHandle EntityHandle) const;
+	UE_API void CheckIfEntityIsValid(FMassEntityHandle EntityHandle) const;
 
 	/** Asserts that IsEntityBuilt */
-	void CheckIfEntityIsActive(FMassEntityHandle EntityHandle) const;
+	UE_API void CheckIfEntityIsActive(FMassEntityHandle EntityHandle) const;
 
 	template<typename FragmentType>
 	FragmentType& GetFragmentDataChecked(FMassEntityHandle EntityHandle) const
@@ -632,7 +633,7 @@ public:
 	/**
 	 * Creates and initializes a FMassExecutionContext instance.
 	 */
-	FMassExecutionContext CreateExecutionContext(const float DeltaSeconds);
+	UE_API FMassExecutionContext CreateExecutionContext(const float DeltaSeconds);
 
 	FScopedProcessing NewProcessingScope() { return FScopedProcessing(ProcessingScopeCount); }
 
@@ -649,16 +650,16 @@ public:
 	 *		a command buffer being flushed (be it the main one or a previously requested one) then this command buffer 
 	 *		will be queued itself.
 	 */
-	void FlushCommands(const TSharedPtr<FMassCommandBuffer>& InCommandBuffer);
+	UE_API void FlushCommands(const TSharedPtr<FMassCommandBuffer>& InCommandBuffer);
 
-	void FlushCommands();
+	UE_API void FlushCommands();
 
 	/** 
 	 * Depending on the current state of Manager's command buffer the function will either move all the commands out of 
 	 * InOutCommandBuffer into the main command buffer or append it to the list of command buffers waiting to be flushed.
 	 * @note as a consequence of the call InOutCommandBuffer can get its contents emptied due some of the underlying code using Move semantics
 	 */
-	void AppendCommands(const TSharedPtr<FMassCommandBuffer>& InOutCommandBuffer);
+	UE_API void AppendCommands(const TSharedPtr<FMassCommandBuffer>& InOutCommandBuffer);
 
 	template<typename T>
 	UE_DEPRECATED(5.5, "This method will no longer be exposed. Use GetOrCreateConstSharedFragment instead.")
@@ -993,7 +994,7 @@ PRAGMA_ENABLE_DEPRECATION_WARNINGS
 	template<UE::Mass::CConstSharedFragment T>
 	void ForEachConstSharedFragmentConditional(TFunctionRef< bool(const T& /*ConstSharedFragment*/) > ConditionFunction, TFunctionRef< void(const T& /*ConstSharedFragment*/) > ExecuteFunction);
 
-	[[nodiscard]] UE::Mass::FEntityBuilder MakeEntityBuilder();
+	[[nodiscard]] UE_API UE::Mass::FEntityBuilder MakeEntityBuilder();
 
 	const UE::Mass::FTypeManager& GetTypeManager() const;
 	UE::Mass::FTypeManager& GetTypeManager();
@@ -1010,20 +1011,20 @@ PRAGMA_ENABLE_DEPRECATION_WARNINGS
 
 	bool IsDuringEntityCreation() const;
 
-	void SetDebugName(const FString& NewDebugGame);
+	UE_API void SetDebugName(const FString& NewDebugGame);
 #if WITH_MASSENTITY_DEBUG
-	void DebugPrintArchetypes(FOutputDevice& Ar, const bool bIncludeEmpty = true) const;
-	void DebugGetArchetypesStringDetails(FOutputDevice& Ar, const bool bIncludeEmpty = true) const;
-	void DebugGetArchetypeFragmentTypes(const FMassArchetypeHandle& Archetype, TArray<const UScriptStruct*>& InOutFragmentList) const;
-	int32 DebugGetArchetypeEntitiesCount(const FMassArchetypeHandle& Archetype) const;
-	int32 DebugGetArchetypeEntitiesCountPerChunk(const FMassArchetypeHandle& Archetype) const;
-	int32 DebugGetEntityCount() const;
-	int32 DebugGetArchetypesCount() const;
-	void DebugRemoveAllEntities();
-	void DebugForceArchetypeDataVersionBump();
-	void DebugGetArchetypeStrings(const FMassArchetypeHandle& Archetype, TArray<FName>& OutFragmentNames, TArray<FName>& OutTagNames);
-	FMassEntityHandle DebugGetEntityIndexHandle(const int32 EntityIndex) const;
-	const FString& DebugGetName() const;
+	UE_API void DebugPrintArchetypes(FOutputDevice& Ar, const bool bIncludeEmpty = true) const;
+	UE_API void DebugGetArchetypesStringDetails(FOutputDevice& Ar, const bool bIncludeEmpty = true) const;
+	UE_API void DebugGetArchetypeFragmentTypes(const FMassArchetypeHandle& Archetype, TArray<const UScriptStruct*>& InOutFragmentList) const;
+	UE_API int32 DebugGetArchetypeEntitiesCount(const FMassArchetypeHandle& Archetype) const;
+	UE_API int32 DebugGetArchetypeEntitiesCountPerChunk(const FMassArchetypeHandle& Archetype) const;
+	UE_API int32 DebugGetEntityCount() const;
+	UE_API int32 DebugGetArchetypesCount() const;
+	UE_API void DebugRemoveAllEntities();
+	UE_API void DebugForceArchetypeDataVersionBump();
+	UE_API void DebugGetArchetypeStrings(const FMassArchetypeHandle& Archetype, TArray<FName>& OutFragmentNames, TArray<FName>& OutTagNames);
+	UE_API FMassEntityHandle DebugGetEntityIndexHandle(const int32 EntityIndex) const;
+	UE_API const FString& DebugGetName() const;
 
 	enum class EDebugFeatures
 	{
@@ -1032,36 +1033,36 @@ PRAGMA_ENABLE_DEPRECATION_WARNINGS
 		All = TraceProcessors
 	};
 
-	void DebugEnableDebugFeature(EDebugFeatures Features);
-	void DebugDisableDebugFeature(EDebugFeatures Features);
-	bool DebugHasAllDebugFeatures(EDebugFeatures Features) const;
+	UE_API void DebugEnableDebugFeature(EDebugFeatures Features);
+	UE_API void DebugDisableDebugFeature(EDebugFeatures Features);
+	UE_API bool DebugHasAllDebugFeatures(EDebugFeatures Features) const;
 
-	FMassRequirementAccessDetector& GetRequirementAccessDetector();
+	UE_API FMassRequirementAccessDetector& GetRequirementAccessDetector();
 
 	// For use by the friend MassDebugger
-	UE::Mass::FStorageType& DebugGetEntityStorageInterface();
+	UE_API UE::Mass::FStorageType& DebugGetEntityStorageInterface();
 	// For use by the friend MassDebugger
-	const UE::Mass::FStorageType& DebugGetEntityStorageInterface() const;
+	UE_API const UE::Mass::FStorageType& DebugGetEntityStorageInterface() const;
 
-	bool DebugHasCommandsToFlush() const;
+	UE_API bool DebugHasCommandsToFlush() const;
 #endif // WITH_MASSENTITY_DEBUG
 
 protected:
 	/** Called on the child process upon process's forking */
-	void OnPostFork(EForkProcessRole Role);
+	UE_API void OnPostFork(EForkProcessRole Role);
 
-	void GetMatchingArchetypes(const FMassFragmentRequirements& Requirements, TArray<FMassArchetypeHandle>& OutValidArchetypes, const uint32 FromArchetypeDataVersion) const;
+	UE_API void GetMatchingArchetypes(const FMassFragmentRequirements& Requirements, TArray<FMassArchetypeHandle>& OutValidArchetypes, const uint32 FromArchetypeDataVersion) const;
 	
 	/** 
 	 * A "similar" archetype is an archetype exactly the same as SourceArchetype except for one composition aspect 
 	 * like Fragments or "Tags" 
 	 */
-	FMassArchetypeHandle InternalCreateSimilarArchetype(const TSharedPtr<FMassArchetypeData>& SourceArchetype, const FMassTagBitSet& OverrideTags);
-	FMassArchetypeHandle InternalCreateSimilarArchetype(const TSharedPtr<FMassArchetypeData>& SourceArchetype, const FMassFragmentBitSet& OverrideFragments);
-	FMassArchetypeHandle InternalCreateSimilarArchetype(const TSharedPtr<FMassArchetypeData>& SourceArchetype, const UE::Mass::FArchetypeGroups& GroupsOverride);
-	FMassArchetypeHandle InternalCreateSimilarArchetype(const FMassArchetypeData& SourceArchetypeRef, FMassArchetypeCompositionDescriptor&& NewComposition, const UE::Mass::FArchetypeGroups& GroupsOverride);
+	UE_API FMassArchetypeHandle InternalCreateSimilarArchetype(const TSharedPtr<FMassArchetypeData>& SourceArchetype, const FMassTagBitSet& OverrideTags);
+	UE_API FMassArchetypeHandle InternalCreateSimilarArchetype(const TSharedPtr<FMassArchetypeData>& SourceArchetype, const FMassFragmentBitSet& OverrideFragments);
+	UE_API FMassArchetypeHandle InternalCreateSimilarArchetype(const TSharedPtr<FMassArchetypeData>& SourceArchetype, const UE::Mass::FArchetypeGroups& GroupsOverride);
+	UE_API FMassArchetypeHandle InternalCreateSimilarArchetype(const FMassArchetypeData& SourceArchetypeRef, FMassArchetypeCompositionDescriptor&& NewComposition, const UE::Mass::FArchetypeGroups& GroupsOverride);
 
-	void InternalAppendFragmentsAndTagsToArchetypeCompositionDescriptor(FMassArchetypeCompositionDescriptor& InOutComposition,
+	UE_API void InternalAppendFragmentsAndTagsToArchetypeCompositionDescriptor(FMassArchetypeCompositionDescriptor& InOutComposition,
 		TConstArrayView<const UScriptStruct*> FragmentsAndTagsList) const;
 
 private:
@@ -1080,13 +1081,13 @@ private:
 	 */
 	void InternalAddFragmentListToEntity(FMassEntityHandle EntityHandle, const FMassFragmentBitSet& InFragments);
 	/** Note that it's the caller's responsibility to ensure `FragmentType` is a kind of FMassFragment */
-	void* InternalGetFragmentDataChecked(FMassEntityHandle EntityHandle, const UScriptStruct* FragmentType) const;
+	UE_API void* InternalGetFragmentDataChecked(FMassEntityHandle EntityHandle, const UScriptStruct* FragmentType) const;
 	/** Note that it's the caller's responsibility to ensure `FragmentType` is a kind of FMassFragment */
-	void* InternalGetFragmentDataPtr(FMassEntityHandle EntityHandle, const UScriptStruct* FragmentType) const;
+	UE_API void* InternalGetFragmentDataPtr(FMassEntityHandle EntityHandle, const UScriptStruct* FragmentType) const;
 	/** Note that it's the caller's responsibility to ensure `ConstSharedFragmentType` is a kind of FMassSharedFragment */
-	const FConstSharedStruct* InternalGetConstSharedFragmentPtr(FMassEntityHandle EntityHandle, const UScriptStruct* ConstSharedFragmentType) const;
+	UE_API const FConstSharedStruct* InternalGetConstSharedFragmentPtr(FMassEntityHandle EntityHandle, const UScriptStruct* ConstSharedFragmentType) const;
 	/** Note that it's the caller's responsibility to ensure `SharedFragmentType` is a kind of FMassSharedFragment */
-	const FSharedStruct* InternalGetSharedFragmentPtr(FMassEntityHandle EntityHandle, const UScriptStruct* SharedFragmentType) const;
+	UE_API const FSharedStruct* InternalGetSharedFragmentPtr(FMassEntityHandle EntityHandle, const UScriptStruct* SharedFragmentType) const;
 
 	TSharedRef<FEntityCreationContext> InternalBatchCreateReservedEntities(const FMassArchetypeHandle& ArchetypeHandle,
 		const FMassArchetypeSharedFragmentValues& SharedFragmentValues, TConstArrayView<FMassEntityHandle> ReservedEntities);
@@ -1260,16 +1261,16 @@ private:
 	//-----------------------------------------------------------------------------
 protected:
 	UE_DEPRECATED(5.6, "This flavor of InternalCreateSimilarArchetype is deprecated due to the introduction of archetype grouping. Use InternalCreateSimilarArchetype with a FArchetypeGroups parameter instead")
-	FMassArchetypeHandle InternalCreateSimilarArchetype(const FMassArchetypeData& SourceArchetypeRef, FMassArchetypeCompositionDescriptor&& NewComposition);
+	UE_API FMassArchetypeHandle InternalCreateSimilarArchetype(const FMassArchetypeData& SourceArchetypeRef, FMassArchetypeCompositionDescriptor&& NewComposition);
 public:
 	UE_DEPRECATED(5.6, "SetEntityFragmentsValues is deprecated. Use SetEntityFragmentValues instead (note the slight change in name).")
-	void SetEntityFragmentsValues(FMassEntityHandle EntityHandle, TArrayView<const FInstancedStruct> FragmentInstanceList);
+	UE_API void SetEntityFragmentsValues(FMassEntityHandle EntityHandle, TArrayView<const FInstancedStruct> FragmentInstanceList);
 
 	UE_DEPRECATED(5.6, "Static BatchSetEntityFragmentsValues is deprecated. Use EntityManager's member function BatchSetEntityFragmentValues (note the slight change in name).")
-	static void BatchSetEntityFragmentsValues(const FMassArchetypeEntityCollection& SparseEntities, TArrayView<const FInstancedStruct> FragmentInstanceList);
+	static UE_API void BatchSetEntityFragmentsValues(const FMassArchetypeEntityCollection& SparseEntities, TArrayView<const FInstancedStruct> FragmentInstanceList);
 
 	UE_DEPRECATED(5.6, "Static BatchSetEntityFragmentsValues is deprecated. Use EntityManager's member function BatchSetEntityFragmentValues (note the slight change in name).")
-	static void BatchSetEntityFragmentsValues(TConstArrayView<FMassArchetypeEntityCollection> EntityCollections, TArrayView<const FInstancedStruct> FragmentInstanceList);
+	static UE_API void BatchSetEntityFragmentsValues(TConstArrayView<FMassArchetypeEntityCollection> EntityCollections, TArrayView<const FInstancedStruct> FragmentInstanceList);
 
 	template<UE::Mass::CConstSharedFragment T>
 	UE_DEPRECATED(5.6, "Using ForEachSharedFragment for Const Shared Fragments has been deprecated. Use ForEachConstSharedFragment instead.")
@@ -1367,3 +1368,5 @@ inline TSharedRef<FMassObserverManager::FObserverLock> FMassEntityManager::GetOr
 {
 	return ObserverManager.GetOrMakeObserverLock();
 }
+
+#undef UE_API 
\ No newline at end of file
diff --git a/Engine/Source/Runtime/MassEntity/Public/MassEntityManagerStorage.h b/Engine/Source/Runtime/MassEntity/Public/MassEntityManagerStorage.h
--- a/Engine/Source/Runtime/MassEntity/Public/MassEntityManagerStorage.h
+++ b/Engine/Source/Runtime/MassEntity/Public/MassEntityManagerStorage.h
@@ -14,7 +14,7 @@ struct FMassEntityHandle;
  * Initialization parameters to configure MassEntityManager to reserve entities only single threaded
  * Supported in all build configurations
  */
-struct MASSENTITY_API FMassEntityManager_InitParams_SingleThreaded {};
+struct FMassEntityManager_InitParams_SingleThreaded {};
 
 /**
  * Initialization parameters to configure MassEntityManager to concurrently reserve entities
@@ -26,7 +26,7 @@ struct MASSENTITY_API FMassEntityManager_InitParams_SingleThreaded {};
  *
  * For default values, expectation is 128kB
  */
-struct MASSENTITY_API FMassEntityManager_InitParams_Concurrent
+struct FMassEntityManager_InitParams_Concurrent
 {
 	/** 
 	 * Maximum supported entities by the MassEntityManager
diff --git a/Engine/Source/Runtime/MassEntity/Public/MassEntityQuery.h b/Engine/Source/Runtime/MassEntity/Public/MassEntityQuery.h
--- a/Engine/Source/Runtime/MassEntity/Public/MassEntityQuery.h
+++ b/Engine/Source/Runtime/MassEntity/Public/MassEntityQuery.h
@@ -16,6 +16,8 @@
 #include "MassRequirements.h"
 #include "MassEntityQuery.generated.h"
 
+#define UE_API MASSENTITY_API
+
 #ifndef WITH_ARCHETYPE_MATCH_OVERRIDE
 #define WITH_ARCHETYPE_MATCH_OVERRIDE (WITH_EDITOR)
 #endif //  WITH_ARCHETYPE_MATCH_OVERRIDE
@@ -47,7 +49,7 @@ struct FMassEntityManager;
  *  EMassFragmentPresence::Optional fragment requirement.
  */
 USTRUCT()
-struct MASSENTITY_API FMassEntityQuery : public FMassFragmentRequirements, public FMassSubsystemRequirements
+struct FMassEntityQuery : public FMassFragmentRequirements, public FMassSubsystemRequirements
 {
 	GENERATED_BODY()
 	
@@ -80,37 +82,37 @@ public:
 	};
 
 	FMassEntityQuery() = default;
-	FMassEntityQuery(UMassProcessor& Owner);
-	FMassEntityQuery(const TSharedPtr<FMassEntityManager>& EntityManager);
-	FMassEntityQuery(const TSharedRef<FMassEntityManager>& EntityManager, std::initializer_list<UScriptStruct*> InitList);
-	FMassEntityQuery(const TSharedRef<FMassEntityManager>& EntityManager, TConstArrayView<const UScriptStruct*> InitList);
+	UE_API FMassEntityQuery(UMassProcessor& Owner);
+	UE_API FMassEntityQuery(const TSharedPtr<FMassEntityManager>& EntityManager);
+	UE_API FMassEntityQuery(const TSharedRef<FMassEntityManager>& EntityManager, std::initializer_list<UScriptStruct*> InitList);
+	UE_API FMassEntityQuery(const TSharedRef<FMassEntityManager>& EntityManager, TConstArrayView<const UScriptStruct*> InitList);
 
-	void RegisterWithProcessor(UMassProcessor& Owner);
+	UE_API void RegisterWithProcessor(UMassProcessor& Owner);
 
 	/** Runs ExecuteFunction on all entities matching Requirements */
-	void ForEachEntityChunk(FMassExecutionContext& ExecutionContext, const FMassExecuteFunction& ExecuteFunction);
+	UE_API void ForEachEntityChunk(FMassExecutionContext& ExecutionContext, const FMassExecuteFunction& ExecuteFunction);
 	
 	/** Will first verify that the archetype given with Collection matches the query's requirements, and if so will run the other, more generic ForEachEntityChunk implementation */
-	void ForEachEntityChunk(const FMassArchetypeEntityCollection& EntityCollection
+	UE_API void ForEachEntityChunk(const FMassArchetypeEntityCollection& EntityCollection
 		, FMassExecutionContext& ExecutionContext, const FMassExecuteFunction& ExecuteFunction);
 
 	/**
 	 * Attempts to process every chunk of every affected archetype in parallel.
 	 */
-	void ParallelForEachEntityChunk(FMassExecutionContext& ExecutionContext
+	UE_API void ParallelForEachEntityChunk(FMassExecutionContext& ExecutionContext
 		, const FMassExecuteFunction& ExecuteFunction, const EParallelExecutionFlags Flags = EParallelExecutionFlags::Default);
 
-	void ForEachEntityChunkInCollections(TConstArrayView<FMassArchetypeEntityCollection> EntityCollections
+	UE_API void ForEachEntityChunkInCollections(TConstArrayView<FMassArchetypeEntityCollection> EntityCollections
 		, FMassExecutionContext& ExecutionContext, const FMassExecuteFunction& ExecuteFunction);
 
-	void ParallelForEachEntityChunkInCollection(TConstArrayView<FMassArchetypeEntityCollection> EntityCollections
+	UE_API void ParallelForEachEntityChunkInCollection(TConstArrayView<FMassArchetypeEntityCollection> EntityCollections
 		, FMassExecutionContext& ExecutionContext, const FMassExecuteFunction& ExecuteFunction
 		, const EParallelExecutionFlags Flags = EParallelExecutionFlags::Default);
 
 	/** Will gather all archetypes from InEntityManager matching this->Requirements.
 	 *  Note that no work will be done if the cached data is up to date (as tracked by EntitySubsystemHash and 
 	 *	ArchetypeDataVersion properties). */
-	void CacheArchetypes();
+	UE_API void CacheArchetypes();
 
 	void Clear()
 	{
@@ -155,32 +157,32 @@ public:
 	 * Note that the function is not const because calling it can result in re-caching of ValidArchetypes 
 	 * @return the number of entities this given query would process if called "now"
 	 */
-	int32 GetNumMatchingEntities();
+	UE_API int32 GetNumMatchingEntities();
 
 	/** 
 	 * Sums the entity range lengths for each collection in EntityCollections, where the collection's 
 	 * archetype matches the query's requirements.
 	 * @return the number of entities this given query would process if called "now" for EntityCollections
 	 */
-	int32 GetNumMatchingEntities(TConstArrayView<FMassArchetypeEntityCollection> EntityCollections);
+	UE_API int32 GetNumMatchingEntities(TConstArrayView<FMassArchetypeEntityCollection> EntityCollections);
 
 	/**
 	 * Checks if any of ValidArchetypes has any entities.
 	 * Note that the function is not const because calling it can result in re-caching of ValidArchetypes
 	 * @return "true" if any of the ValidArchetypes has any entities, "false" otherwise
 	 */
-	bool HasMatchingEntities();
+	UE_API bool HasMatchingEntities();
 
 	/**
 	 * Creates an array of FMassArchetypeEntityCollection instances that identify all the entities
 	 * currently matching this query.
 	 */
-	TArray<FMassArchetypeEntityCollection> CreateMatchingEntitiesCollection();
+	UE_API TArray<FMassArchetypeEntityCollection> CreateMatchingEntitiesCollection();
 
 	/**
 	 * Fetches entity handles of all the entities currently matching this query.
 	 */
-	TArray<FMassEntityHandle> GetMatchingEntityHandles();
+	UE_API TArray<FMassEntityHandle> GetMatchingEntityHandles();
 
 	/** 
 	 * Sets a chunk filter condition that will be applied to each chunk of all valid archetypes. Note 
@@ -193,15 +195,15 @@ public:
 
 	bool HasChunkFilter() const { return bool(ChunkCondition); }
 
-	void GroupBy(UE::Mass::FArchetypeGroupType GroupType);
-	void GroupBy(UE::Mass::FArchetypeGroupType GroupType, const TFunction<bool(const UE::Mass::FArchetypeGroupID, const UE::Mass::FArchetypeGroupID)>& Predicate);
-	void ResetGrouping();
+	UE_API void GroupBy(UE::Mass::FArchetypeGroupType GroupType);
+	UE_API void GroupBy(UE::Mass::FArchetypeGroupType GroupType, const TFunction<bool(const UE::Mass::FArchetypeGroupID, const UE::Mass::FArchetypeGroupID)>& Predicate);
+	UE_API void ResetGrouping();
 
 	/** 
 	 * @return whether the query is configured to use archetype group information to group and sort
 	 *		archetypes to be processed.
 	 */
-	 bool IsGrouping() const;
+	bool IsGrouping() const;
 
 #if WITH_ARCHETYPE_MATCH_OVERRIDE
 	template<typename T> requires TArchetypeMatchOverrideConcept<T>
@@ -214,9 +216,9 @@ public:
 	 * If ArchetypeHandle is among ValidArchetypes then the function retrieves requirements mapping cached for it,
 	 * otherwise an empty mapping will be returned (and the requirements binding will be done the slow way).
 	 */
-	const FMassQueryRequirementIndicesMapping& GetRequirementsMappingForArchetype(const FMassArchetypeHandle ArchetypeHandle) const;
+	UE_API const FMassQueryRequirementIndicesMapping& GetRequirementsMappingForArchetype(const FMassArchetypeHandle ArchetypeHandle) const;
 
-	void ExportRequirements(FMassExecutionRequirements& OutRequirements) const;
+	UE_API void ExportRequirements(FMassExecutionRequirements& OutRequirements) const;
 
 	/** 
 	 * Controls whether ParallelForEachEntityChunk creates separate command buffers for each job.
@@ -228,7 +230,7 @@ public:
 	 * Configures the query to support per-entity logging based on their individual UObject "owners",
 	 * as declared via debug fragments.
 	 */
-	void DebugEnableEntityOwnerLogging();
+	UE_API void DebugEnableEntityOwnerLogging();
 
 private:
 	/**
@@ -298,10 +300,10 @@ public:
 	// DEPRECATED
 	//-----------------------------------------------------------------------------
 	UE_DEPRECATED(5.6, "This type of FMassEntityQuery is no longer supported. Use one of the other constructors instead.")
-	FMassEntityQuery(std::initializer_list<UScriptStruct*> InitList);
+	UE_API FMassEntityQuery(std::initializer_list<UScriptStruct*> InitList);
 
 	UE_DEPRECATED(5.6, "This type of FMassEntityQuery is no longer supported. Use one of the other constructors instead.")
-	FMassEntityQuery(TConstArrayView<const UScriptStruct*> InitList);
+	UE_API FMassEntityQuery(TConstArrayView<const UScriptStruct*> InitList);
 
 	enum EParallelForMode
 	{
@@ -310,33 +312,33 @@ public:
 	};
 
 	UE_DEPRECATED(5.6, "ForEachEntityChunk is deprecated. New version doesn't require the FMassEntityManager parameter")
-	void ForEachEntityChunk(FMassEntityManager& EntityManager, FMassExecutionContext& ExecutionContext, const FMassExecuteFunction& ExecuteFunction);
+	UE_API void ForEachEntityChunk(FMassEntityManager& EntityManager, FMassExecutionContext& ExecutionContext, const FMassExecuteFunction& ExecuteFunction);
 	
 	UE_DEPRECATED(5.6, "ForEachEntityChunk is deprecated. New version doesn't require the FMassEntityManager parameter")
-	void ForEachEntityChunk(const FMassArchetypeEntityCollection& EntityCollection, FMassEntityManager& EntityManager
+	UE_API void ForEachEntityChunk(const FMassArchetypeEntityCollection& EntityCollection, FMassEntityManager& EntityManager
 		, FMassExecutionContext& ExecutionContext, const FMassExecuteFunction& ExecuteFunction);
 
 	UE_DEPRECATED(5.6, "ParallelForEachEntityChunk is deprecated. New version doesn't require the FMassEntityManager parameter. Also ParallelMode parameter changed type, usee EParallelExecutionFlags instead.")
-	void ParallelForEachEntityChunk(FMassEntityManager& EntityManager, FMassExecutionContext& ExecutionContext
+	UE_API void ParallelForEachEntityChunk(FMassEntityManager& EntityManager, FMassExecutionContext& ExecutionContext
 		, const FMassExecuteFunction& ExecuteFunction, const EParallelForMode ParallelMode = Default);
 
 	UE_DEPRECATED(5.6, "ForEachEntityChunkInCollections is deprecated. New version doesn't require the FMassEntityManager parameter")
-	void ForEachEntityChunkInCollections(TConstArrayView<FMassArchetypeEntityCollection> EntityCollections, FMassEntityManager& EntityManager
+	UE_API void ForEachEntityChunkInCollections(TConstArrayView<FMassArchetypeEntityCollection> EntityCollections, FMassEntityManager& EntityManager
 		, FMassExecutionContext& ExecutionContext, const FMassExecuteFunction& ExecuteFunction);
 
 	UE_DEPRECATED(5.6, "ParallelForEachEntityChunkInCollection is deprecated. New version doesn't require the FMassEntityManager parameter. Also ParallelMode parameter changed type, usee EParallelExecutionFlags instead.")
-	void ParallelForEachEntityChunkInCollection(TConstArrayView<FMassArchetypeEntityCollection> EntityCollections
+	UE_API void ParallelForEachEntityChunkInCollection(TConstArrayView<FMassArchetypeEntityCollection> EntityCollections
 		, FMassEntityManager& EntityManager, FMassExecutionContext& ExecutionContext, const FMassExecuteFunction& ExecuteFunction
 		, const EParallelForMode ParallelMode);
 
 	UE_DEPRECATED(5.6, "This flavor of CacheArchetypes is deprecated. EntityQueries are not tied to a specific EntityManager instance and there's no need to pass it in as a parameter")
-	void CacheArchetypes(const FMassEntityManager& InEntityManager);
+	UE_API void CacheArchetypes(const FMassEntityManager& InEntityManager);
 
 	UE_DEPRECATED(5.6, "This flavor of GetNumMatchingEntities is deprecated. EntityQueries are not tied to a specific EntityManager instance and there's no need to pass it in as a parameter")
-	int32 GetNumMatchingEntities(FMassEntityManager& InEntityManager);
+	UE_API int32 GetNumMatchingEntities(FMassEntityManager& InEntityManager);
 
 	UE_DEPRECATED(5.6, "This flavor of HasMatchingEntities is deprecated. EntityQueries are not tied to a specific EntityManager instance and there's no need to pass it in as a parameter")
-	bool HasMatchingEntities(FMassEntityManager& InEntityManager);
+	UE_API bool HasMatchingEntities(FMassEntityManager& InEntityManager);
 };
 
 ENUM_CLASS_FLAGS(FMassEntityQuery::EParallelExecutionFlags);
@@ -382,4 +384,6 @@ inline const TSharedPtr<FMassEntityManager>& FMassEntityQuery::GetEntityManager(
 }
 
 #undef MASS_ACHETYPE_OVERRIDE_MAX_SIZE
-#undef MASS_ACHETYPE_OVERRIDE_MAX_ALIGNMENT
\ No newline at end of file
+#undef MASS_ACHETYPE_OVERRIDE_MAX_ALIGNMENT
+
+#undef UE_API 
\ No newline at end of file
diff --git a/Engine/Source/Runtime/MassEntity/Public/MassEntityTypes.h b/Engine/Source/Runtime/MassEntity/Public/MassEntityTypes.h
--- a/Engine/Source/Runtime/MassEntity/Public/MassEntityTypes.h
+++ b/Engine/Source/Runtime/MassEntity/Public/MassEntityTypes.h
@@ -235,7 +235,7 @@ struct FMassArchetypeCompositionDescriptor
  * Note that having multiple instanced of a given struct type is not supported and Add* functions will fetch the previously 
  * added fragment instead of adding a new one.
  */
-struct MASSENTITY_API FMassArchetypeSharedFragmentValues
+struct FMassArchetypeSharedFragmentValues
 {
 	FMassArchetypeSharedFragmentValues() = default;
 	FMassArchetypeSharedFragmentValues(const FMassArchetypeSharedFragmentValues& OtherFragmentValues) = default;
@@ -289,7 +289,7 @@ struct MASSENTITY_API FMassArchetypeSharedFragmentValues
 	 * does care however about matching "mode", meaning ConstSharedFragments and SharedFragments arrays are compared
 	 * independently.
 	 */
-	bool HasSameValues(const FMassArchetypeSharedFragmentValues& Other) const;
+	MASSENTITY_API bool HasSameValues(const FMassArchetypeSharedFragmentValues& Other) const;
 
 	FORCEINLINE bool ContainsType(const UScriptStruct* FragmentType) const
 	{
@@ -338,7 +338,7 @@ struct MASSENTITY_API FMassArchetypeSharedFragmentValues
 	 * In that case the method will return the previously added instance if the given type has been added
 	 * as a CONST shared fragment and if not it will return an empty FConstSharedStruct.
 	 */
-	FConstSharedStruct Add_GetRef(const FConstSharedStruct& Fragment);
+	MASSENTITY_API FConstSharedStruct Add_GetRef(const FConstSharedStruct& Fragment);
 
 	UE_DEPRECATED(5.6, "Use Add or Add_GetRef instead depending on whether you need the return value.")
 		FConstSharedStruct AddConstSharedFragment(const FConstSharedStruct& Fragment)
@@ -361,7 +361,7 @@ struct MASSENTITY_API FMassArchetypeSharedFragmentValues
 	 * In that case the method will return the previously added instance if the given type has been added
 	 * as a NON-CONST shared fragment and if not it will return an empty FSharedStruct.
 	 */
-	FSharedStruct Add_GetRef(const FSharedStruct& Fragment);
+	MASSENTITY_API FSharedStruct Add_GetRef(const FSharedStruct& Fragment);
 
 	UE_DEPRECATED(5.6, "Use Add or Add_GetRef instead depending on whether you need the return value.")
 	FSharedStruct AddSharedFragment(const FSharedStruct& Fragment)
@@ -375,7 +375,7 @@ struct MASSENTITY_API FMassArchetypeSharedFragmentValues
 	 * Note that it's callers responsibility to ensure every fragment type in Fragments already has an instance in
 	 * this FMassArchetypeSharedFragmentValues instance. Failing that assumption will result in ensure failure. 
 	 */
-	void ReplaceSharedFragments(TConstArrayView<FSharedStruct> Fragments);
+	MASSENTITY_API void ReplaceSharedFragments(TConstArrayView<FSharedStruct> Fragments);
 
 	/** 
 	 * Appends contents of Other to `this` instance. All common fragments will get overridden with values in Other.
@@ -383,19 +383,19 @@ struct MASSENTITY_API FMassArchetypeSharedFragmentValues
 	 * ensure when that is attempted.
 	 * @return number of fragments added or changed
 	 */
-	int32 Append(const FMassArchetypeSharedFragmentValues& Other);
+	MASSENTITY_API int32 Append(const FMassArchetypeSharedFragmentValues& Other);
 
 	/** 
 	 * Note that the function removes the shared fragments by type
 	 * @return number of fragments types removed
 	 */
-	int32 Remove(const FMassSharedFragmentBitSet& SharedFragmentToRemoveBitSet);
+	MASSENTITY_API int32 Remove(const FMassSharedFragmentBitSet& SharedFragmentToRemoveBitSet);
 
 	/** 
 	 * Note that the function removes the const shared fragments by type
 	 * @return number of fragments types removed
 	 */
-	int32 Remove(const FMassConstSharedFragmentBitSet& ConstSharedFragmentToRemoveBitSet);
+	MASSENTITY_API int32 Remove(const FMassConstSharedFragmentBitSet& ConstSharedFragmentToRemoveBitSet);
 
 	/**
 	 * Remove all the shared and const shared fragments indicated by InDescriptor
@@ -470,7 +470,7 @@ struct MASSENTITY_API FMassArchetypeSharedFragmentValues
 		return SharedFragmentValues.HashCache;
 	}
 
-	uint32 CalculateHash() const;
+	MASSENTITY_API uint32 CalculateHash() const;
 	SIZE_T GetAllocatedSize() const;
 
 	void Sort()
diff --git a/Engine/Source/Runtime/MassEntity/Public/MassEntityView.h b/Engine/Source/Runtime/MassEntity/Public/MassEntityView.h
--- a/Engine/Source/Runtime/MassEntity/Public/MassEntityView.h
+++ b/Engine/Source/Runtime/MassEntity/Public/MassEntityView.h
@@ -9,6 +9,7 @@
 #endif // UE_ENABLE_INCLUDE_ORDER_DEPRECATED_IN_5_6
 #include "MassEntityView.generated.h"
 
+#define UE_API MASSENTITY_API
 
 struct FMassEntityManager;
 struct FMassArchetypeData;
@@ -19,7 +20,7 @@ struct FMassArchetypeHandle;
  * validity only within the scope it has been created in. Don't store it. 
  */
 USTRUCT()
-struct MASSENTITY_API FMassEntityView
+struct FMassEntityView
 {
 	GENERATED_BODY()
 
@@ -31,19 +32,19 @@ struct MASSENTITY_API FMassEntityView
 	 *  FMassEntityManager-flavored constructor is recommended since it will first find the appropriate archetype for
 	 *  Entity. 
 	 */
-	FMassEntityView(const FMassArchetypeHandle& ArchetypeHandle, FMassEntityHandle Entity);
+	UE_API FMassEntityView(const FMassArchetypeHandle& ArchetypeHandle, FMassEntityHandle Entity);
 
 	/** 
 	 *  Finds the archetype Entity belongs to and then resolves against it. The caller is responsible for ensuring
 	 *  that the given Entity is in fact a valid ID tied to any of the archetypes 
 	 */
-	FMassEntityView(const FMassEntityManager& EntityManager, FMassEntityHandle Entity);
+	UE_API FMassEntityView(const FMassEntityManager& EntityManager, FMassEntityHandle Entity);
 
 	/** 
 	 * If the given handle represents a valid entity the function will create a FMassEntityView just like a constructor 
 	 * would. If the entity is not valid the produced view will be "unset".
 	 */
-	static FMassEntityView TryMakeView(const FMassEntityManager& EntityManager, FMassEntityHandle Entity);
+	static UE_API FMassEntityView TryMakeView(const FMassEntityManager& EntityManager, FMassEntityHandle Entity);
 
 	FMassEntityHandle GetEntity() const	{ return Entity; }
 
@@ -144,19 +145,19 @@ struct MASSENTITY_API FMassEntityView
 		return HasTag(*T::StaticStruct());
 	}
 
-	bool HasTag(const UScriptStruct& TagType) const;
+	UE_API bool HasTag(const UScriptStruct& TagType) const;
 
 	bool IsSet() const;
 	bool IsValid() const;
 	bool operator==(const FMassEntityView& Other) const;
 
 protected:
-	void* GetFragmentPtr(const UScriptStruct& FragmentType) const;
-	void* GetFragmentPtrChecked(const UScriptStruct& FragmentType) const;
-	const void* GetConstSharedFragmentPtr(const UScriptStruct& FragmentType) const;
-	const void* GetConstSharedFragmentPtrChecked(const UScriptStruct& FragmentType) const;
-	void* GetSharedFragmentPtr(const UScriptStruct& FragmentType) const;
-	void* GetSharedFragmentPtrChecked(const UScriptStruct& FragmentType) const;
+	UE_API void* GetFragmentPtr(const UScriptStruct& FragmentType) const;
+	UE_API void* GetFragmentPtrChecked(const UScriptStruct& FragmentType) const;
+	UE_API const void* GetConstSharedFragmentPtr(const UScriptStruct& FragmentType) const;
+	UE_API const void* GetConstSharedFragmentPtrChecked(const UScriptStruct& FragmentType) const;
+	UE_API void* GetSharedFragmentPtr(const UScriptStruct& FragmentType) const;
+	UE_API void* GetSharedFragmentPtrChecked(const UScriptStruct& FragmentType) const;
 
 private:
 	FMassEntityHandle Entity;
@@ -181,3 +182,5 @@ inline bool FMassEntityView::operator==(const FMassEntityView& Other) const
 {
 	return Archetype == Other.Archetype && EntityDataHandle == Other.EntityDataHandle;
 }
+
+#undef UE_API
diff --git a/Engine/Source/Runtime/MassEntity/Public/MassProcessingContext.h b/Engine/Source/Runtime/MassEntity/Public/MassProcessingContext.h
--- a/Engine/Source/Runtime/MassEntity/Public/MassProcessingContext.h
+++ b/Engine/Source/Runtime/MassEntity/Public/MassProcessingContext.h
@@ -10,7 +10,7 @@
 
 namespace UE::Mass
 {
-	struct MASSENTITY_API FProcessingContext
+	struct FProcessingContext
 	{
 		UE_DEPRECATED_FORGAME(5.6, "This constructor is deprecated. Use one of the other ones.")
 		FProcessingContext() = default;
@@ -23,7 +23,7 @@ namespace UE::Mass
 		explicit FProcessingContext(const TSharedRef<FMassEntityManager>& InEntityManager, const float InDeltaSeconds = 0.f, const bool bInFlushCommandBuffer = true);
 		explicit FProcessingContext(TSharedRef<FMassEntityManager>&& InEntityManager, const float InDeltaSeconds = 0.f, const bool bInFlushCommandBuffer = true);
 		explicit FProcessingContext(const TSharedPtr<FMassEntityManager>& InEntityManager, const float InDeltaSeconds = 0.f, const bool bInFlushCommandBuffer = true);
-		~FProcessingContext();
+		MASSENTITY_API ~FProcessingContext();
 		
 		[[nodiscard]] FMassExecutionContext& GetExecutionContext() &;
 		[[nodiscard]] FMassExecutionContext&& GetExecutionContext() &&;
diff --git a/Engine/Source/Runtime/MassEntity/Public/MassProcessingTypes.h b/Engine/Source/Runtime/MassEntity/Public/MassProcessingTypes.h
--- a/Engine/Source/Runtime/MassEntity/Public/MassProcessingTypes.h
+++ b/Engine/Source/Runtime/MassEntity/Public/MassProcessingTypes.h
@@ -7,6 +7,8 @@
 #include "Templates/SubclassOf.h"
 #include "MassProcessingTypes.generated.h"
 
+#define UE_API MASSENTITY_API
+
 #ifndef MASS_DO_PARALLEL
 #define MASS_DO_PARALLEL !UE_SERVER
 #endif // MASS_DO_PARALLEL
@@ -42,7 +44,7 @@ struct FProcessorAuxDataBase
 };
 
 USTRUCT(BlueprintType, meta = (Deprecated = "5.6"))
-struct MASSENTITY_API FMassProcessingContext_DEPRECATED
+struct FMassProcessingContext_DEPRECATED
 {
 	GENERATED_BODY()
 
@@ -66,7 +68,7 @@ using FMassProcessingContext = UE::Mass::FProcessingContext;
  *  Runtime-usable array of MassProcessor copies
  */
 USTRUCT()
-struct MASSENTITY_API FMassRuntimePipeline
+struct FMassRuntimePipeline
 {
 	GENERATED_BODY()
 
@@ -78,31 +80,31 @@ private:
 
 public:
 	explicit FMassRuntimePipeline(EProcessorExecutionFlags WorldExecutionFlags = EProcessorExecutionFlags::None);
-	FMassRuntimePipeline(TConstArrayView<TObjectPtr<UMassProcessor>> SeedProcessors, EProcessorExecutionFlags WorldExecutionFlags);
-	FMassRuntimePipeline(TConstArrayView<UMassProcessor*> SeedProcessors, EProcessorExecutionFlags WorldExecutionFlags);
+	UE_API FMassRuntimePipeline(TConstArrayView<TObjectPtr<UMassProcessor>> SeedProcessors, EProcessorExecutionFlags WorldExecutionFlags);
+	UE_API FMassRuntimePipeline(TConstArrayView<UMassProcessor*> SeedProcessors, EProcessorExecutionFlags WorldExecutionFlags);
 
-	void Reset();
-	void Initialize(UObject& Owner, const TSharedRef<FMassEntityManager>& EntityManager);
+	UE_API void Reset();
+	UE_API void Initialize(UObject& Owner, const TSharedRef<FMassEntityManager>& EntityManager);
 	
 	/** Copies the given array over to this FMassRuntimePipeline instance. */
-	void SetProcessors(TArrayView<UMassProcessor*> InProcessors);
+	UE_API void SetProcessors(TArrayView<UMassProcessor*> InProcessors);
 
 	/** Directly moves the contents of given array to the Processors member array. */
-	void SetProcessors(TArray<TObjectPtr<UMassProcessor>>&& InProcessors);
+	UE_API void SetProcessors(TArray<TObjectPtr<UMassProcessor>>&& InProcessors);
 
 	/** Creates runtime copies of UMassProcessors given in InProcessors input parameter, using InOwner as new UMassProcessors' outer. */
-	void CreateFromArray(TConstArrayView<const UMassProcessor*> InProcessors, UObject& InOwner);
+	UE_API void CreateFromArray(TConstArrayView<const UMassProcessor*> InProcessors, UObject& InOwner);
 
 	/** Calls CreateFromArray and calls Initialize on all processors afterwards. */
-	void InitializeFromArray(TConstArrayView<const UMassProcessor*> InProcessors, UObject& InOwner, const TSharedRef<FMassEntityManager>& EntityManager);
+	UE_API void InitializeFromArray(TConstArrayView<const UMassProcessor*> InProcessors, UObject& InOwner, const TSharedRef<FMassEntityManager>& EntityManager);
 	
 	/** Creates runtime instances of UMassProcessors for each processor class given via InProcessorClasses. 
 	 *  The instances will be created with InOwner as outer. */
-	void InitializeFromClassArray(TConstArrayView<TSubclassOf<UMassProcessor>> InProcessorClasses, UObject& InOwner, const TSharedRef<FMassEntityManager>& EntityManager);
+	UE_API void InitializeFromClassArray(TConstArrayView<TSubclassOf<UMassProcessor>> InProcessorClasses, UObject& InOwner, const TSharedRef<FMassEntityManager>& EntityManager);
 
 	/** Creates a runtime instance of every processors in the given InProcessors array. If a processor of that class
 	 *  already exists in Processors array it gets overridden. Otherwise it gets added to the end of the collection.*/
-	void AppendOrOverrideRuntimeProcessorCopies(TConstArrayView<const UMassProcessor*> InProcessors, UObject& InOwner);
+	UE_API void AppendOrOverrideRuntimeProcessorCopies(TConstArrayView<const UMassProcessor*> InProcessors, UObject& InOwner);
 
 	/** Creates a runtime instance of every processors in the given array if there's no processor of that class in Processors already.
 	 *  Call this function when adding processors to an already configured FMassRuntimePipeline instance. If you're creating 
@@ -110,25 +112,25 @@ public:
 	 *  or call AppendOrOverrideRuntimeProcessorCopies.
 	 *	NOTE: there's a change in functionality since 5.6 - the function will no longer create duplicates for processors returning true from ShouldAllowMultipleInstances
 	 */
-	void AppendUniqueRuntimeProcessorCopies(TConstArrayView<const UMassProcessor*> InProcessors, UObject& InOwner, const TSharedRef<FMassEntityManager>& EntityManager);
+	UE_API void AppendUniqueRuntimeProcessorCopies(TConstArrayView<const UMassProcessor*> InProcessors, UObject& InOwner, const TSharedRef<FMassEntityManager>& EntityManager);
 
 	/** Adds InProcessor(s) to Processors without any additional checks */
-	void AppendProcessor(UMassProcessor& InProcessor);
-	void AppendProcessors(TArrayView<TObjectPtr<UMassProcessor>> InProcessors);
+	UE_API void AppendProcessor(UMassProcessor& InProcessor);
+	UE_API void AppendProcessors(TArrayView<TObjectPtr<UMassProcessor>> InProcessors);
 
 	/** @return true if the given processor has indeed been added, i.e. will return false if Processor was already part of the pipeline. */
-	bool AppendUniqueProcessor(UMassProcessor& Processor);
+	UE_API bool AppendUniqueProcessor(UMassProcessor& Processor);
 
 	/** Creates an instance of ProcessorClass and adds it to Processors without any additional checks */
-	void AppendProcessor(TSubclassOf<UMassProcessor> ProcessorClass, UObject& InOwner);
+	UE_API void AppendProcessor(TSubclassOf<UMassProcessor> ProcessorClass, UObject& InOwner);
 
 	/** @return whether the given processor has been removed from hosted processors collection */
-	bool RemoveProcessor(const UMassProcessor& InProcessor);
+	UE_API bool RemoveProcessor(const UMassProcessor& InProcessor);
 
 	/** goes through Processor looking for a UMassCompositeProcessor instance which GroupName matches the one given as the parameter */
-	UMassCompositeProcessor* FindTopLevelGroupByName(const FName GroupName);
+	UE_API UMassCompositeProcessor* FindTopLevelGroupByName(const FName GroupName);
 
-	bool HasProcessorOfExactClass(TSubclassOf<UMassProcessor> InClass) const;
+	UE_API bool HasProcessorOfExactClass(TSubclassOf<UMassProcessor> InClass) const;
 	bool IsEmpty() const;
 
 	int32 Num() const;
@@ -139,25 +141,25 @@ public:
 	/** Returns Processors array using move semantics. This operation clears out this FMassRuntimePipeline instance. */
 	TArray<TObjectPtr<UMassProcessor>>&& MoveProcessorsArray();
 
-	MASSENTITY_API friend uint32 GetTypeHash(const FMassRuntimePipeline& Instance);
+	UE_API friend uint32 GetTypeHash(const FMassRuntimePipeline& Instance);
 
 	//-----------------------------------------------------------------------------
 	// DEPRECATED
 	//-----------------------------------------------------------------------------
 	UE_DEPRECATED(5.6, "This flavor of Initialize is deprecated. Please use the one requiring a FMassEntityManager parameter")
-	void Initialize(UObject& Owner);
+	UE_API void Initialize(UObject& Owner);
 
 	UE_DEPRECATED(5.6, "This flavor of InitializeFromArray is deprecated. Please use the one requiring a FMassEntityManager parameter")
-	void InitializeFromArray(TConstArrayView<const UMassProcessor*> InProcessors, UObject& InOwner);
+	UE_API void InitializeFromArray(TConstArrayView<const UMassProcessor*> InProcessors, UObject& InOwner);
 	
 	UE_DEPRECATED(5.6, "This flavor of InitializeFromClassArray is deprecated. Please use the one requiring a FMassEntityManager parameter")
-	void InitializeFromClassArray(TConstArrayView<TSubclassOf<UMassProcessor>> InProcessorClasses, UObject& InOwner);
+	UE_API void InitializeFromClassArray(TConstArrayView<TSubclassOf<UMassProcessor>> InProcessorClasses, UObject& InOwner);
 
 	UE_DEPRECATED(5.6, "This flavor of AppendUniqueRuntimeProcessorCopies is deprecated. Please use the one requiring a FMassEntityManager parameter")
-	void AppendUniqueRuntimeProcessorCopies(TConstArrayView<const UMassProcessor*> InProcessors, UObject& InOwner);
+	UE_API void AppendUniqueRuntimeProcessorCopies(TConstArrayView<const UMassProcessor*> InProcessors, UObject& InOwner);
 
 	UE_DEPRECATED(5.6, "This flavor of SetProcessors is deprecated. Please use one of the others.")
-	void SetProcessors(TArray<UMassProcessor*>&& InProcessors);
+	UE_API void SetProcessors(TArray<UMassProcessor*>&& InProcessors);
 };
 
 UENUM()
@@ -222,3 +224,5 @@ inline TArray<TObjectPtr<UMassProcessor>>&& FMassRuntimePipeline::MoveProcessors
 {
 	return MoveTemp(Processors);
 }
+
+#undef UE_API
diff --git a/Engine/Source/Runtime/MassEntity/Public/MassProcessor.h b/Engine/Source/Runtime/MassEntity/Public/MassProcessor.h
--- a/Engine/Source/Runtime/MassEntity/Public/MassProcessor.h
+++ b/Engine/Source/Runtime/MassEntity/Public/MassProcessor.h
@@ -11,6 +11,7 @@
 #include "MassCommandBuffer.h"
 #include "MassProcessor.generated.h"
 
+#define UE_API MASSENTITY_API
 
 struct FMassProcessingPhaseConfig;
 class UMassCompositeProcessor;
@@ -63,28 +64,28 @@ enum class EMassQueryBasedPruning : uint8
 	Default = Prune
 };
 
-UCLASS(abstract, EditInlineNew, CollapseCategories, config = Mass, defaultconfig, ConfigDoNotCheckDefaults)
-class MASSENTITY_API UMassProcessor : public UObject
+UCLASS(abstract, EditInlineNew, CollapseCategories, config = Mass, defaultconfig, ConfigDoNotCheckDefaults, MinimalAPI)
+class UMassProcessor : public UObject
 {
 	GENERATED_BODY()
 
 public:
 
-	UMassProcessor();
-	explicit UMassProcessor(const FObjectInitializer& ObjectInitializer);
+	UE_API UMassProcessor();
+	UE_API explicit UMassProcessor(const FObjectInitializer& ObjectInitializer);
 
 	bool IsInitialized() const;
 
 	/** Calls InitializeInternal and handles initialization bookkeeping. */
-	void CallInitialize(const TNotNull<UObject*> Owner, const TSharedRef<FMassEntityManager>& EntityManager);
+	UE_API void CallInitialize(const TNotNull<UObject*> Owner, const TSharedRef<FMassEntityManager>& EntityManager);
 	
-	virtual FGraphEventRef DispatchProcessorTasks(const TSharedPtr<FMassEntityManager>& EntityManager, FMassExecutionContext& ExecutionContext, const FGraphEventArray& Prerequisites = FGraphEventArray());
+	UE_API virtual FGraphEventRef DispatchProcessorTasks(const TSharedPtr<FMassEntityManager>& EntityManager, FMassExecutionContext& ExecutionContext, const FGraphEventArray& Prerequisites = FGraphEventArray());
 
 	EProcessorExecutionFlags GetExecutionFlags() const;
 
 	/** Whether this processor should execute according the CurrentExecutionFlags parameters */
 	bool ShouldExecute(const EProcessorExecutionFlags CurrentExecutionFlags) const;
-	void CallExecute(FMassEntityManager& EntityManager, FMassExecutionContext& Context);
+	UE_API void CallExecute(FMassEntityManager& EntityManager, FMassExecutionContext& Context);
 
 	/** 
 	 * Controls whether there can be multiple instances of a given class in a single FMassRuntimePipeline and during 
@@ -92,9 +93,9 @@ public:
 	 */
 	bool ShouldAllowMultipleInstances() const;
 
-	void DebugOutputDescription(FOutputDevice& Ar) const;
-	virtual void DebugOutputDescription(FOutputDevice& Ar, int32 Indent) const;
-	virtual FString GetProcessorName() const;
+	UE_API void DebugOutputDescription(FOutputDevice& Ar) const;
+	UE_API virtual void DebugOutputDescription(FOutputDevice& Ar, int32 Indent) const;
+	UE_API virtual FString GetProcessorName() const;
 	
 	//----------------------------------------------------------------------//
 	// Ordering functions 
@@ -107,17 +108,17 @@ public:
 	 * By default, the processor will be the subject of pruning when bRuntimeMode == true.
 	 * @param bRuntimeMode indicates whether the pruning is being done for game runtime (true) or editor-time presentation (false)
 	 */
-	virtual bool ShouldAllowQueryBasedPruning(const bool bRuntimeMode = true) const;
+	UE_API virtual bool ShouldAllowQueryBasedPruning(const bool bRuntimeMode = true) const;
 
-	virtual EMassProcessingPhase GetProcessingPhase() const;
-	virtual void SetProcessingPhase(EMassProcessingPhase Phase);
+	UE_API virtual EMassProcessingPhase GetProcessingPhase() const;
+	UE_API virtual void SetProcessingPhase(EMassProcessingPhase Phase);
 	bool DoesRequireGameThreadExecution() const;
 	
 	const FMassProcessorExecutionOrder& GetExecutionOrder() const;
 
 	/** By default,  fetches requirements declared entity queries registered via RegisterQuery. Processors can override 
 	 *	this function to supply additional requirements */
-	virtual void ExportRequirements(FMassExecutionRequirements& OutRequirements) const;
+	UE_API virtual void ExportRequirements(FMassExecutionRequirements& OutRequirements) const;
 
 	const FMassSubsystemRequirements& GetProcessorRequirements() const;
 
@@ -138,7 +139,7 @@ public:
 
 	/** Adds Query to RegisteredQueries list. Query is required to be a member variable of this processor. Not meeting
 	 *  this requirement will cause check failure and the query won't be registered. */
-	void RegisterQuery(FMassEntityQuery& Query);
+	UE_API void RegisterQuery(FMassEntityQuery& Query);
 
 	void MarkAsDynamic();
 	bool IsDynamic() const;
@@ -169,10 +170,10 @@ public:
 	 * instantiated to be automatically evaluated every frame. @see FMassProcessingPhaseManager
 	 * Note that calling this function is only valid on CDOs. Calling it on a regular instance will fail an ensure and 
 	 * have no other effect, i.e. CDO's value won't change */
-	void SetShouldAutoRegisterWithGlobalList(const bool bAutoRegister);
+	UE_API void SetShouldAutoRegisterWithGlobalList(const bool bAutoRegister);
 
-	void GetArchetypesMatchingOwnedQueries(const FMassEntityManager& EntityManager, TArray<FMassArchetypeHandle>& OutArchetype);
-	bool DoesAnyArchetypeMatchOwnedQueries(const FMassEntityManager& EntityManager);
+	UE_API void GetArchetypesMatchingOwnedQueries(const FMassEntityManager& EntityManager, TArray<FMassArchetypeHandle>& OutArchetype);
+	UE_API bool DoesAnyArchetypeMatchOwnedQueries(const FMassEntityManager& EntityManager);
 	int32 GetOwnedQueriesNum() const;
 	
 #if CPUPROFILERTRACE_ENABLED
@@ -181,17 +182,17 @@ public:
 	
 protected:
 	/** Called to initialize the processor's internal state. Override to perform custom steps. */
-	virtual void InitializeInternal(UObject& Owner, const TSharedRef<FMassEntityManager>& EntityManager);
+	UE_API virtual void InitializeInternal(UObject& Owner, const TSharedRef<FMassEntityManager>& EntityManager);
 
 	/**
 	 * Called internally during processor's initialization so that child classes configure their owned queries
 	 * with requirements. This function is called before the processor gets considered by Mass dependency
 	 * solver, and the requirement information stored in queries is crucial for that process. 
 	 */
-	virtual void ConfigureQueries(const TSharedRef<FMassEntityManager>& EntityManager);
+	UE_API virtual void ConfigureQueries(const TSharedRef<FMassEntityManager>& EntityManager);
 
-	virtual void PostInitProperties() override;
-	virtual void Execute(FMassEntityManager& EntityManager, FMassExecutionContext& Context) PURE_VIRTUAL(UMassProcessor::Execute);
+	UE_API virtual void PostInitProperties() override;
+	UE_API virtual void Execute(FMassEntityManager& EntityManager, FMassExecutionContext& Context) PURE_VIRTUAL(UMassProcessor::Execute);
 
 protected:
 	/** Configures when this given processor can be executed in relation to other processors and processing groups, within its processing phase. */
@@ -292,14 +293,14 @@ private:
 	//-----------------------------------------------------------------------------
 public:
 	UE_DEPRECATED(5.6, "Initialize is deprecated. Override InitializeInternal(UObject&, const TSharedRef<FMassEntityManager>&) instead. If you want to call the function, use CallInitialize.")
-	virtual void Initialize(UObject& Owner) final;
+	UE_API virtual void Initialize(UObject& Owner) final;
 	UE_DEPRECATED(5.6, "This flavor of ConfigureQueries is deprecated. Override ConfigureQueries(const TSharedRef<FMassEntityManager>&) instead.")
 	virtual void ConfigureQueries() final {};
 };
 
 
-UCLASS()
-class MASSENTITY_API UMassCompositeProcessor : public UMassProcessor
+UCLASS(MinimalAPI)
+class UMassCompositeProcessor : public UMassProcessor
 {
 	GENERATED_BODY()
 
@@ -316,24 +317,24 @@ public:
 	};
 
 public:
-	UMassCompositeProcessor();
+	UE_API UMassCompositeProcessor();
 
-	void SetChildProcessors(TArrayView<UMassProcessor*> InProcessors);
-	void SetChildProcessors(TArray<TObjectPtr<UMassProcessor>>&& InProcessors);
+	UE_API void SetChildProcessors(TArrayView<UMassProcessor*> InProcessors);
+	UE_API void SetChildProcessors(TArray<TObjectPtr<UMassProcessor>>&& InProcessors);
 
-	virtual void InitializeInternal(UObject& Owner, const TSharedRef<FMassEntityManager>& InEntityManager) override;
-	virtual void DebugOutputDescription(FOutputDevice& Ar, int32 Indent = 0) const override;
-	virtual void SetProcessingPhase(EMassProcessingPhase Phase) override;
+	UE_API virtual void InitializeInternal(UObject& Owner, const TSharedRef<FMassEntityManager>& InEntityManager) override;
+	UE_API virtual void DebugOutputDescription(FOutputDevice& Ar, int32 Indent = 0) const override;
+	UE_API virtual void SetProcessingPhase(EMassProcessingPhase Phase) override;
 
-	void SetGroupName(FName NewName);
-	FName GetGroupName() const;
+	UE_API void SetGroupName(FName NewName);
+	UE_API FName GetGroupName() const;
 
-	virtual void SetProcessors(TArrayView<UMassProcessor*> InProcessorInstances, const TSharedPtr<FMassEntityManager>& EntityManager = nullptr);
+	UE_API virtual void SetProcessors(TArrayView<UMassProcessor*> InProcessorInstances, const TSharedPtr<FMassEntityManager>& EntityManager = nullptr);
 
 	/** 
 	 * Builds flat processing graph that's being used for multithreading execution of hosted processors.
 	 */
-	virtual void BuildFlatProcessingGraph(TConstArrayView<FMassProcessorOrderInfo> SortedProcessors);
+	UE_API virtual void BuildFlatProcessingGraph(TConstArrayView<FMassProcessorOrderInfo> SortedProcessors);
 
 	/**
 	 * Adds processors in InOutOrderedProcessors to ChildPipeline. 
@@ -342,24 +343,24 @@ public:
 	 * will be updated to reflect the reuse.
 	 * The described behavior however is available only for processors with bAllowMultipleInstances == false.
 	 */
-	void UpdateProcessorsCollection(TArrayView<FMassProcessorOrderInfo> InOutOrderedProcessors, EProcessorExecutionFlags InWorldExecutionFlags = EProcessorExecutionFlags::None);
+	UE_API void UpdateProcessorsCollection(TArrayView<FMassProcessorOrderInfo> InOutOrderedProcessors, EProcessorExecutionFlags InWorldExecutionFlags = EProcessorExecutionFlags::None);
 
 	/** adds SubProcessor to an appropriately named group. If RequestedGroupName == None then SubProcessor
 	 *  will be added directly to ChildPipeline. If not then the indicated group will be searched for in ChildPipeline 
 	 *  and if it's missing it will be created and AddGroupedProcessor will be called recursively */
-	void AddGroupedProcessor(FName RequestedGroupName, UMassProcessor& SubProcessor);
+	UE_API void AddGroupedProcessor(FName RequestedGroupName, UMassProcessor& SubProcessor);
 
-	virtual FGraphEventRef DispatchProcessorTasks(const TSharedPtr<FMassEntityManager>& EntityManager, FMassExecutionContext& ExecutionContext, const FGraphEventArray& Prerequisites = FGraphEventArray()) override;
+	UE_API virtual FGraphEventRef DispatchProcessorTasks(const TSharedPtr<FMassEntityManager>& EntityManager, FMassExecutionContext& ExecutionContext, const FGraphEventArray& Prerequisites = FGraphEventArray()) override;
 
-	bool IsEmpty() const;
+	UE_API bool IsEmpty() const;
 
-	virtual FString GetProcessorName() const override;
+	UE_API virtual FString GetProcessorName() const override;
 
-	TConstArrayView<UMassProcessor*> GetChildProcessorsView() const;
+	UE_API TConstArrayView<UMassProcessor*> GetChildProcessorsView() const;
 
 protected:
-	virtual void ConfigureQueries(const TSharedRef<FMassEntityManager>& EntityManager) override;
-	virtual void Execute(FMassEntityManager& EntityManager, FMassExecutionContext& Context) override;
+	UE_API virtual void ConfigureQueries(const TSharedRef<FMassEntityManager>& EntityManager) override;
+	UE_API virtual void Execute(FMassEntityManager& EntityManager, FMassExecutionContext& Context) override;
 
 	/** RequestedGroupName can indicate a multi-level group name, like so: A.B.C
 	 *  We need to extract the highest-level group name ('A' in the example), and see if it already exists. 
@@ -367,7 +368,7 @@ protected:
 	 *  @param RequestedGroupName name of the group for which we want to find or create the processor.
 	 *  @param OutRemainingGroupName contains the group name after cutting the high-level group. In the used example it
 	 *    will contain "B.C". This value is then used to recursively create subgroups */
-	UMassCompositeProcessor* FindOrAddGroupProcessor(FName RequestedGroupName, FString* OutRemainingGroupName = nullptr);
+	UE_API UMassCompositeProcessor* FindOrAddGroupProcessor(FName RequestedGroupName, FString* OutRemainingGroupName = nullptr);
 
 protected:
 	UPROPERTY(VisibleAnywhere, Category=Mass)
@@ -412,7 +413,7 @@ public:
 	virtual void Initialize(UObject& Owner) final;*/
 
 	UE_DEPRECATED(5.6, "This flavor of SetChildProcessors is deprecated. Please use one of the others.")
-	void SetChildProcessors(TArray<UMassProcessor*>&& InProcessors);
+	UE_API void SetChildProcessors(TArray<UMassProcessor*>&& InProcessors);
 };
 
 
@@ -534,3 +535,4 @@ inline TConstArrayView<UMassProcessor*> UMassCompositeProcessor::GetChildProcess
 	return ChildPipeline.GetProcessors();
 }
 
+#undef UE_API 
\ No newline at end of file
diff --git a/Engine/Source/Runtime/MassEntity/Public/MassSubsystemAccess.h b/Engine/Source/Runtime/MassEntity/Public/MassSubsystemAccess.h
--- a/Engine/Source/Runtime/MassEntity/Public/MassSubsystemAccess.h
+++ b/Engine/Source/Runtime/MassEntity/Public/MassSubsystemAccess.h
@@ -19,9 +19,9 @@
 #endif // WITH_EDITOR
 
 
-struct MASSENTITY_API FMassSubsystemAccess
+struct FMassSubsystemAccess
 {
-	explicit FMassSubsystemAccess(UWorld* InWorld = nullptr);
+	MASSENTITY_API explicit FMassSubsystemAccess(UWorld* InWorld = nullptr);
 
 	//-----------------------------------------------------------------------------
 	// Statically-typed subsystems
@@ -114,8 +114,8 @@ struct MASSENTITY_API FMassSubsystemAccess
 	//-----------------------------------------------------------------------------
 	// remaining API
 	//-----------------------------------------------------------------------------
-	bool CacheSubsystemRequirements(const FMassSubsystemRequirements& SubsystemRequirements);
-	void SetSubsystemRequirements(const FMassSubsystemRequirements& SubsystemRequirements);
+	MASSENTITY_API bool CacheSubsystemRequirements(const FMassSubsystemRequirements& SubsystemRequirements);
+	MASSENTITY_API void SetSubsystemRequirements(const FMassSubsystemRequirements& SubsystemRequirements);
 
 	void GetSubsystemRequirementBits(FMassExternalSubsystemBitSet& OutConstSubsystemsBitSet, FMassExternalSubsystemBitSet& OutMutableSubsystemsBitSet)
 	{
@@ -182,7 +182,7 @@ struct MASSENTITY_API FMassSubsystemAccess
 		}
 	}
 
-	static USubsystem* FetchSubsystemInstance(UWorld* World, TSubclassOf<USubsystem> SubsystemClass);
+	static MASSENTITY_API USubsystem* FetchSubsystemInstance(UWorld* World, TSubclassOf<USubsystem> SubsystemClass);
 
 protected:
 	template<typename T>
@@ -226,7 +226,7 @@ protected:
 		return Cast<T>(SystemInstance);
 	}
 
-	bool CacheSubsystem(const uint32 SystemIndex);
+	MASSENTITY_API bool CacheSubsystem(const uint32 SystemIndex);
 
 	FMassExternalSubsystemBitSet ConstSubsystemsBitSet;
 	FMassExternalSubsystemBitSet MutableSubsystemsBitSet;
diff --git a/Engine/Source/Runtime/MassEntity/Public/MassSubsystemBase.h b/Engine/Source/Runtime/MassEntity/Public/MassSubsystemBase.h
--- a/Engine/Source/Runtime/MassEntity/Public/MassSubsystemBase.h
+++ b/Engine/Source/Runtime/MassEntity/Public/MassSubsystemBase.h
@@ -8,6 +8,7 @@
 #include "MassEntityConcepts.h"
 #include "MassSubsystemBase.generated.h"
 
+#define UE_API MASSENTITY_API
 
 struct FMassEntityManager;
 
@@ -31,29 +32,29 @@ namespace UE::Mass
  * The sole responsibility of this world subsystem class is to serve functionality common to all 
  * Mass-related UWorldSubsystem-based subsystems, like whether the subsystems should get created at all. 
  */
-UCLASS(Abstract, config = Mass, defaultconfig)
-class MASSENTITY_API UMassSubsystemBase : public UWorldSubsystem
+UCLASS(Abstract, config = Mass, defaultconfig, MinimalAPI)
+class UMassSubsystemBase : public UWorldSubsystem
 {
 	GENERATED_BODY()
 
 public:
-	static bool AreRuntimeMassSubsystemsAllowed(UObject* Outer);
+	static UE_API bool AreRuntimeMassSubsystemsAllowed(UObject* Outer);
 	UE::Mass::Subsystems::FInitializationState GetInitializationState() const { return InitializationState; }
 
 protected:
 	//~USubsystem interface
-	virtual bool ShouldCreateSubsystem(UObject* Outer) const override;
-	virtual void Initialize(FSubsystemCollectionBase& Collection) override;
-	virtual void PostInitialize() override;
-	virtual void Deinitialize() override;
-	virtual void OnWorldBeginPlay(UWorld& InWorld) override;
+	UE_API virtual bool ShouldCreateSubsystem(UObject* Outer) const override;
+	UE_API virtual void Initialize(FSubsystemCollectionBase& Collection) override;
+	UE_API virtual void PostInitialize() override;
+	UE_API virtual void Deinitialize() override;
+	UE_API virtual void OnWorldBeginPlay(UWorld& InWorld) override;
 	//~End of USubsystem interface
 
 	/**
 	 * Needs to be called in Initialize for subsystems we want to behave properly when dynamically added after UWorld::BeginPlay
 	 * (for example via GameplayFeatureActions). This is required for subsystems relying on their PostInitialize and/or OnWorldBeginPlay called.
 	 */
-	void HandleLateCreation();
+	UE_API void HandleLateCreation();
 
 	/**
 	 * Registers given subsystem class as part of Mass type information. Needs to be called as part of Initialize override.
@@ -76,8 +77,8 @@ protected:
  * The sole responsibility of this tickable world subsystem class is to serve functionality common to all
  * Mass-related UTickableWorldSubsystem-based subsystems, like whether the subsystems should get created at all.
  */
-UCLASS(Abstract, config = Mass, defaultconfig)
-class MASSENTITY_API UMassTickableSubsystemBase : public UTickableWorldSubsystem
+UCLASS(Abstract, config = Mass, defaultconfig, MinimalAPI)
+class UMassTickableSubsystemBase : public UTickableWorldSubsystem
 {
 	GENERATED_BODY()
 
@@ -86,11 +87,11 @@ public:
 
 protected:
 	//~USubsystem interface
-	virtual bool ShouldCreateSubsystem(UObject* Outer) const override;
-	virtual void Initialize(FSubsystemCollectionBase& Collection) override;
-	virtual void PostInitialize() override;
-	virtual void Deinitialize() override;
-	virtual void OnWorldBeginPlay(UWorld& InWorld) override;
+	UE_API virtual bool ShouldCreateSubsystem(UObject* Outer) const override;
+	UE_API virtual void Initialize(FSubsystemCollectionBase& Collection) override;
+	UE_API virtual void PostInitialize() override;
+	UE_API virtual void Deinitialize() override;
+	UE_API virtual void OnWorldBeginPlay(UWorld& InWorld) override;
 	//~End of USubsystem interface
 
 	/**
@@ -107,7 +108,7 @@ protected:
 	 * Needs to be called in Initialize for subsystems we want to behave properly when dynamically added after UWorld::BeginPlay
 	 * (for example via GameplayFeatureActions). This is required for subsystems relying on their PostInitialize and/or OnWorldBeginPlay called.
 	 */
-	void HandleLateCreation();
+	UE_API void HandleLateCreation();
 
 private:
 	/** 
@@ -116,3 +117,5 @@ private:
 	 */
 	UE::Mass::Subsystems::FInitializationState InitializationState;
 };
+
+#undef UE_API
\ No newline at end of file
