[Mass] Added a way to fetch base-typed fragment views from Execution context while providing the actual, specific class with a UScriptStruct pointer.

diff --git a/Engine/Source/Runtime/MassEntity/Public/MassExecutionContext.h b/Engine/Source/Runtime/MassEntity/Public/MassExecutionContext.h
--- a/Engine/Source/Runtime/MassEntity/Public/MassExecutionContext.h
+++ b/Engine/Source/Runtime/MassEntity/Public/MassExecutionContext.h
@@ -461,6 +461,22 @@ public:
 		return View->FragmentView;
 	}
 
+	template<typename TFragmentBase>
+	TConstArrayView<TFragmentBase> GetFragmentView(TNotNull<const UScriptStruct*> FragmentType) const
+	{
+		check(FragmentType->IsChildOf(TFragmentBase::StaticStruct()));
+		TConstArrayView<FMassFragment> View = GetFragmentView(FragmentType);
+		return TConstArrayView<TFragmentBase>(reinterpret_cast<const TFragmentBase*>(View.GetData()), View.Num());
+	}
+
+	template<typename TFragmentBase>
+	TArrayView<TFragmentBase> GetMutableFragmentView(TNotNull<const UScriptStruct*> FragmentType)
+	{
+		check(FragmentType->IsChildOf(TFragmentBase::StaticStruct()));
+		TArrayView<FMassFragment> View = GetMutableFragmentView(FragmentType);
+		return TArrayView<TFragmentBase>(reinterpret_cast<TFragmentBase*>(View.GetData()), View.Num());;
+	}
+
 	template<typename T, typename = typename TEnableIf<TIsDerivedFrom<T, USubsystem>::IsDerived>::Type>
 	T* GetMutableSubsystem()
 	{
