[MassEntity] fixed typo in forward declaration

diff --git a/Engine/Source/Runtime/MassEntity/Public/MassCommandBuffer.h b/Engine/Source/Runtime/MassEntity/Public/MassCommandBuffer.h
--- a/Engine/Source/Runtime/MassEntity/Public/MassCommandBuffer.h
+++ b/Engine/Source/Runtime/MassEntity/Public/MassCommandBuffer.h
@@ -11,7 +11,7 @@
 #include "MassCommands.h"
 
 
-struct FMassEntitymanager;
+struct FMassEntityManager;
 
 //@TODO: Consider debug information in case there is an assert when replaying the command buffer
 // (e.g., which system added the command, or even file/line number in development builds for the specific call via a macro)
