[Mass Debugger]
Added Fragments view tab
Added new "Selected Entity" breakpoint type to break for whichever entity is selected in the gameplay debugger or debug UI

diff --git a/Engine/Source/Editor/MassEntityDebugger/Private/MassDebuggerModel.cpp b/Engine/Source/Editor/MassEntityDebugger/Private/MassDebuggerModel.cpp
--- a/Engine/Source/Editor/MassEntityDebugger/Private/MassDebuggerModel.cpp
+++ b/Engine/Source/Editor/MassEntityDebugger/Private/MassDebuggerModel.cpp
@@ -315,6 +315,22 @@ int32 FMassDebuggerArchetypeData::GetTotalBitsUsedCount() const
 	return Composition.CountStoredTypes();
 }
 
+//----------------------------------------------------------------------//
+// FMassDebuggerFragmentData
+//----------------------------------------------------------------------//
+FMassDebuggerFragmentData::FMassDebuggerFragmentData(TNotNull<const UScriptStruct*> InFragment)
+	: Fragment(InFragment)
+{
+	if (Fragment.IsValid())
+	{
+		Name = Fragment.Pin()->GetDisplayNameText();
+	}
+	else
+	{
+		Name = LOCTEXT("Invalid", "Invalid");
+	}
+}
+
 //----------------------------------------------------------------------//
 // FMassDebuggerProcessingGraphNode
 //----------------------------------------------------------------------//
@@ -336,7 +352,7 @@ FText FMassDebuggerProcessingGraphNode::GetLabel() const
 		return FText::FromString(ProcessorData->Label);
 	}
 
-	return LOCTEXT("InvalidProcessor", "Invalid");
+	return LOCTEXT("Invalid", "Invalid");
 }
 
 //----------------------------------------------------------------------//
@@ -439,6 +455,7 @@ void FMassDebuggerModel::RefreshAll()
 		TArray<TNotNull<const UMassCompositeProcessor*>> CompositeProcessors;
 		CacheProcessorsData(TransientArchetypesMap, CompositeProcessors);
 		CacheProcessingGraphs(CompositeProcessors);
+		CacheFragmentData();
 
 		ClearArchetypeSelection();
 
@@ -608,6 +625,67 @@ void FMassDebuggerModel::CacheProcessingGraphs(TConstArrayView<TNotNull<const UM
 	}
 }
 
+void FMassDebuggerModel::CacheFragmentData()
+{
+	CachedFragmentData.Reset();
+
+#if WITH_MASSENTITY_DEBUG
+	TMap<const UScriptStruct*, TSharedPtr<FMassDebuggerFragmentData>> TempFragmentDataMap;
+
+	for (const TWeakObjectPtr<const UScriptStruct>& WeakStruct : FMassFragmentBitSet::DebugGetAllStructTypes())
+	{
+		if (WeakStruct.IsValid())
+		{
+			TempFragmentDataMap.Add(WeakStruct.Get(), MakeShared<FMassDebuggerFragmentData>(WeakStruct.Get()));
+		}
+	}
+
+	if (IsCurrentEnvironmentValid())
+	{
+		const FMassEntityManager& EntityManager = *Environment->GetEntityManager();
+		EntityManager.ForEachArchetype(0, TNumericLimits<int32>::Max(), [this, &TempFragmentDataMap](const FMassEntityManager& EntityManager, const FMassArchetypeHandle& ArchetypeHandle)
+		{
+			const int32 NumEntitiesInArchetype = EntityManager.DebugGetArchetypeEntitiesCount(ArchetypeHandle);
+
+			EntityManager.ForEachArchetypeFragmentType(ArchetypeHandle, [this, &TempFragmentDataMap, &ArchetypeHandle, NumEntitiesInArchetype](const UScriptStruct* FragmentStruct)
+			{
+				if (FragmentStruct)
+				{
+					TSharedPtr<FMassDebuggerFragmentData>* FragmentEntry = TempFragmentDataMap.Find(FragmentStruct);
+
+					if (FragmentEntry && (*FragmentEntry))
+					{
+						(*FragmentEntry)->Archetypes.AddUnique(ArchetypeHandle);
+						(*FragmentEntry)->NumEntities += NumEntitiesInArchetype;
+					}
+				}
+			});
+		});
+	}
+
+	CachedFragmentData.Reserve(TempFragmentDataMap.Num());
+	for (auto const& [FragmentStruct, FragmentData] : TempFragmentDataMap)
+	{
+		CachedFragmentData.Add(FragmentData);
+	}
+
+	if (IsCurrentEnvironmentValid())
+	{
+		TWeakPtr<const FMassEntityManager> EntityManager = Environment->GetEntityManager().ToWeakPtr();
+		for (TSharedPtr<FMassDebuggerFragmentData>& FragmentData : CachedFragmentData)
+		{
+			FragmentData->EntityManager = EntityManager;
+		}
+	}
+
+	static auto SortPredicate = [](const TSharedPtr<FMassDebuggerFragmentData>& A, const TSharedPtr<FMassDebuggerFragmentData>& B)
+	{
+		return A->Name.CompareTo(B->Name) <= 0;
+	};
+	CachedFragmentData.Sort(SortPredicate);
+#endif
+}
+
 void FMassDebuggerModel::CacheArchetypesData(TMap<FMassArchetypeHandle, TSharedPtr<FMassDebuggerArchetypeData>>& OutTransientArchetypesMap)
 {
 	CachedAllArchetypes.Reset();
@@ -717,7 +795,10 @@ void FMassDebuggerModel::MarkAsStale()
 
 bool FMassDebuggerModel::IsStale() const 
 { 
-	return Environment.IsValid() == false || (Environment->NeedsValidWorld() == true && Environment->IsWorldValid() == false);
+	return Environment.IsValid() == false
+		|| Environment->EntityManager.IsValid() == false
+		|| FMassDebugger::IsEntityManagerInitialized(*Environment->EntityManager.Pin()) == false
+		|| (Environment->NeedsValidWorld() == true && Environment->IsWorldValid() == false);
 }
 
 const TSharedPtr<FMassDebuggerProcessorData>& FMassDebuggerModel::GetProcessorDataChecked(const UMassProcessor& Processor) const
diff --git a/Engine/Source/Editor/MassEntityDebugger/Private/MassDebuggerModel.h b/Engine/Source/Editor/MassEntityDebugger/Private/MassDebuggerModel.h
--- a/Engine/Source/Editor/MassEntityDebugger/Private/MassDebuggerModel.h
+++ b/Engine/Source/Editor/MassEntityDebugger/Private/MassDebuggerModel.h
@@ -124,6 +124,28 @@ public:
 #endif // WITH_MASSENTITY_DEBUG	
 };
 
+struct FMassDebuggerFragmentData
+{
+	FMassDebuggerFragmentData(TNotNull<const UScriptStruct*> InFragment);
+
+	bool operator==(const FMassDebuggerFragmentData& Other) const
+	{
+		return Fragment == Other.Fragment;
+	}
+
+	friend uint32 GetTypeHash(const FMassDebuggerFragmentData& Item)
+	{
+		return GetTypeHash(Item.Fragment);
+	}
+
+	FText Name;
+	TWeakObjectPtr<const UScriptStruct> Fragment;
+	TWeakPtr<const FMassEntityManager> EntityManager;
+	TArray<FMassArchetypeHandle> Archetypes;
+	TArray<TWeakObjectPtr<UMassProcessor>> Processors;
+	int32 NumEntities = 0;
+};
+
 struct FMassDebuggerProcessingGraphNode
 {
 	explicit FMassDebuggerProcessingGraphNode(const TSharedPtr<FMassDebuggerProcessorData>& InProcessorData, const UMassCompositeProcessor::FDependencyNode& InProcessorNode = UMassCompositeProcessor::FDependencyNode());
@@ -193,6 +215,8 @@ struct FMassDebuggerModel
 	void CacheProcessorsData(const TMap<FMassArchetypeHandle, TSharedPtr<FMassDebuggerArchetypeData>>& InTransientArchetypesMap, TArray<TNotNull<const UMassCompositeProcessor*>>& OutCompositeProcessors);
 	void CacheProcessingGraphs(TConstArrayView<TNotNull<const UMassCompositeProcessor*>> InCompositeProcessors);
 
+	void CacheFragmentData();
+
 	float MinDistanceToSelectedArchetypes(const TSharedPtr<FMassDebuggerArchetypeData>& InArchetypeData) const;
 
 	FText GetDisplayName() const;
@@ -241,6 +265,7 @@ public:
 		TArray<TSharedPtr<FMassDebuggerProcessorData>> Container;
 	};
 
+	TArray<TSharedPtr<FMassDebuggerFragmentData>> CachedFragmentData;
 	TArray<TSharedPtr<FProcessorCollection>> CachedProcessorCollections;
 	TArray<TSharedPtr<FMassDebuggerProcessorData>> SelectedProcessors;
 	TArray<TSharedPtr<FMassDebuggerArchetypeData>> CachedAllArchetypes;
diff --git a/Engine/Source/Editor/MassEntityDebugger/Private/SMassBreakpointsView.cpp b/Engine/Source/Editor/MassEntityDebugger/Private/SMassBreakpointsView.cpp
--- a/Engine/Source/Editor/MassEntityDebugger/Private/SMassBreakpointsView.cpp
+++ b/Engine/Source/Editor/MassEntityDebugger/Private/SMassBreakpointsView.cpp
@@ -69,13 +69,26 @@ void SMassBreakpointsView::RefreshBreakpoints()
 			return;
 		}
 
-		TMultiMap<const UScriptStruct*, FMassEntityHandle>::TIterator FragmentIterator = Env->FragmentWriteBreakpoints.CreateIterator();
+		TMultiMap<TObjectKey<const UScriptStruct>, FMassEntityHandle>::TIterator FragmentIterator = Env->FragmentWriteBreakpoints.CreateIterator();
 		for (; FragmentIterator; ++FragmentIterator)
 		{
 			TSharedPtr<FBreakpointDisplay> NewDisplay = MakeShared<FBreakpointDisplay>();
-			NewDisplay->WriteFragment = FragmentIterator.Key();
+			NewDisplay->WriteFragment = FragmentIterator.Key().ResolveObjectPtr();
 			NewDisplay->Entity = FragmentIterator.Value();
-			WriteBreakpoints.Add(NewDisplay);
+			if (NewDisplay->WriteFragment.IsValid())
+			{
+				WriteBreakpoints.Add(NewDisplay);
+			}
+		}
+		for (TObjectKey<const UScriptStruct>& FragmentTypeKey : Env->SelectedEntityFragmentWriteBreakpoints)
+		{
+			TSharedPtr<FBreakpointDisplay> NewDisplay = MakeShared<FBreakpointDisplay>();
+			NewDisplay->WriteFragment = FragmentTypeKey.ResolveObjectPtr();
+			NewDisplay->bBreakForSelectedEntity = true;
+			if (NewDisplay->WriteFragment.IsValid())
+			{
+				WriteBreakpoints.Add(NewDisplay);
+			}
 		}
 	}
 	if (WriteBreakpointsListView)
@@ -121,13 +134,20 @@ TSharedRef<ITableRow> SMassBreakpointsView::OnGenerateBreakpointRow(TSharedPtr<F
 						const FMassEntityManager* EntityManager = DebuggerModel->Environment->EntityManager.Pin().Get();
 						if (EntityManager)
 						{
-							if (InItem->Processor)
+							if (InItem->Processor.IsValid())
 							{
-								FMassDebugger::ClearProcessorBreakpoint(*EntityManager, InItem->Processor, InItem->Entity);
+								FMassDebugger::ClearProcessorBreakpoint(*EntityManager, InItem->Processor.Get(), InItem->Entity);
 							}
-							else if (InItem->WriteFragment)
+							else if (InItem->WriteFragment.IsValid())
 							{
-								FMassDebugger::ClearFragmentWriteBreak(*EntityManager, InItem->WriteFragment, InItem->Entity);
+								if (InItem->bBreakForSelectedEntity)
+								{
+									FMassDebugger::ToggleFragmentWriteBreakForSelectedEntity(*EntityManager, InItem->WriteFragment.Get());
+								}
+								else
+								{
+									FMassDebugger::ClearFragmentWriteBreak(*EntityManager, InItem->WriteFragment.Get(), InItem->Entity);
+								}
 							}
 						}
 					}
@@ -145,9 +165,9 @@ TSharedRef<ITableRow> SMassBreakpointsView::OnGenerateBreakpointRow(TSharedPtr<F
             [
                 // Display WriteFragment or Processor name
                 SNew(STextBlock)
-                .Text(InItem->WriteFragment 
+                .Text(InItem->WriteFragment.Get()
                     ? FText::FromString(InItem->WriteFragment->GetName()) 
-                    : (InItem->Processor 
+                    : (InItem->Processor.Get()
                         ? FText::FromString(InItem->Processor->GetName()) 
                         : LOCTEXT("UnknownBreakpoint", "Unknown")))
             ]
@@ -155,9 +175,14 @@ TSharedRef<ITableRow> SMassBreakpointsView::OnGenerateBreakpointRow(TSharedPtr<F
             .AutoWidth()
             .Padding(5.0f, 2.0f)
             [
-                // Display Entity handle
+                // Display Entity handle or "Selected Entity"
                 SNew(STextBlock)
-				.Text(FText::FromString(InItem->Entity.DebugGetDescription()))
+				.Text(InItem->bBreakForSelectedEntity
+					? LOCTEXT("SelectedEntity", "Selected Entity")
+					: FText::FromString(InItem->Entity.DebugGetDescription()))
+				.ToolTip(FSlateApplicationBase::Get().MakeToolTip(InItem->bBreakForSelectedEntity
+					? LOCTEXT("SelectedEntity", "Selected Entity (Gameplay Debugger)")
+					: FText::FromString(InItem->Entity.DebugGetDescription())))
             ]
 #else
 			SNew(STextBlock)
diff --git a/Engine/Source/Editor/MassEntityDebugger/Private/SMassBreakpointsView.h b/Engine/Source/Editor/MassEntityDebugger/Private/SMassBreakpointsView.h
--- a/Engine/Source/Editor/MassEntityDebugger/Private/SMassBreakpointsView.h
+++ b/Engine/Source/Editor/MassEntityDebugger/Private/SMassBreakpointsView.h
@@ -43,8 +43,10 @@ private:
 
 	struct FBreakpointDisplay
 	{
-		const UMassProcessor* Processor;
-		const UScriptStruct* WriteFragment;
+		TWeakObjectPtr<const UMassProcessor> Processor;
+		TWeakObjectPtr<const UScriptStruct> WriteFragment;
+		/** Specifies that this breakpoint is not for a specific entity, but will trigger for any entity selected in the mass or gameplay debugger*/
+		bool bBreakForSelectedEntity = false;
 		FMassEntityHandle Entity;
 	};
 	TArray<TSharedPtr<FBreakpointDisplay>> WriteBreakpoints;
diff --git a/Engine/Source/Editor/MassEntityDebugger/Private/SMassDebugger.cpp b/Engine/Source/Editor/MassEntityDebugger/Private/SMassDebugger.cpp
--- a/Engine/Source/Editor/MassEntityDebugger/Private/SMassDebugger.cpp
+++ b/Engine/Source/Editor/MassEntityDebugger/Private/SMassDebugger.cpp
@@ -6,6 +6,7 @@
 #include "SMassProcessingView.h"
 #include "SMassArchetypesView.h"
 #include "SMassEntitiesView.h"
+#include "SMassFragmentsView.h"
 #include "MassDebuggerModel.h"
 #include "Engine/Engine.h"
 #include "CoreGlobals.h"
@@ -32,12 +33,13 @@ void FMassDebuggerCommands::RegisterCommands()
 
 namespace UE::Mass::Debugger::Private
 {
-	const FName ToolbarTabId("Toolbar");
-	const FName BreakpointsTabId("Breakpoints");
-	const FName ProcessorsTabId("Processors");
-	const FName ProcessingGraphTabId("Processing Graphs");
-	const FName ArchetypesTabId("Archetypes");
-	const FName EntitiesTabId("Entities");
+	const FLazyName ToolbarTabId("Toolbar");
+	const FLazyName BreakpointsTabId("Breakpoints");
+	const FLazyName ProcessorsTabId("Processors");
+	const FLazyName ProcessingGraphTabId("Processing Graphs");
+	const FLazyName FragmentsTabId("Fragments");
+	const FLazyName ArchetypesTabId("Archetypes");
+	const FLazyName EntitiesTabId("Entities");
 
 	bool IsSupportedWorldType(const EWorldType::Type WorldType)
 	{
@@ -95,6 +97,10 @@ void SMassDebugger::Construct(const FArguments& InArgs, const TSharedRef<SDockTa
 		.SetDisplayName(LOCTEXT("ProcessingTabTitle", "Processing Graphs"))
 		.SetGroup(AppMenuGroup);
 
+	TabManager->RegisterTabSpawner(FragmentsTabId, FOnSpawnTab::CreateRaw(this, &SMassDebugger::SpawnFragmentsTab))
+		.SetDisplayName(LOCTEXT("FragmentsTabTitle", "Fragments"))
+		.SetGroup(AppMenuGroup);
+
 	TabManager->RegisterTabSpawner(ArchetypesTabId, FOnSpawnTab::CreateRaw(this, &SMassDebugger::SpawnArchetypesTab))
 		.SetDisplayName(LOCTEXT("ArchetypesTabTitle", "Archetypes"))
 		.SetGroup(AppMenuGroup);
@@ -159,13 +165,14 @@ TSharedRef<FTabManager::FLayout> SMassDebugger::CreateDefaultLayout()
 						FTabManager::NewStack()
 						->AddTab(ProcessorsTabId, ETabState::OpenedTab)
 						->AddTab(ProcessingGraphTabId, ETabState::OpenedTab)
-						->SetForegroundTab(ProcessorsTabId)
+						->AddTab(FragmentsTabId, ETabState::OpenedTab)
+						->SetForegroundTab(ProcessorsTabId.Resolve())
 					)
 					->Split
 					(
 						FTabManager::NewStack()
 						->AddTab(BreakpointsTabId, ETabState::OpenedTab)
-						->SetForegroundTab(BreakpointsTabId)
+						->SetForegroundTab(BreakpointsTabId.Resolve())
 					)
 				)
 				->Split
@@ -173,7 +180,7 @@ TSharedRef<FTabManager::FLayout> SMassDebugger::CreateDefaultLayout()
 					FTabManager::NewStack()
 					->AddTab(ArchetypesTabId, ETabState::OpenedTab)
 					->AddTab(EntitiesTabId, ETabState::OpenedTab)
-					->SetForegroundTab(ArchetypesTabId)
+					->SetForegroundTab(ArchetypesTabId.Resolve())
 				)
 			)
 		);
@@ -244,6 +251,7 @@ TSharedRef<SWidget> SMassDebugger::GenerateWindowMenu()
 	FUIAction ProcessorsUIAction(FExecuteAction::CreateSP(this, &SMassDebugger::ShowProcessorView));
 	FUIAction ArchetypesUIAction(FExecuteAction::CreateSP(this, &SMassDebugger::ShowArchetypesView));
 	FUIAction ProcessingGraphsUIAction(FExecuteAction::CreateSP(this, &SMassDebugger::ShowProcessingGraphsView));
+	FUIAction FragmentsUIAction(FExecuteAction::CreateSP(this, &SMassDebugger::ShowFragmentsView));
 	FUIAction BreakpointsUIAction(FExecuteAction::CreateSP(this, &SMassDebugger::ShowBreakpointsView));
 	FUIAction EntitesUIAction(FExecuteAction::CreateSP(this, &SMassDebugger::ShowEntitesView));
 	FUIAction ResetLayoutUIAction(FExecuteAction::CreateSP(this, &SMassDebugger::ResetLayout));
@@ -252,6 +260,7 @@ TSharedRef<SWidget> SMassDebugger::GenerateWindowMenu()
 	MenuBuilder.AddMenuEntry(LOCTEXT("ProcessorsTabLabel", "Processors"), LOCTEXT("ProcessorsTabTooltip", "Show Processors Tab"), FSlateIcon(), ProcessorsUIAction);
 	MenuBuilder.AddMenuEntry(LOCTEXT("ArchetypesTabLabel", "Archetypes"), LOCTEXT("ArchtypesTabTooltip", "Show Archtypes Tab"), FSlateIcon(), ArchetypesUIAction);
 	MenuBuilder.AddMenuEntry(LOCTEXT("ProcessingGraphsTabLabel", "Processing Graphs"), LOCTEXT("ProcessingGraphsTabTooltip", "Show Processing Graphs Tab"), FSlateIcon(), ProcessingGraphsUIAction);
+	MenuBuilder.AddMenuEntry(LOCTEXT("FragmentsTabLabel", "Fragments"), LOCTEXT("FragmentsTabTooltip", "Show Fragments"), FSlateIcon(), FragmentsUIAction);
 	MenuBuilder.AddMenuEntry(LOCTEXT("BreakpointsTabLabel", "Breakpoints"), LOCTEXT("BreakpointsTabTooltip", "Show Breakpoints Tab"), FSlateIcon(), BreakpointsUIAction);
 	MenuBuilder.AddMenuEntry(LOCTEXT("EntitiesTabLabel", "Entities"), LOCTEXT("EntitiesTabTooltip", "Show Entities Tab"), FSlateIcon(), EntitesUIAction);
 	MenuBuilder.AddSeparator();
@@ -382,6 +391,17 @@ TSharedRef<SDockTab> SMassDebugger::SpawnProcessingTab(const FSpawnTabArgs& Args
 	return MajorTab;
 }
 
+TSharedRef<SDockTab> SMassDebugger::SpawnFragmentsTab(const FSpawnTabArgs& Args)
+{
+	const TSharedRef<SDockTab> MajorTab = SNew(SDockTab)
+		.TabRole(ETabRole::PanelTab);
+
+	TSharedPtr<SWidget> TabContent = SNew(UE::MassDebugger::FragmentsView::SFragmentsView, DebuggerModel);
+	MajorTab->SetContent(TabContent.ToSharedRef());
+
+	return MajorTab;
+}
+
 TSharedRef<SDockTab> SMassDebugger::SpawnArchetypesTab(const FSpawnTabArgs& Args)
 {
 	const TSharedRef<SDockTab> MajorTab = SNew(SDockTab)
@@ -432,27 +452,32 @@ void SMassDebugger::RefreshData()
 
 void SMassDebugger::ShowProcessorView() const
 {
-	TabManager->TryInvokeTab(UE::Mass::Debugger::Private::ProcessorsTabId);
+	TabManager->TryInvokeTab(UE::Mass::Debugger::Private::ProcessorsTabId.Resolve());
 }
 
 void SMassDebugger::ShowArchetypesView() const
 {
-	TabManager->TryInvokeTab(UE::Mass::Debugger::Private::ArchetypesTabId);
+	TabManager->TryInvokeTab(UE::Mass::Debugger::Private::ArchetypesTabId.Resolve());
 }
 
 void SMassDebugger::ShowBreakpointsView() const
 {
-	TabManager->TryInvokeTab(UE::Mass::Debugger::Private::BreakpointsTabId);
+	TabManager->TryInvokeTab(UE::Mass::Debugger::Private::BreakpointsTabId.Resolve());
 }
 
 void SMassDebugger::ShowProcessingGraphsView() const
 {
-	TabManager->TryInvokeTab(UE::Mass::Debugger::Private::ProcessingGraphTabId);
+	TabManager->TryInvokeTab(UE::Mass::Debugger::Private::ProcessingGraphTabId.Resolve());
+}
+
+void SMassDebugger::ShowFragmentsView() const
+{
+	TabManager->TryInvokeTab(UE::Mass::Debugger::Private::FragmentsTabId.Resolve());
 }
 
 void SMassDebugger::ShowEntitesView() const
 {
-	TabManager->TryInvokeTab(UE::Mass::Debugger::Private::EntitiesTabId);
+	TabManager->TryInvokeTab(UE::Mass::Debugger::Private::EntitiesTabId.Resolve());
 }
 
 void SMassDebugger::ResetLayout()
diff --git a/Engine/Source/Editor/MassEntityDebugger/Private/SMassDebugger.h b/Engine/Source/Editor/MassEntityDebugger/Private/SMassDebugger.h
--- a/Engine/Source/Editor/MassEntityDebugger/Private/SMassDebugger.h
+++ b/Engine/Source/Editor/MassEntityDebugger/Private/SMassDebugger.h
@@ -60,6 +60,7 @@ public:
 	void ShowArchetypesView() const;
 	void ShowBreakpointsView() const;
 	void ShowProcessingGraphsView() const;
+	void ShowFragmentsView() const;
 	void ShowEntitesView() const;
 	void ResetLayout();
 	void ShowSelectedView() const;
@@ -68,6 +69,7 @@ protected:
 	TSharedRef<SDockTab> SpawnToolbar(const FSpawnTabArgs& Args);
 	TSharedRef<SDockTab> SpawnProcessorsTab(const FSpawnTabArgs& Args);
 	TSharedRef<SDockTab> SpawnProcessingTab(const FSpawnTabArgs& Args);
+	TSharedRef<SDockTab> SpawnFragmentsTab(const FSpawnTabArgs& Args);
 	TSharedRef<SDockTab> SpawnArchetypesTab(const FSpawnTabArgs& Args);
 	TSharedRef<SDockTab> SpawnBreakpointsTab(const FSpawnTabArgs& Args);
 	TSharedRef<SDockTab> SpawnEntitiesTab(const FSpawnTabArgs& Args);
diff --git a/Engine/Source/Editor/MassEntityDebugger/Private/SMassEntitiesList.cpp b/Engine/Source/Editor/MassEntityDebugger/Private/SMassEntitiesList.cpp
--- a/Engine/Source/Editor/MassEntityDebugger/Private/SMassEntitiesList.cpp
+++ b/Engine/Source/Editor/MassEntityDebugger/Private/SMassEntitiesList.cpp
@@ -700,7 +700,7 @@ TSharedRef<SWidget> SMassEntitiesList::SEntitiesTableRow::GenerateBreakpointWidg
 				TSharedPtr<FMassDebuggerModel> Model = WeakModel.Pin();
 				if (Model.IsValid() && Model->Environment.IsValid() && Model->Environment->EntityManager.IsValid())
 				{
-					FMassDebugger::SetFragmentWriteBreak(*Model->Environment->EntityManager.Pin(), FragmentType, Entity);
+					FMassDebugger::SetFragmentWriteBreakpoint(*Model->Environment->EntityManager.Pin(), FragmentType, Entity);
 				}
 				return FReply::Handled();
 			})
diff --git a/Engine/Source/Editor/MassEntityDebugger/Private/SMassFragmentsView.cpp b/Engine/Source/Editor/MassEntityDebugger/Private/SMassFragmentsView.cpp
new file mode 100644
--- /dev/null
+++ b/Engine/Source/Editor/MassEntityDebugger/Private/SMassFragmentsView.cpp
@@ -0,0 +1,272 @@
+// Copyright Epic Games, Inc. All Rights Reserved.
+
+#include "SMassFragmentsView.h"
+
+#include "MassDebuggerModel.h"
+#include "MassEntityManager.h"         
+#include "MassArchetypeData.h"         
+#include "MassProcessingTypes.h"       
+#include "MassProcessor.h"             
+
+#include "Styling/AppStyle.h"
+#include "Widgets/Input/SButton.h"
+#include "Widgets/Views/SHeaderRow.h"
+#include "Widgets/Text/STextBlock.h"
+#include "Framework/Views/TableViewMetadata.h"
+#include "Misc/TextFilter.h"
+
+#define LOCTEXT_NAMESPACE "SMassDebugger"
+
+//----------------------------------------------------------------------//
+// SFragmentsView
+//----------------------------------------------------------------------//
+namespace UE::MassDebugger::FragmentsView
+{
+	namespace Private
+	{
+		static const FLazyName ColumnName = TEXT("Name");
+		static const FLazyName ColumnDebugControls = TEXT("DebugControls");
+		static const FLazyName ColumnNumArchetypes = TEXT("Archetypes");
+		static const FLazyName ColumnNumEntities = TEXT("Entities");
+	}
+
+void SFragmentsView::Construct(const FArguments& InArgs, TSharedRef<FMassDebuggerModel> InDebuggerModel)
+{ 
+#if WITH_MASSENTITY_DEBUG
+	Initialize(InDebuggerModel);
+	TSharedRef<SHeaderRow> HeaderRowWidget =
+		SNew(SHeaderRow)
+		+ SHeaderRow::Column(Private::ColumnName)
+		.DefaultLabel(LOCTEXT("FragmentColumnName", "Fragment Type"))
+		.FillWidth(0.5f)
+		.SortMode(this, &SFragmentsView::GetColumnSortMode, Private::ColumnName.Resolve())
+		.OnSort(this, &SFragmentsView::OnColumnSortModeChanged)
+		+ SHeaderRow::Column(Private::ColumnDebugControls)
+		.DefaultLabel(LOCTEXT("FragmentColumnDebugControls", "Debug Controls"))
+		.FixedWidth(30.f)
+		.HAlignHeader(HAlign_Center)
+		.HAlignCell(HAlign_Center)
+		.VAlignCell(VAlign_Center)
+		.SortMode(EColumnSortMode::None)
+		+ SHeaderRow::Column(Private::ColumnNumArchetypes)
+		.DefaultLabel(LOCTEXT("FragmentColumnArchetypes", "Archetypes"))
+		.FillWidth(0.25f)
+		.SortMode(this, &SFragmentsView::GetColumnSortMode, Private::ColumnNumArchetypes.Resolve())
+		.OnSort(this, &SFragmentsView::OnColumnSortModeChanged)
+		+ SHeaderRow::Column(Private::ColumnNumEntities)
+		.DefaultLabel(LOCTEXT("FragmentColumnEntities", "Entities"))
+		.FillWidth(0.25f)
+		.SortMode(this, &SFragmentsView::GetColumnSortMode, Private::ColumnNumEntities.Resolve())
+		.OnSort(this, &SFragmentsView::OnColumnSortModeChanged);
+
+	SAssignNew(FragmentsList, SListView<FFragmentsTableRowPtr>)
+		.ListItemsSource(&FragmentListItemsSource)
+		.OnGenerateRow(this, &SFragmentsView::OnGenerateFragmentRow)
+		.HeaderRow(HeaderRowWidget);
+
+	ChildSlot
+		[
+			FragmentsList.ToSharedRef()
+		];
+
+	RefreshFragmentList();
+
+#else // WITH_MASSENTITY_DEBUG
+	ChildSlot
+		[
+			SNew(STextBlock)
+				.Text(LOCTEXT("MassEntityDebuggingNotEnabled", "Mass Entity Debugging Not Enabled for this configuration"))
+				.Justification(ETextJustify::Center)
+		];
+#endif // WITH_MASSENTITY_DEBUG
+}
+
+TSharedRef<ITableRow> SFragmentsView::OnGenerateFragmentRow(FFragmentsTableRowPtr InItem, const TSharedRef<STableViewBase>& OwnerTable)
+{
+	return SNew(SFragmentsTableRow, OwnerTable, InItem);
+}
+
+EColumnSortMode::Type SFragmentsView::GetColumnSortMode(FName ColumnId) const
+{
+	if (SortByColumnId != ColumnId)
+	{
+		return EColumnSortMode::None;
+	}
+
+	return bSortAscending ? EColumnSortMode::Ascending : EColumnSortMode::Descending;
+}
+
+void SFragmentsView::OnColumnSortModeChanged(const EColumnSortPriority::Type SortPriority, const FName& ColumnId, const EColumnSortMode::Type InSortMode)
+{
+	if (ColumnId == Private::ColumnDebugControls)
+	{
+		return;
+	}
+
+	if (SortPriority == EColumnSortPriority::Primary)
+	{
+		SortByColumnId = ColumnId;
+		bSortAscending = InSortMode == EColumnSortMode::Ascending ? true : false;
+	}
+	RefreshFragmentList();
+}
+
+void SFragmentsView::OnRefresh()
+{
+#if WITH_MASSENTITY_DEBUG
+	RefreshFragmentList();
+#endif // WITH_MASSENTITY_DEBUG
+}
+
+void SFragmentsView::RefreshFragmentList()
+{
+	FragmentListItemsSource.Reset();
+#if WITH_MASSENTITY_DEBUG
+	if (!DebuggerModel.IsValid())
+	{
+		return;
+	}
+
+	FragmentListItemsSource.Reserve(DebuggerModel->CachedFragmentData.Num());
+
+	for (TSharedPtr<FMassDebuggerFragmentData>& FragmentData : DebuggerModel->CachedFragmentData)
+	{
+		// filter out unused fragments
+		// @todo: add checkbox to "Show Unused Fragments"
+		if (FragmentData->Archetypes.Num() > 0)
+		{
+			FragmentListItemsSource.Add(FragmentData);
+		}
+	}
+
+	if (SortByColumnId == Private::ColumnName)
+	{
+		auto SortName = [bAscending = bSortAscending](const TSharedPtr<FMassDebuggerFragmentData>& A, const TSharedPtr<FMassDebuggerFragmentData>& B)
+		{
+			const int32 Compare = A->Name.CompareTo(B->Name) <= 0;
+			return bAscending ? Compare > 0 : Compare <= 0;
+		};
+		FragmentListItemsSource.StableSort(SortName);
+	}
+	else if (SortByColumnId == Private::ColumnNumArchetypes)
+	{
+		auto SortArchetypes = [bAscending = bSortAscending](const TSharedPtr<FMassDebuggerFragmentData>& A, const TSharedPtr<FMassDebuggerFragmentData>& B)
+		{
+			return bAscending ? A->Archetypes.Num() < B->Archetypes.Num() : A->Archetypes.Num() > B->Archetypes.Num();
+		};
+		FragmentListItemsSource.StableSort(SortArchetypes);
+	}
+	else if (SortByColumnId == Private::ColumnNumEntities)
+	{
+		auto SortEntities = [bAscending = bSortAscending](const TSharedPtr<FMassDebuggerFragmentData>& A, const TSharedPtr<FMassDebuggerFragmentData>& B)
+		{
+			return bAscending ? A->NumEntities < B->NumEntities : A->NumEntities > B->NumEntities;
+		};
+		FragmentListItemsSource.StableSort(SortEntities);
+	}
+
+	if (FragmentsList.IsValid())
+	{
+		FragmentsList->RequestListRefresh();
+	}
+#endif // WITH_MASSENTITY_DEBUG
+}
+
+//----------------------------------------------------------------------//
+// SFragmentsTableRow - Represents a single row in the list
+//----------------------------------------------------------------------//
+void SFragmentsTableRow::Construct(const FArguments& InArgs, const TSharedRef<STableViewBase>& InOwnerTableView, FFragmentsTableRowPtr InFragmentDataPtr)
+{
+#if WITH_MASSENTITY_DEBUG
+
+	FragmentDataPtr = InFragmentDataPtr;
+
+	SMultiColumnTableRow<FFragmentsTableRowPtr>::Construct(
+		SMultiColumnTableRow<FFragmentsTableRowPtr>::FArguments()
+		.ShowSelection(false)
+		, InOwnerTableView);
+
+#else // WITH_MASSENTITY_DEBUG
+	SMultiColumnTableRow<FFragmentsTableRowPtr>::Construct(FArguments(), InOwnerTableView);
+#endif //WITH_MASSENTITY_DEBUG
+}
+
+TSharedRef<SWidget> SFragmentsTableRow::GenerateWidgetForColumn(const FName& InColumnName)
+{
+	TSharedRef<SWidget> CellWidget = SNullWidget::NullWidget;
+
+#if WITH_MASSENTITY_DEBUG
+
+	if (!FragmentDataPtr.IsValid())
+	{
+		return CellWidget;
+	}
+
+	if (InColumnName == Private::ColumnName)
+	{
+		SAssignNew(CellWidget, STextBlock)
+			.Text(FragmentDataPtr->Name)
+			.ToolTipText(FragmentDataPtr->Fragment.IsValid() ? FText::FromString(FragmentDataPtr->Fragment.Pin()->GetPathName()) : LOCTEXT("InvalidFragmentPath", "Invalid Fragment Struct"));
+	}
+	else if (InColumnName == Private::ColumnNumArchetypes)
+	{
+		SAssignNew(CellWidget, STextBlock)
+			.Text(FText::AsNumber(FragmentDataPtr->Archetypes.Num()))
+			.Justification(ETextJustify::Right);
+	}
+	else if (InColumnName == Private::ColumnNumEntities)
+	{
+		SAssignNew(CellWidget, STextBlock)
+			.Text(FText::AsNumber(FragmentDataPtr->NumEntities))
+			.Justification(ETextJustify::Right);
+	}
+	else if (InColumnName == Private::ColumnDebugControls)
+	{
+		SAssignNew(CellWidget, SButton)
+			.ToolTipText(LOCTEXT("BreakpointTooltip", "Break on writes for selected entity."))
+			.ButtonStyle(FAppStyle::Get(), "NoBorder")
+			.ContentPadding(0)
+			.OnClicked_Lambda([WeakEntityManager = FragmentDataPtr->EntityManager, Fragment = FragmentDataPtr->Fragment]()
+			{
+				TSharedPtr<const FMassEntityManager> StrongEntityManager = WeakEntityManager.Pin();
+				TStrongObjectPtr<const UScriptStruct> StrongFragment = Fragment.Pin();
+				if (StrongEntityManager && StrongFragment)
+				{
+					FMassDebugger::ToggleFragmentWriteBreakForSelectedEntity(*StrongEntityManager, StrongFragment.Get());
+				}
+				return FReply::Handled();
+			})
+			[
+				SNew(SImage)
+				.ColorAndOpacity_Lambda([WeakEntityManager = FragmentDataPtr->EntityManager, Fragment = FragmentDataPtr->Fragment]()
+				{
+					TSharedPtr<const FMassEntityManager> StrongEntityManager = WeakEntityManager.Pin();
+					TStrongObjectPtr<const UScriptStruct> StrongFragment = Fragment.Pin();
+					if (StrongEntityManager && StrongFragment)
+					{
+						if(FMassDebugger::IsFragmentWriteBreakForSelectedEntitySet(*StrongEntityManager, StrongFragment.Get()))
+						{
+							return FSlateColor(FLinearColor(1.f, 0.f, 0.f, 1.f));
+						}
+					}
+					return FSlateColor(FLinearColor(1.f, 1.f, 1.f, 0.8f));
+				})
+				.Image(FAppStyle::GetBrush("Icons.FilledCircle"))
+			];
+	}
+
+
+#endif //WITH_MASSENTITY_DEBUG
+
+	return SNew(SHorizontalBox)
+		+ SHorizontalBox::Slot()
+			.Padding(FMargin(4.f, 0.f))
+			.VAlign(VAlign_Center)
+			[
+				CellWidget
+			];
+}
+
+} // UE::MassDebugger::FragmentsView
+
+#undef LOCTEXT_NAMESPACE
\ No newline at end of file
diff --git a/Engine/Source/Editor/MassEntityDebugger/Private/SMassFragmentsView.h b/Engine/Source/Editor/MassEntityDebugger/Private/SMassFragmentsView.h
new file mode 100644
--- /dev/null
+++ b/Engine/Source/Editor/MassEntityDebugger/Private/SMassFragmentsView.h
@@ -0,0 +1,68 @@
+// Copyright Epic Games, Inc. All Rights Reserved.
+#pragma once
+
+#include "SMassDebuggerViewBase.h"
+#include "Widgets/Views/SListView.h"
+#include "Widgets/Views/STableViewBase.h"
+#include "Widgets/Views/STableRow.h"
+#include "Templates/SharedPointer.h"
+
+struct FMassDebuggerModel;
+struct FMassDebuggerFragmentData;
+class SHeaderRow;
+class STableViewBase;
+class UMassProcessor;
+class UScriptStruct;
+struct FMassArchetypeHandle;
+
+
+namespace UE::MassDebugger::FragmentsView
+{
+
+using FFragmentsTableRowPtr = TSharedPtr<FMassDebuggerFragmentData>;
+
+class SFragmentsTableRow : public SMultiColumnTableRow<FFragmentsTableRowPtr>
+{
+	using Super = SMultiColumnTableRow<FFragmentsTableRowPtr>;
+public:
+
+	SLATE_BEGIN_ARGS(SFragmentsTableRow) {}
+	SLATE_END_ARGS()
+
+	void Construct(const FArguments& InArgs, const TSharedRef<STableViewBase>& InOwnerTableView, FFragmentsTableRowPtr InFragmentDataPtr);
+	virtual TSharedRef<SWidget> GenerateWidgetForColumn(const FName& InColumnName) override;
+
+private:
+	FFragmentsTableRowPtr FragmentDataPtr;
+};
+
+
+class SFragmentsView : public SMassDebuggerViewBase
+{
+public:
+	SLATE_BEGIN_ARGS(SFragmentsView)
+		{}
+	SLATE_END_ARGS()
+	void Construct(const FArguments& InArgs, TSharedRef<FMassDebuggerModel> InDebuggerModel);
+
+	virtual void OnRefresh() override;
+	virtual void OnProcessorsSelected(TConstArrayView<TSharedPtr<FMassDebuggerProcessorData>> SelectedProcessors, ESelectInfo::Type SelectInfo) override
+	{}
+	virtual void OnArchetypesSelected(TConstArrayView<TSharedPtr<FMassDebuggerArchetypeData>> SelectedArchetypes, ESelectInfo::Type SelectInfo) override
+	{}
+
+protected:
+	TSharedRef<ITableRow> OnGenerateFragmentRow(FFragmentsTableRowPtr InItem, const TSharedRef<STableViewBase>& OwnerTable);
+	void RefreshFragmentList();
+
+	EColumnSortMode::Type GetColumnSortMode(FName ColumnId) const;
+	void OnColumnSortModeChanged(const EColumnSortPriority::Type SortPriority, const FName& ColumnId, const EColumnSortMode::Type InSortMode);
+
+private:
+	TSharedPtr<SListView<FFragmentsTableRowPtr>> FragmentsList;
+	TArray<FFragmentsTableRowPtr> FragmentListItemsSource;
+	FName SortByColumnId;
+	bool bSortAscending = true;
+};
+
+} // UE::MassDebugger::FragmentsView
\ No newline at end of file
diff --git a/Engine/Source/Runtime/MassEntity/Private/MassDebugger.cpp b/Engine/Source/Runtime/MassEntity/Private/MassDebugger.cpp
--- a/Engine/Source/Runtime/MassEntity/Private/MassDebugger.cpp
+++ b/Engine/Source/Runtime/MassEntity/Private/MassDebugger.cpp
@@ -411,7 +411,7 @@ namespace UE::Mass::Debug
 					{
 						if (const UScriptStruct* FragmentType = FindElementTypeByName(PartialFragmentName))
 						{
-							FMassDebugger::SetFragmentWriteBreak(EntityManager, FragmentType, SelectedEntity);
+							FMassDebugger::SetFragmentWriteBreakpoint(EntityManager, FragmentType, SelectedEntity);
 						}
 						else
 						{
@@ -833,7 +833,7 @@ void FMassDebugger::SelectEntity(const FMassEntityManager& EntityManager, const
 	{
 		UE::Mass::Debug::SetDebugEntityRange(EntityHandle.Index, EntityHandle.Index);
 
-		GetActiveEnvironment(EntityManager).SelectedEntity = EntityHandle;
+		GetActiveEnvironmentChecked(EntityManager).SelectedEntity = EntityHandle;
 
 		OnEntitySelectedDelegate.Broadcast(EntityManager, EntityHandle);
 	}
@@ -841,17 +841,29 @@ void FMassDebugger::SelectEntity(const FMassEntityManager& EntityManager, const
 
 FMassEntityHandle FMassDebugger::GetSelectedEntity(const FMassEntityManager& EntityManager)
 {
-	return GetActiveEnvironment(EntityManager).SelectedEntity;
+	return GetActiveEnvironmentChecked(EntityManager).SelectedEntity;
 }
 
 void FMassDebugger::HighlightEntity(const FMassEntityManager& EntityManager, const FMassEntityHandle EntityHandle)
 {
-	GetActiveEnvironment(EntityManager).HighlightedEntity = EntityHandle;
+	if (FEnvironment* ActiveEnvironment = GetActiveEnvironment(EntityManager))
+	{
+		ActiveEnvironment->HighlightedEntity = EntityHandle;
+	}
 }
 
 FMassEntityHandle FMassDebugger::GetHighlightedEntity(const FMassEntityManager& EntityManager)
 {
-	return GetActiveEnvironment(EntityManager).HighlightedEntity;
+	if (FEnvironment* ActiveEnvironment = GetActiveEnvironment(EntityManager))
+	{
+		return ActiveEnvironment->HighlightedEntity;
+	}
+	return FMassEntityHandle();
+}
+
+bool FMassDebugger::IsEntityManagerInitialized(const FMassEntityManager& EntityManager)
+{
+	return EntityManager.InitializationState == FMassEntityManager::EInitializationState::Initialized;
 }
 
 int32 FMassDebugger::RegisterEntityManager(FMassEntityManager& EntityManager)
@@ -936,7 +948,7 @@ bool FMassDebugger::ShouldProcessorBreak(const FMassEntityManager& EntityManager
 		return false;
 	}
 
-	FEnvironment& ActiveEnvironment = GetActiveEnvironment(EntityManager);
+	FEnvironment& ActiveEnvironment = GetActiveEnvironmentChecked(EntityManager);
 
 	if (LIKELY(!ActiveEnvironment.bHasBreakpoint))
 	{
@@ -955,7 +967,7 @@ bool FMassDebugger::HasAnyProcessorBreakpoints(const FMassEntityManager& EntityM
 		return false;
 	}
 
-	FEnvironment& ActiveEnvironment = GetActiveEnvironment(EntityManager);
+	FEnvironment& ActiveEnvironment = GetActiveEnvironmentChecked(EntityManager);
 
 	if (LIKELY(!ActiveEnvironment.bHasBreakpoint))
 	{
@@ -972,7 +984,7 @@ bool FMassDebugger::ShouldBreakOnFragmentWrite(const FMassEntityManager& EntityM
 		return false;
 	}
 
-	FEnvironment& ActiveEnvironment = GetActiveEnvironment(EntityManager);
+	FEnvironment& ActiveEnvironment = GetActiveEnvironmentChecked(EntityManager);
 
 	if (LIKELY(!ActiveEnvironment.bHasBreakpoint))
 	{
@@ -981,7 +993,20 @@ bool FMassDebugger::ShouldBreakOnFragmentWrite(const FMassEntityManager& EntityM
 
 	const FMassEntityHandle* Found = ActiveEnvironment.FragmentWriteBreakpoints.FindPair(FragmentType, Entity);
 
-	return Found != nullptr;
+	if (Found != nullptr)
+	{
+		return true;
+	}
+
+	if (Entity == ActiveEnvironment.SelectedEntity)
+	{
+		if (ActiveEnvironment.SelectedEntityFragmentWriteBreakpoints.Contains(FragmentType))
+		{
+			return true;
+		}
+	}
+
+	return false;
 }
 
 bool FMassDebugger::HasAnyFragmentWriteBreakpoints(const FMassEntityManager& EntityManager, const UScriptStruct* FragmentType)
@@ -991,7 +1016,7 @@ bool FMassDebugger::HasAnyFragmentWriteBreakpoints(const FMassEntityManager& Ent
 		return false;
 	}
 
-	FEnvironment& ActiveEnvironment = GetActiveEnvironment(EntityManager);
+	FEnvironment& ActiveEnvironment = GetActiveEnvironmentChecked(EntityManager);
 
 	if (LIKELY(!ActiveEnvironment.bHasBreakpoint))
 	{
@@ -1000,10 +1025,10 @@ bool FMassDebugger::HasAnyFragmentWriteBreakpoints(const FMassEntityManager& Ent
 
 	if (FragmentType == nullptr)
 	{
-		return ActiveEnvironment.FragmentWriteBreakpoints.Num() > 0;
+		return ActiveEnvironment.SelectedEntityFragmentWriteBreakpoints.Num() > 0 || ActiveEnvironment.FragmentWriteBreakpoints.Num() > 0;
 	}
 
-	return ActiveEnvironment.FragmentWriteBreakpoints.Contains(FragmentType);
+	return ActiveEnvironment.SelectedEntityFragmentWriteBreakpoints.Contains(FragmentType) || ActiveEnvironment.FragmentWriteBreakpoints.Contains(FragmentType);
 }
 
 void FMassDebugger::SetProcessorBreakpoint(const FMassEntityManager& EntityManager, TNotNull<const UMassProcessor*> Processor, FMassEntityHandle Entity)
@@ -1013,7 +1038,7 @@ void FMassDebugger::SetProcessorBreakpoint(const FMassEntityManager& EntityManag
 		FMessageDialog::Open(EAppMsgType::Ok, LOCTEXT("NoDebuggerAttached", "Breakpoint set but no debugger is attached."));
 	}
 
-	FEnvironment& ActiveEnvironment = GetActiveEnvironment(EntityManager);
+	FEnvironment& ActiveEnvironment = GetActiveEnvironmentChecked(EntityManager);
 
 	ActiveEnvironment.bHasBreakpoint = true;
 	bHasBreakpoint = true;
@@ -1022,14 +1047,14 @@ void FMassDebugger::SetProcessorBreakpoint(const FMassEntityManager& EntityManag
 	OnBreakpointsChangedDelegate.Broadcast();
 }
 
-void FMassDebugger::SetFragmentWriteBreak(const FMassEntityManager& EntityManager, const UScriptStruct* FragmentType, FMassEntityHandle Entity)
+void FMassDebugger::SetFragmentWriteBreakpoint(const FMassEntityManager& EntityManager, TNotNull<const UScriptStruct*> FragmentType, FMassEntityHandle Entity)
 {
 	if (!FPlatformMisc::IsDebuggerPresent())
 	{
 		FMessageDialog::Open(EAppMsgType::Ok, LOCTEXT("NoDebuggerAttached", "Breakpoint set but no debugger is attached."));
 	}
 
-	FEnvironment& ActiveEnvironment = GetActiveEnvironment(EntityManager);
+	FEnvironment& ActiveEnvironment = GetActiveEnvironmentChecked(EntityManager);
 
 	ActiveEnvironment.bHasBreakpoint = true;
 	bHasBreakpoint = true;
@@ -1038,9 +1063,39 @@ void FMassDebugger::SetFragmentWriteBreak(const FMassEntityManager& EntityManage
 	OnBreakpointsChangedDelegate.Broadcast();
 }
 
+void FMassDebugger::ToggleFragmentWriteBreakForSelectedEntity(const FMassEntityManager& EntityManager, TNotNull<const UScriptStruct*> FragmentType)
+{
+	FEnvironment& ActiveEnvironment = GetActiveEnvironmentChecked(EntityManager);
+
+	ActiveEnvironment.bHasBreakpoint = true;
+	bHasBreakpoint = true;
+
+	if (ActiveEnvironment.SelectedEntityFragmentWriteBreakpoints.Contains(FragmentType))
+	{
+		ActiveEnvironment.SelectedEntityFragmentWriteBreakpoints.Remove(FragmentType);
+	}
+	else
+	{
+		if (!FPlatformMisc::IsDebuggerPresent())
+		{
+			FMessageDialog::Open(EAppMsgType::Ok, LOCTEXT("NoDebuggerAttached", "Breakpoint set but no debugger is attached."));
+		}
+		ActiveEnvironment.SelectedEntityFragmentWriteBreakpoints.Add(FragmentType);
+	}
+		
+	OnBreakpointsChangedDelegate.Broadcast();
+}
+
+bool FMassDebugger::IsFragmentWriteBreakForSelectedEntitySet(const FMassEntityManager& EntityManager, TNotNull<const UScriptStruct*> FragmentType)
+{
+	FEnvironment* ActiveEnvironment = GetActiveEnvironment(EntityManager);
+	return ActiveEnvironment ? ActiveEnvironment->SelectedEntityFragmentWriteBreakpoints.Contains(FragmentType) : false;
+}
+
+
 void FMassDebugger::ClearProcessorBreakpoint(const FMassEntityManager& EntityManager, const UMassProcessor* Processor, FMassEntityHandle Entity)
 {
-	FEnvironment& ActiveEnvironment = GetActiveEnvironment(EntityManager);
+	FEnvironment& ActiveEnvironment = GetActiveEnvironmentChecked(EntityManager);
 
 	ActiveEnvironment.ProcessorBreakpoints.Remove(Processor, Entity);
 
@@ -1050,14 +1105,18 @@ void FMassDebugger::ClearProcessorBreakpoint(const FMassEntityManager& EntityMan
 
 void FMassDebugger::ClearAllProcessorBreakpoints(const FMassEntityManager& EntityManager, const UMassProcessor* Processor)
 {
-	FEnvironment& ActiveEnvironment = GetActiveEnvironment(EntityManager);
+	FEnvironment& ActiveEnvironment = GetActiveEnvironmentChecked(EntityManager);
 	ActiveEnvironment.ProcessorBreakpoints.Remove(Processor);
 	OnBreakpointsChangedDelegate.Broadcast();
 }
 
 void FMassDebugger::ClearFragmentWriteBreak(const FMassEntityManager& EntityManager, const UScriptStruct* FragmentType, FMassEntityHandle Entity)
 {
-	FEnvironment& ActiveEnvironment = GetActiveEnvironment(EntityManager);
+	FEnvironment& ActiveEnvironment = GetActiveEnvironmentChecked(EntityManager);
+	if (ActiveEnvironment.SelectedEntity == Entity)
+	{
+		ActiveEnvironment.SelectedEntityFragmentWriteBreakpoints.Remove(FragmentType);
+	}
 	ActiveEnvironment.FragmentWriteBreakpoints.Remove(FragmentType, Entity);
 	UpdateHasBreakpoint();
 	OnBreakpointsChangedDelegate.Broadcast();
@@ -1065,15 +1124,15 @@ void FMassDebugger::ClearFragmentWriteBreak(const FMassEntityManager& EntityMana
 
 void FMassDebugger::ClearAllFragmentWriteBreak(const FMassEntityManager& EntityManager, const UScriptStruct* FragmentType)
 {
-	FEnvironment& ActiveEnvironment = GetActiveEnvironment(EntityManager);
+	FEnvironment& ActiveEnvironment = GetActiveEnvironmentChecked(EntityManager);
 	ActiveEnvironment.FragmentWriteBreakpoints.Remove(FragmentType);
 	OnBreakpointsChangedDelegate.Broadcast();
 }
 
 void FMassDebugger::ClearAllEntityBreakpoints(const FMassEntityManager& EntityManager, FMassEntityHandle Entity)
 {
-	FEnvironment& ActiveEnvironment = GetActiveEnvironment(EntityManager);
-	TMultiMap<const UMassProcessor*, FMassEntityHandle>::TIterator ProccessorIterator = ActiveEnvironment.ProcessorBreakpoints.CreateIterator();
+	FEnvironment& ActiveEnvironment = GetActiveEnvironmentChecked(EntityManager);
+	TMultiMap<TObjectKey<const UMassProcessor>, FMassEntityHandle>::TIterator ProccessorIterator = ActiveEnvironment.ProcessorBreakpoints.CreateIterator();
 
 	for (; ProccessorIterator; ++ProccessorIterator)
 	{
@@ -1083,7 +1142,7 @@ void FMassDebugger::ClearAllEntityBreakpoints(const FMassEntityManager& EntityMa
 		}
 	}
 
-	TMultiMap<const UScriptStruct*, FMassEntityHandle>::TIterator FragmentIterator = ActiveEnvironment.FragmentWriteBreakpoints.CreateIterator();
+	TMultiMap<TObjectKey<const UScriptStruct>, FMassEntityHandle>::TIterator FragmentIterator = ActiveEnvironment.FragmentWriteBreakpoints.CreateIterator();
 	for (; FragmentIterator; ++FragmentIterator)
 	{
 		if (FragmentIterator.Value() == Entity)
@@ -1100,7 +1159,7 @@ void FMassDebugger::BreakOnFragmentWriteForSelectedEntity(FName FragmentName)
 	{
 		if (Environment.EntityManager.Pin()->IsEntityValid(Environment.SelectedEntity))
 		{
-			SetFragmentWriteBreak(*Environment.EntityManager.Pin(), GetFragmentTypeFromName(FragmentName), Environment.SelectedEntity);
+			SetFragmentWriteBreakpoint(*Environment.EntityManager.Pin(), GetFragmentTypeFromName(FragmentName), Environment.SelectedEntity);
 		}
 	}
 	OnBreakpointsChangedDelegate.Broadcast();
@@ -1297,12 +1356,13 @@ void FMassDebugger::UpdateHasBreakpoint()
 	for (FEnvironment& Environment : ActiveEnvironments)
 	{
 		Environment.bHasBreakpoint = Environment.ProcessorBreakpoints.Num() != 0
-			|| Environment.FragmentWriteBreakpoints.Num() != 0;
+			|| Environment.FragmentWriteBreakpoints.Num() != 0
+			|| Environment.SelectedEntityFragmentWriteBreakpoints.Num() != 0;
 		bHasBreakpoint |= Environment.bHasBreakpoint;
 	}
 }
 
-FMassDebugger::FEnvironment& FMassDebugger::GetActiveEnvironment(const FMassEntityManager& EntityManager)
+FMassDebugger::FEnvironment& FMassDebugger::GetActiveEnvironmentChecked(const FMassEntityManager& EntityManager)
 {
 	const int32 Index = ActiveEnvironments.IndexOfByPredicate([WeakManager = EntityManager.AsWeak()](const FEnvironment& Element)
 	{
@@ -1314,6 +1374,21 @@ FMassDebugger::FEnvironment& FMassDebugger::GetActiveEnvironment(const FMassEnti
 	return ActiveEnvironments[Index];
 }
 
+FMassDebugger::FEnvironment* FMassDebugger::GetActiveEnvironment(const FMassEntityManager& EntityManager)
+{
+	const int32 Index = ActiveEnvironments.IndexOfByPredicate([WeakManager = EntityManager.AsWeak()](const FEnvironment& Element)
+	{
+		return Element.EntityManager == WeakManager;
+	});
+
+	if (Index == INDEX_NONE)
+	{
+		return nullptr;
+	}
+
+	return &ActiveEnvironments[Index];
+}
+
 FMassDebugger::FEnvironment::FEnvironment(const FMassEntityManager& InEntityManager)
 	: EntityManager(InEntityManager.AsWeak())
 {
@@ -1345,6 +1420,7 @@ FMassDebugger::FEnvironment::~FEnvironment()
 
 void FMassDebugger::FEnvironment::ClearBreakpoints()
 {
+	SelectedEntityFragmentWriteBreakpoints.Reset();
 	ProcessorBreakpoints.Reset();
 	FragmentWriteBreakpoints.Reset();
 	bHasBreakpoint = false;
diff --git a/Engine/Source/Runtime/MassEntity/Public/MassDebugger.h b/Engine/Source/Runtime/MassEntity/Public/MassDebugger.h
--- a/Engine/Source/Runtime/MassEntity/Public/MassDebugger.h
+++ b/Engine/Source/Runtime/MassEntity/Public/MassDebugger.h
@@ -10,6 +10,7 @@
 #include "MassProcessor.h"
 #include "Misc/SpinLock.h"
 #include "Logging/TokenizedMessage.h"
+#include "UObject/ObjectKey.h"
 #if UE_ENABLE_INCLUDE_ORDER_DEPRECATED_IN_5_6
 #include "StructUtils/InstancedStruct.h"
 #endif // UE_ENABLE_INCLUDE_ORDER_DEPRECATED_IN_5_6
@@ -143,13 +144,13 @@ struct FMassDebugger
 		TWeakPtr<const FMassEntityManager> EntityManager;
 		TMap<FName, UE::Mass::Debug::FProcessorProviderFunction> ProcessorProviders;
 		FMassEntityHandle SelectedEntity;
-		// js@todo HighlightedEntity will be used to actually highlight the entity (arrow at transform, etc) in the near future
 		FMassEntityHandle HighlightedEntity;
 
 		bool bHasBreakpoint = false;
 
-		TMultiMap<const UMassProcessor*, FMassEntityHandle> ProcessorBreakpoints;
-		TMultiMap<const UScriptStruct*, FMassEntityHandle> FragmentWriteBreakpoints;
+		TMultiMap<TObjectKey<const UMassProcessor>, FMassEntityHandle> ProcessorBreakpoints;
+		TMultiMap<TObjectKey<const UScriptStruct>, FMassEntityHandle> FragmentWriteBreakpoints;
+		TSet<TObjectKey<const UScriptStruct>> SelectedEntityFragmentWriteBreakpoints;
 
 #if UE_MASS_TRACE_ENABLED
 		FDelegateHandle TraceStartedDelegateHandle;
@@ -241,6 +242,13 @@ struct FMassDebugger
 	 */
 	static MASSENTITY_API int32 RegisterEntityManager(FMassEntityManager& EntityManager);
 	static MASSENTITY_API void UnregisterEntityManager(FMassEntityManager& EntityManager);
+	
+	/**
+	 * Confirms that the initialization state of the EntityManager is Initialized.
+	 * @return true if EntityManager is initialized, false otherwise
+	 */
+	static MASSENTITY_API bool IsEntityManagerInitialized(const FMassEntityManager& EntityManager);
+
 	/**
 	 * Registers the given ProviderFunction with the existing FEnvironment associated with the provided EntityManager.
 	 * If one doesn't exist yet, it will be created (i.e. will automatically call RegisterEntityManager() with the provided EntityManager).
@@ -287,7 +295,17 @@ struct FMassDebugger
 	/**
 	 * Sets a break to be triggered for a processor that's about to write a given fragment on an entity.
 	 */
-	static MASSENTITY_API void SetFragmentWriteBreak(const FMassEntityManager& EntityManager, const UScriptStruct* FragmentType, FMassEntityHandle Entity);
+	static MASSENTITY_API void SetFragmentWriteBreakpoint(const FMassEntityManager& EntityManager, TNotNull<const UScriptStruct*> FragmentType, FMassEntityHandle Entity);
+
+	/**
+	 * Sets or removes a break before a write to a given fragment on whichever entity is selected at the time.
+	 */
+	static MASSENTITY_API void ToggleFragmentWriteBreakForSelectedEntity(const FMassEntityManager& EntityManager, TNotNull<const UScriptStruct*> FragmentType);
+
+	/**
+	 * Returns if a break fragment write break is set for the selected entity.
+	 */
+	static MASSENTITY_API bool IsFragmentWriteBreakForSelectedEntitySet(const FMassEntityManager& EntityManager, TNotNull<const UScriptStruct*> FragmentType);
 
 	/**
 	 * Clears a breakpoint triggered on processor execute for an entity.
@@ -372,7 +390,8 @@ private:
 	static MASSENTITY_API TMap<FName, const UScriptStruct*> FragmentsByName;
 
 	static MASSENTITY_API void UpdateHasBreakpoint();
-	static MASSENTITY_API FEnvironment& GetActiveEnvironment(const FMassEntityManager& EntityManager);
+	static MASSENTITY_API FEnvironment& GetActiveEnvironmentChecked(const FMassEntityManager& EntityManager);
+	static MASSENTITY_API FEnvironment* GetActiveEnvironment(const FMassEntityManager& EntityManager);
 };
 
 #else
