[Mass] added a way to as a MassCommandBuffer about its host thread.

diff --git a/Engine/Source/Runtime/MassEntity/Public/MassCommandBuffer.h b/Engine/Source/Runtime/MassEntity/Public/MassCommandBuffer.h
--- a/Engine/Source/Runtime/MassEntity/Public/MassCommandBuffer.h
+++ b/Engine/Source/Runtime/MassEntity/Public/MassCommandBuffer.h
@@ -182,6 +182,11 @@ public:
 		CleanUp();
 	}
 
+	bool IsInOwnerThread() const
+	{
+		return OwnerThreadId == FPlatformTLS::GetCurrentThreadId();
+	}
+
 private:
 	friend FMassEntityManager;
 
