[Mass] Added TRACE_CPUPROFILER_EVENT_SCOPE_STR to FMassObserverManager::OnPreEntitiesDestroyed

diff --git a/Engine/Source/Runtime/MassEntity/Private/MassObserverManager.cpp b/Engine/Source/Runtime/MassEntity/Private/MassObserverManager.cpp
--- a/Engine/Source/Runtime/MassEntity/Private/MassObserverManager.cpp
+++ b/Engine/Source/Runtime/MassEntity/Private/MassObserverManager.cpp
@@ -315,12 +315,14 @@ bool FMassObserverManager::OnPostEntityCreated(FMassEntityHandle EntityHandle, c
 
 bool FMassObserverManager::OnPreEntitiesDestroyed(const FMassArchetypeEntityCollection& EntityCollection)
 {
+	TRACE_CPUPROFILER_EVENT_SCOPE_STR("OnPreEntitiesDestroyed")
 	const FMassArchetypeCompositionDescriptor& ArchetypeComposition = EntityManager.GetArchetypeComposition(EntityCollection.GetArchetype());
 	return OnCompositionChanged(EntityCollection, ArchetypeComposition, EMassObservedOperation::Remove);
 }
 
 bool FMassObserverManager::OnPreEntitiesDestroyed(UE::Mass::FProcessingContext& ProcessingContext, const FMassArchetypeEntityCollection& EntityCollection)
 {
+	TRACE_CPUPROFILER_EVENT_SCOPE_STR("OnPreEntitiesDestroyed")
 	const FMassArchetypeCompositionDescriptor& ArchetypeComposition = EntityManager.GetArchetypeComposition(EntityCollection.GetArchetype());	
 	return OnCompositionChanged(EntityCollection, ArchetypeComposition, EMassObservedOperation::Remove, &ProcessingContext);
 }
