Fix issue where USTRUCT was inheriting from multiple USTRUCTs.  The change removes the USTRUCT from the second base struct since it appears it isn't needed.

#rnx

diff --git a/Engine/Source/Runtime/MassEntity/Public/MassRequirements.h b/Engine/Source/Runtime/MassEntity/Public/MassRequirements.h
--- a/Engine/Source/Runtime/MassEntity/Public/MassRequirements.h
+++ b/Engine/Source/Runtime/MassEntity/Public/MassRequirements.h
@@ -68,10 +68,8 @@ struct FMassFragmentRequirementDescription
 /**
  *  FMassSubsystemRequirements is a structure that declares runtime subsystem access type given calculations require.
  */
-USTRUCT()
 struct FMassSubsystemRequirements
 {
-	GENERATED_BODY()
 
 	friend FMassDebugger;
 	friend FMassRequirementAccessDetector;
