Add support for Mass observers initializing outside the disregard for GC window

#rnx

diff --git a/Engine/Source/Runtime/MassEntity/Private/MassEntitySettings.cpp b/Engine/Source/Runtime/MassEntity/Private/MassEntitySettings.cpp
--- a/Engine/Source/Runtime/MassEntity/Private/MassEntitySettings.cpp
+++ b/Engine/Source/Runtime/MassEntity/Private/MassEntitySettings.cpp
@@ -112,6 +112,13 @@ void UMassEntitySettings::BuildProcessorList()
 #endif // WITH_EDITOR
 		)
 		{
+			// Observers might register later than the GC disregard window, causing a GC mismatch between this early-initialized class and them.
+			const bool bIsDisregardForGC = GUObjectArray.IsDisregardForGC(this);
+			if (bIsDisregardForGC && !(GUObjectArray.IsDisregardForGC(ProcessorCDO) || ProcessorCDO->HasAnyFlags(RF_MarkAsRootSet)))
+			{
+				ProcessorCDO->AddToRoot();
+			}
+
 			ProcessorCDOs.Add(ProcessorCDO);
 			if (ProcessorCDO->ShouldAutoAddToGlobalList())
 			{
diff --git a/Engine/Source/Runtime/MassEntity/Private/MassObserverRegistry.cpp b/Engine/Source/Runtime/MassEntity/Private/MassObserverRegistry.cpp
--- a/Engine/Source/Runtime/MassEntity/Private/MassObserverRegistry.cpp
+++ b/Engine/Source/Runtime/MassEntity/Private/MassObserverRegistry.cpp
@@ -19,6 +19,13 @@ void UMassObserverRegistry::RegisterObserver(const UScriptStruct& ObservedType,
 	const bool bIsFragment = UE::Mass::IsA<FMassFragment>(&ObservedType);
 	checkSlow(bIsFragment || UE::Mass::IsA<FMassTag>(&ObservedType));
 
+	// Observers might register later than the GC disregard window, causing a GC mismatch between this early-initialized class and them.
+	const bool bIsDisregardForGC = GUObjectArray.IsDisregardForGC(this);
+	if (bIsDisregardForGC && !(GUObjectArray.IsDisregardForGC(&ObservedType) || ObservedType.HasAnyFlags(RF_MarkAsRootSet)))
+	{
+		const_cast<UScriptStruct&>(ObservedType).AddToRoot();
+	}
+
 	if (bIsFragment)
 	{
 		(*FragmentObservers[(uint8)Operation]).FindOrAdd(&ObservedType).ClassCollection.AddUnique(ObserverClass);
