[MassEntity]
Properly handle shared fragment hash collisions.

diff --git a/Engine/Source/Developer/MassEntityTestSuite/Private/MassSharedFragmentsTest.cpp b/Engine/Source/Developer/MassEntityTestSuite/Private/MassSharedFragmentsTest.cpp
--- a/Engine/Source/Developer/MassEntityTestSuite/Private/MassSharedFragmentsTest.cpp
+++ b/Engine/Source/Developer/MassEntityTestSuite/Private/MassSharedFragmentsTest.cpp
@@ -864,6 +864,24 @@ struct FSharedFragmentValues_GetOrCreateNoArgs : FEntityTestBase
 };
 IMPLEMENT_AI_INSTANT_TEST(FSharedFragmentValues_GetOrCreateNoArgs, "System.Mass.SharedFragments.GetOrCreate.NoArgs");
 
+struct FSharedFragmentValues_GetOrCreateArray : FEntityTestBase
+{
+	virtual bool InstantTest() override
+	{
+		const FSharedStruct SharedFragment1 = EntityManager->GetOrCreateSharedFragment<FTestSharedFragment_Array>(TArray<int32>{ 1, 2, 3 });
+		const FSharedStruct SharedFragment2 = EntityManager->GetOrCreateSharedFragment<FTestSharedFragment_Array>(TArray<int32>{ 1, 2, 3 });
+		FTestSharedFragment_Array TestFragment;
+		TestFragment.Value = { 1,2,3 };
+		const FSharedStruct SharedFragment3 = EntityManager->GetOrCreateSharedFragment(TestFragment);
+
+		AITEST_EQUAL("Shared fragments created for same struct type using same TArray contents should share memory", SharedFragment1, SharedFragment2);
+		AITEST_EQUAL("Shared fragments created for same struct type using same TArray contents should share memory", SharedFragment1, SharedFragment3);
+
+		return true;
+	}
+};
+IMPLEMENT_AI_INSTANT_TEST(FSharedFragmentValues_GetOrCreateArray, "System.Mass.SharedFragments.GetOrCreate.WithArray");
+
 struct FSharedFragmentValues_GetOrCreateConstNoArgs : FEntityTestBase
 {
 	virtual bool InstantTest() override
diff --git a/Engine/Source/Developer/MassEntityTestSuite/Public/MassEntityTestTypes.h b/Engine/Source/Developer/MassEntityTestSuite/Public/MassEntityTestTypes.h
--- a/Engine/Source/Developer/MassEntityTestSuite/Public/MassEntityTestTypes.h
+++ b/Engine/Source/Developer/MassEntityTestSuite/Public/MassEntityTestTypes.h
@@ -123,6 +123,20 @@ struct FTestSharedFragment_Int : public FMassSharedFragment
 	FTestSharedFragment_Int(const int32 InValue = 0) : Value(InValue) {}
 };
 
+USTRUCT()
+struct FTestSharedFragment_Array : public FMassSharedFragment
+{
+	using FValueType = TArray<int32>;
+
+	GENERATED_BODY()
+
+	UPROPERTY()
+	TArray<int32> Value;
+
+	FTestSharedFragment_Array() = default;
+	FTestSharedFragment_Array(const TArray<int32>&& InValue) : Value(InValue) {}
+};
+
 template<>
 struct TMassSharedFragmentTraits<FTestSharedFragment_Int> final
 {
diff --git a/Engine/Source/Runtime/MassEntity/Public/MassEntityManager.h b/Engine/Source/Runtime/MassEntity/Public/MassEntityManager.h
--- a/Engine/Source/Runtime/MassEntity/Public/MassEntityManager.h
+++ b/Engine/Source/Runtime/MassEntity/Public/MassEntityManager.h
@@ -1,4 +1,4 @@
-// Copyright Epic Games, Inc. All Rights Reserved.
+ // Copyright Epic Games, Inc. All Rights Reserved.
 
 #pragma once
 
@@ -1164,22 +1164,41 @@ private:
 
 		TSharedStructType& FindOrAdd(const uint32 Hash, const UScriptStruct* Type, const uint8* Data)
 		{
-			int32& Index = HashToInstanceIndexMap.FindOrAddByHash(Hash, Hash, INDEX_NONE);
-			if (Index == INDEX_NONE)
+			for (TMultiMap<uint32, int32>::TConstKeyIterator It = HashToInstanceIndexMap.CreateConstKeyIterator(Hash); It; ++It)
 			{
-				Index = Add(TSharedStructType::Make(Type, Data));
+				TSharedStructType& Instance = Instances[It.Value()];
+
+				if (Instance.GetScriptStruct() == Type
+					&& Type->CompareScriptStruct(Instance.GetMemory(), Data, PPF_None))
+				{
+					return Instance;
+				}
 			}
+
+			int32 Index = Add(TSharedStructType::Make(Type, Data));
+			HashToInstanceIndexMap.Add(Hash, Index);
 			return Instances[Index];
 		}
 
 		template<typename T, typename... TArgs>
 		TSharedStructType& FindOrAdd(const uint32 Hash, TArgs&&... InArgs)
 		{
-			int32& Index = HashToInstanceIndexMap.FindOrAddByHash(Hash, Hash, INDEX_NONE);
-			if (Index == INDEX_NONE)
+			// Need to actually construct the struct to make proper comparison to possible existing instance
+			TSharedStructType TempInstance = Make<T>(Forward<TArgs>(InArgs)...);
+
+			for (TMultiMap<uint32, int32>::TConstKeyIterator It = HashToInstanceIndexMap.CreateConstKeyIterator(Hash); It; ++It)
 			{
-				Index = Add(TSharedStructType::template Make<T>(Forward<TArgs>(InArgs)...));
+				TSharedStructType& Instance = Instances[It.Value()];
+
+				if (Instance.GetScriptStruct() == T::StaticStruct()
+					&& T::StaticStruct()->CompareScriptStruct(Instance.GetMemory(), TempInstance.GetMemory(), PPF_None))
+				{
+					return Instance;
+				}
 			}
+
+			int32 Index = Add(MoveTemp(TempInstance));
+			HashToInstanceIndexMap.Add(Hash, Index);
 			return Instances[Index];
 		}
 
@@ -1206,7 +1225,7 @@ private:
 
 		TArray<TSharedStructType> Instances;
 		// Hash/Index in array pair
-		TMap<uint32, int32> HashToInstanceIndexMap;
+		TMultiMap<uint32, int32> HashToInstanceIndexMap;
 		// Maps specific struct type to a collection of FSharedStruct instances of that type
 		TMap<const UScriptStruct*, TArray<TSharedStructType>> TypeToInstanceMap;
 	};
