[Mass] Refactored UMassObserverRegistry to be more robust against types being unloaded. It no longer relies on TSubclassOf<UMassProcessor> and TObjectPtr<const UScriptStruct>, and instead uses FSoftClassPath and TObjectKey>const UScriptStruct> respectively.

#rnx

diff --git a/Engine/Source/Runtime/MassEntity/Private/MassObserverManager.cpp b/Engine/Source/Runtime/MassEntity/Private/MassObserverManager.cpp
--- a/Engine/Source/Runtime/MassEntity/Private/MassObserverManager.cpp
+++ b/Engine/Source/Runtime/MassEntity/Private/MassObserverManager.cpp
@@ -34,25 +34,31 @@ namespace UE::Mass::ObserverManager
 		};
 
 		// a helper function to reduce code duplication in FMassObserverManager::Initialize
-		template<typename TBitSet, typename TPointerType>
+		template<typename TBitSet>
 		void AddRegisteredObserverProcessorInstances(FMassEntityManager& EntityManager, const EProcessorExecutionFlags WorldExecutionFlags, UObject& Owner
-			, const TMap<TPointerType, FMassProcessorClassCollection>& RegisteredObserverTypes, TBitSet& InOutObservedBitSet, FMassObserversMap& Observers)
+			, const UMassObserverRegistry::FObserverClassesMap& RegisteredObserverTypes, TBitSet& InOutObservedBitSet, FMassObserversMap& Observers)
 		{
 			for (auto It : RegisteredObserverTypes)
 			{
-				if (It.Value.ClassCollection.Num() == 0)
+				const UObject* Type = It.Key.ResolveObjectPtr();
+				
+				if (Type == nullptr || It.Value.Num() == 0)
 				{
 					continue;
 				}
 
-				InOutObservedBitSet.Add(*It.Key);
-				FMassRuntimePipeline& Pipeline = (*Observers).FindOrAdd(It.Key);
+				const UScriptStruct* StructType = CastChecked<const UScriptStruct>(Type);
+				InOutObservedBitSet.Add(*StructType);
+				FMassRuntimePipeline& Pipeline = (*Observers).FindOrAdd(StructType);
 
-				for (const TSubclassOf<UMassProcessor>& ProcessorClass : It.Value.ClassCollection)
+				for (const FSoftClassPath& SoftProcessorClass : It.Value)
 				{
-					if (ProcessorClass->GetDefaultObject<UMassProcessor>()->ShouldExecute(WorldExecutionFlags))
+					if (TSubclassOf<UMassProcessor> ProcessorClass = SoftProcessorClass.ResolveClass())
 					{
-						Pipeline.AppendProcessor(ProcessorClass, Owner);
+						if (ProcessorClass->GetDefaultObject<UMassProcessor>()->ShouldExecute(WorldExecutionFlags))
+						{
+							Pipeline.AppendProcessor(ProcessorClass, Owner);
+						}
 					}
 				}
 				Pipeline.Initialize(Owner, EntityManager.AsShared());
@@ -251,8 +257,8 @@ void FMassObserverManager::Initialize()
 	using UE::Mass::ObserverManager::Private::AddRegisteredObserverProcessorInstances;
 	for (int i = 0; i < (int)EMassObservedOperation::MAX; ++i)
 	{
-		AddRegisteredObserverProcessorInstances(EntityManager, WorldExecutionFlags, *Owner, *Registry.FragmentObservers[i], ObservedFragments[i], FragmentObservers[i]);
-		AddRegisteredObserverProcessorInstances(EntityManager, WorldExecutionFlags, *Owner, *Registry.TagObservers[i], ObservedTags[i], TagObservers[i]);
+		AddRegisteredObserverProcessorInstances(EntityManager, WorldExecutionFlags, *Owner, Registry.FragmentObserverMaps[i], ObservedFragments[i], FragmentObservers[i]);
+		AddRegisteredObserverProcessorInstances(EntityManager, WorldExecutionFlags, *Owner, Registry.TagObserverMaps[i], ObservedTags[i], TagObservers[i]);
 	}
 
 #if WITH_MASSENTITY_DEBUG
diff --git a/Engine/Source/Runtime/MassEntity/Private/MassObserverRegistry.cpp b/Engine/Source/Runtime/MassEntity/Private/MassObserverRegistry.cpp
--- a/Engine/Source/Runtime/MassEntity/Private/MassObserverRegistry.cpp
+++ b/Engine/Source/Runtime/MassEntity/Private/MassObserverRegistry.cpp
@@ -19,21 +19,12 @@ void UMassObserverRegistry::RegisterObserver(const UScriptStruct& ObservedType,
 	const bool bIsFragment = UE::Mass::IsA<FMassFragment>(&ObservedType);
 	checkSlow(bIsFragment || UE::Mass::IsA<FMassTag>(&ObservedType));
 
-	// Observers might register later than the GC disregard window, causing a GC mismatch between this early-initialized class and them.
-	const bool bIsDisregardForGC = GUObjectArray.IsDisregardForGC(this);
-	if (bIsDisregardForGC && !(GUObjectArray.IsDisregardForGC(&ObservedType) || ObservedType.HasAnyFlags(RF_MarkAsRootSet)))
-	{
-		const_cast<UScriptStruct&>(ObservedType).AddToRoot();
-	}
-
 	if (bIsFragment)
 	{
-		(*FragmentObservers[(uint8)Operation]).FindOrAdd(&ObservedType).ClassCollection.AddUnique(ObserverClass);
+		FragmentObserverMaps[static_cast<uint8>(Operation)].FindOrAdd(&ObservedType).AddUnique(FSoftClassPath(ObserverClass));
 	}
 	else
 	{
-		(*TagObservers[(uint8)Operation]).FindOrAdd(&ObservedType).ClassCollection.AddUnique(ObserverClass);
+		TagObserverMaps[static_cast<uint8>(Operation)].FindOrAdd(&ObservedType).AddUnique(FSoftClassPath(ObserverClass));
 	}
 }
-
-
diff --git a/Engine/Source/Runtime/MassEntity/Public/MassObserverRegistry.h b/Engine/Source/Runtime/MassEntity/Public/MassObserverRegistry.h
--- a/Engine/Source/Runtime/MassEntity/Public/MassObserverRegistry.h
+++ b/Engine/Source/Runtime/MassEntity/Public/MassObserverRegistry.h
@@ -4,6 +4,7 @@
 
 #include "MassEntityTypes.h"
 #include "MassProcessor.h"
+#include "UObject/ObjectKey.h"
 #include "MassObserverRegistry.generated.h"
 
 #define UE_API MASSENTITY_API
@@ -14,6 +15,7 @@ struct FMassObserverManager;
 /**
  * A wrapper type for a TArray to support having map-of-arrays UPROPERTY members in FMassEntityObserverClassesMap
  */
+//USTRUCT(meta = (Deprecated = "5.7", DepracationMessage = "FMassProcessorClassCollection is no longer being used and will be removed in the upcoming engine released"))
 USTRUCT()
 struct FMassProcessorClassCollection
 {
@@ -26,6 +28,7 @@ struct FMassProcessorClassCollection
 /**
  * A wrapper type for a TMap to support having array-of-maps UPROPERTY members in UMassObserverRegistry
  */
+//USTRUCT(meta = (Deprecated = "5.7", DepracationMessage = "FMassEntityObserverClassesMap is no longer being used and will be removed in the upcoming engine released"))
 USTRUCT()
 struct FMassEntityObserverClassesMap
 {
@@ -60,12 +63,17 @@ public:
 
 	UE_API void RegisterObserver(const UScriptStruct& ObservedType, const EMassObservedOperation Operation, TSubclassOf<UMassProcessor> ObserverClass);
 
+	using FObserverClassesMap = TMap<TObjectKey<const UScriptStruct>, TArray<FSoftClassPath>>;
 protected:
 	friend FMassObserverManager;
+	FObserverClassesMap FragmentObserverMaps[static_cast<uint8>(EMassObservedOperation::MAX)];
+	FObserverClassesMap TagObserverMaps[static_cast<uint8>(EMassObservedOperation::MAX)];
 
+	UE_DEPRECATED(5.7, "Use FragmentObserverMaps instead")
 	UPROPERTY()
 	FMassEntityObserverClassesMap FragmentObservers[(uint8)EMassObservedOperation::MAX];
 
+	UE_DEPRECATED(5.7, "Use TagObserverMaps instead")
 	UPROPERTY()
 	FMassEntityObserverClassesMap TagObservers[(uint8)EMassObservedOperation::MAX];
 };
