[Mass] Marked FMassProcessorClassCollection and FMassEntityObserverClassesMap as deprecated, which was supposed to be a part of the previous CL

#rnx

diff --git a/Engine/Source/Runtime/MassEntity/Public/MassObserverRegistry.h b/Engine/Source/Runtime/MassEntity/Public/MassObserverRegistry.h
--- a/Engine/Source/Runtime/MassEntity/Public/MassObserverRegistry.h
+++ b/Engine/Source/Runtime/MassEntity/Public/MassObserverRegistry.h
@@ -15,8 +15,7 @@ struct FMassObserverManager;
 /**
  * A wrapper type for a TArray to support having map-of-arrays UPROPERTY members in FMassEntityObserverClassesMap
  */
-//USTRUCT(meta = (Deprecated = "5.7", DepracationMessage = "FMassProcessorClassCollection is no longer being used and will be removed in the upcoming engine released"))
-USTRUCT()
+USTRUCT(meta = (Deprecated = "5.7", DepracationMessage = "FMassProcessorClassCollection is no longer being used and will be removed in the upcoming engine released"))
 struct FMassProcessorClassCollection
 {
 	GENERATED_BODY()
@@ -28,8 +27,7 @@ struct FMassProcessorClassCollection
 /**
  * A wrapper type for a TMap to support having array-of-maps UPROPERTY members in UMassObserverRegistry
  */
-//USTRUCT(meta = (Deprecated = "5.7", DepracationMessage = "FMassEntityObserverClassesMap is no longer being used and will be removed in the upcoming engine released"))
-USTRUCT()
+USTRUCT(meta = (Deprecated = "5.7", DepracationMessage = "FMassEntityObserverClassesMap is no longer being used and will be removed in the upcoming engine released"))
 struct FMassEntityObserverClassesMap
 {
 	GENERATED_BODY()
