[Mass] added a flavor of FMassRuntimePipeline.AppendProcessors that lets caller provide a movable TArray of processor pointers.

#rnx

diff --git a/Engine/Source/Runtime/MassEntity/Private/MassProcessingTypes.cpp b/Engine/Source/Runtime/MassEntity/Private/MassProcessingTypes.cpp
--- a/Engine/Source/Runtime/MassEntity/Private/MassProcessingTypes.cpp
+++ b/Engine/Source/Runtime/MassEntity/Private/MassProcessingTypes.cpp
@@ -212,6 +212,18 @@ void FMassRuntimePipeline::AppendProcessors(TArrayView<TObjectPtr<UMassProcessor
 	Processors.Append(InProcessors);
 }
 
+void FMassRuntimePipeline::AppendProcessors(TArray<TObjectPtr<UMassProcessor>>&& InProcessors)
+{
+	if (Processors.Num())
+	{
+		Processors.Append(MoveTemp(InProcessors));
+	}
+	else
+	{
+		Processors = MoveTemp(InProcessors);
+	}
+}
+
 bool FMassRuntimePipeline::AppendUniqueProcessor(UMassProcessor& Processor)
 {
 	const int32 PreviousCount = Processors.Num();
diff --git a/Engine/Source/Runtime/MassEntity/Public/MassProcessingTypes.h b/Engine/Source/Runtime/MassEntity/Public/MassProcessingTypes.h
--- a/Engine/Source/Runtime/MassEntity/Public/MassProcessingTypes.h
+++ b/Engine/Source/Runtime/MassEntity/Public/MassProcessingTypes.h
@@ -117,6 +117,7 @@ public:
 	/** Adds InProcessor(s) to Processors without any additional checks */
 	UE_API void AppendProcessor(UMassProcessor& InProcessor);
 	UE_API void AppendProcessors(TArrayView<TObjectPtr<UMassProcessor>> InProcessors);
+	UE_API void AppendProcessors(TArray<TObjectPtr<UMassProcessor>>&& InProcessors);
 
 	/** @return true if the given processor has indeed been added, i.e. will return false if Processor was already part of the pipeline. */
 	UE_API bool AppendUniqueProcessor(UMassProcessor& Processor);
