[Mass] Mass Observers can now differentiate between Add and Created operations as well as between Remove and Destroy operations.
* Observer processors can now be declared to observe multiple operations, and the same processor instance will be used for all the declared operations.
* With the extension of EMassObservedOperation the existing UE::Mass::ObserverManager::EObservedOperationNotification became redundant and got removed.
* The observer processor can now find out which operation type is being handled. This was achieved by changing the type of payload passed to the execution context when observers are executed.
* new payload contains information about the current type being handled, as well as what type of operation is being handled, and a way to peek at the other types being handled in the same operation.
* a minor consequence of improvements to UE::Mass::ObserverManager::Private::AddRegisteredObserverProcessorInstances we now consider an element "observed" only if there are actually any observer instances created for that operation (i.e., we don't if all the candidate processor classes filtered out, due to, for example, world flags).

Update notes:
* ObserverProcessors checking the type currently being handled by looking at the type of data stored by ExecutionContext.AuxData need to update their code. The type can now be obtained by calling ExecutionContext.GetAuxData().Get<FMassObserverExecutionContext>().GetCurrentType()

diff --git a/Engine/Plugins/AI/MassAI/Source/MassAIBehavior/Private/MassLookAtProcessors.cpp b/Engine/Plugins/AI/MassAI/Source/MassAIBehavior/Private/MassLookAtProcessors.cpp
--- a/Engine/Plugins/AI/MassAI/Source/MassAIBehavior/Private/MassLookAtProcessors.cpp
+++ b/Engine/Plugins/AI/MassAI/Source/MassAIBehavior/Private/MassLookAtProcessors.cpp
@@ -667,7 +667,7 @@ UMassLookAtTargetRemoverProcessor::UMassLookAtTargetRemoverProcessor()
 	: Query(*this)
 {
 	ObservedType = FMassLookAtTargetFragment::StaticStruct();
-	Operation = EMassObservedOperation::Remove;
+	ObservedOperations = EMassObservedOperationFlags::Remove;
 	ExecutionFlags = static_cast<int32>(EProcessorExecutionFlags::AllNetModes);
 }
 
@@ -698,7 +698,7 @@ UMassLookAtRequestInitializer::UMassLookAtRequestInitializer()
 	: Query(*this)
 {
 	ObservedType = FMassLookAtRequestFragment::StaticStruct();
-	Operation = EMassObservedOperation::Add;
+	ObservedOperations = EMassObservedOperationFlags::Add;
 	ExecutionFlags = static_cast<int32>(EProcessorExecutionFlags::Client | EProcessorExecutionFlags::Standalone);
 }
 
@@ -734,7 +734,7 @@ UMassLookAtRequestDeinitializer::UMassLookAtRequestDeinitializer()
 	: Query(*this)
 {
 	ObservedType = FMassLookAtRequestFragment::StaticStruct();
-	Operation = EMassObservedOperation::Remove;
+	ObservedOperations = EMassObservedOperationFlags::Remove;
 	ExecutionFlags = static_cast<int32>(EProcessorExecutionFlags::Client | EProcessorExecutionFlags::Standalone);
 }
 
diff --git a/Engine/Plugins/AI/MassAI/Source/MassAIBehavior/Private/MassStateTreeProcessors.cpp b/Engine/Plugins/AI/MassAI/Source/MassAIBehavior/Private/MassStateTreeProcessors.cpp
--- a/Engine/Plugins/AI/MassAI/Source/MassAIBehavior/Private/MassStateTreeProcessors.cpp
+++ b/Engine/Plugins/AI/MassAI/Source/MassAIBehavior/Private/MassStateTreeProcessors.cpp
@@ -97,7 +97,7 @@ UMassStateTreeFragmentDestructor::UMassStateTreeFragmentDestructor()
 {
 	ExecutionFlags = (int32)(EProcessorExecutionFlags::Standalone | EProcessorExecutionFlags::Server);
 	ObservedType = FMassStateTreeInstanceFragment::StaticStruct();
-	Operation = EMassObservedOperation::Remove;
+	ObservedOperations = EMassObservedOperationFlags::Remove;
 	bRequiresGameThreadExecution = true;
 }
 
diff --git a/Engine/Plugins/AI/MassAI/Source/MassAIBehavior/Private/MassZoneGraphAnnotationProcessors.cpp b/Engine/Plugins/AI/MassAI/Source/MassAIBehavior/Private/MassZoneGraphAnnotationProcessors.cpp
--- a/Engine/Plugins/AI/MassAI/Source/MassAIBehavior/Private/MassZoneGraphAnnotationProcessors.cpp
+++ b/Engine/Plugins/AI/MassAI/Source/MassAIBehavior/Private/MassZoneGraphAnnotationProcessors.cpp
@@ -19,7 +19,7 @@ UMassZoneGraphAnnotationTagsInitializer::UMassZoneGraphAnnotationTagsInitializer
 	: EntityQuery(*this)
 {
 	ObservedType = FMassZoneGraphAnnotationFragment::StaticStruct();
-	Operation = EMassObservedOperation::Add;
+	ObservedOperations = EMassObservedOperationFlags::Add;
 }
 
 void UMassZoneGraphAnnotationTagsInitializer::ConfigureQueries(const TSharedRef<FMassEntityManager>& EntityManager)
diff --git a/Engine/Plugins/AI/MassAI/Source/MassNavigation/Private/MassNavigationProcessors.cpp b/Engine/Plugins/AI/MassAI/Source/MassNavigation/Private/MassNavigationProcessors.cpp
--- a/Engine/Plugins/AI/MassAI/Source/MassNavigation/Private/MassNavigationProcessors.cpp
+++ b/Engine/Plugins/AI/MassAI/Source/MassNavigation/Private/MassNavigationProcessors.cpp
@@ -122,7 +122,7 @@ UMassMoveTargetFragmentInitializer::UMassMoveTargetFragmentInitializer()
 	: InitializerQuery(*this)
 {
 	ObservedType = FMassMoveTargetFragment::StaticStruct();
-	Operation = EMassObservedOperation::Add;
+	ObservedOperations = EMassObservedOperationFlags::Add;
 }
 
 void UMassMoveTargetFragmentInitializer::ConfigureQueries(const TSharedRef<FMassEntityManager>& EntityManager)
@@ -274,7 +274,7 @@ UMassNavigationObstacleRemoverProcessor::UMassNavigationObstacleRemoverProcessor
 	: EntityQuery(*this)
 {
 	ObservedType = FMassNavigationObstacleGridCellLocationFragment::StaticStruct();
-	Operation = EMassObservedOperation::Remove;
+	ObservedOperations = EMassObservedOperationFlags::Remove;
 	ExecutionFlags = (int32)EProcessorExecutionFlags::AllNetModes;
 }
 
diff --git a/Engine/Plugins/AI/MassAI/Source/MassZoneGraphNavigation/Private/MassZoneGraphNavigationProcessors.cpp b/Engine/Plugins/AI/MassAI/Source/MassZoneGraphNavigation/Private/MassZoneGraphNavigationProcessors.cpp
--- a/Engine/Plugins/AI/MassAI/Source/MassZoneGraphNavigation/Private/MassZoneGraphNavigationProcessors.cpp
+++ b/Engine/Plugins/AI/MassAI/Source/MassZoneGraphNavigation/Private/MassZoneGraphNavigationProcessors.cpp
@@ -27,7 +27,7 @@ UMassZoneGraphLocationInitializer::UMassZoneGraphLocationInitializer()
 	: EntityQuery(*this)
 {
 	ObservedType = FMassZoneGraphLaneLocationFragment::StaticStruct();
-	Operation = EMassObservedOperation::Add;
+	ObservedOperations = EMassObservedOperationFlags::Add;
 	ExecutionFlags = (int32)EProcessorExecutionFlags::AllNetModes;
 }
 
diff --git a/Engine/Plugins/AI/MassCrowd/Source/MassCrowd/Private/MassCrowdNavigationProcessor.cpp b/Engine/Plugins/AI/MassCrowd/Source/MassCrowd/Private/MassCrowdNavigationProcessor.cpp
--- a/Engine/Plugins/AI/MassCrowd/Source/MassCrowd/Private/MassCrowdNavigationProcessor.cpp
+++ b/Engine/Plugins/AI/MassCrowd/Source/MassCrowd/Private/MassCrowdNavigationProcessor.cpp
@@ -68,7 +68,7 @@ UMassCrowdLaneTrackingDestructor::UMassCrowdLaneTrackingDestructor()
 {
 	ExecutionFlags = (int32)(EProcessorExecutionFlags::Standalone | EProcessorExecutionFlags::Server);
 	ObservedType = FMassCrowdLaneTrackingFragment::StaticStruct();
-	Operation = EMassObservedOperation::Remove;
+	ObservedOperations = EMassObservedOperationFlags::Remove;
 }
 
 void UMassCrowdLaneTrackingDestructor::ConfigureQueries(const TSharedRef<FMassEntityManager>& EntityManager)
@@ -227,7 +227,7 @@ UMassCrowdDynamicObstacleInitializer::UMassCrowdDynamicObstacleInitializer()
 {
 	ExecutionFlags = (int32)(EProcessorExecutionFlags::Standalone | EProcessorExecutionFlags::Server);
 	ObservedType = FMassCrowdObstacleFragment::StaticStruct();
-	Operation = EMassObservedOperation::Add;
+	ObservedOperations = EMassObservedOperationFlags::Add;
 }
 
 void UMassCrowdDynamicObstacleInitializer::ConfigureQueries(const TSharedRef<FMassEntityManager>& EntityManager)
@@ -266,7 +266,7 @@ UMassCrowdDynamicObstacleDeinitializer::UMassCrowdDynamicObstacleDeinitializer()
 {
 	ExecutionFlags = (int32)(EProcessorExecutionFlags::Standalone | EProcessorExecutionFlags::Server);
 	ObservedType = FMassCrowdObstacleFragment::StaticStruct();
-	Operation = EMassObservedOperation::Remove;
+	ObservedOperations = EMassObservedOperationFlags::Remove;
 }
 
 void UMassCrowdDynamicObstacleDeinitializer::InitializeInternal(UObject& Owner, const TSharedRef<FMassEntityManager>& EntityManager)
diff --git a/Engine/Plugins/Experimental/EditorDataStorage/Source/TedsCore/Private/Processors/TypedElementProcessorAdaptors.cpp b/Engine/Plugins/Experimental/EditorDataStorage/Source/TedsCore/Private/Processors/TypedElementProcessorAdaptors.cpp
--- a/Engine/Plugins/Experimental/EditorDataStorage/Source/TedsCore/Private/Processors/TypedElementProcessorAdaptors.cpp
+++ b/Engine/Plugins/Experimental/EditorDataStorage/Source/TedsCore/Private/Processors/TypedElementProcessorAdaptors.cpp
@@ -1862,9 +1862,9 @@ const UScriptStruct* UTypedElementQueryObserverCallbackAdapterProcessorBase::Get
 	return ObservedType;
 }
 
-EMassObservedOperation UTypedElementQueryObserverCallbackAdapterProcessorBase::GetObservedOperation() const
+EMassObservedOperationFlags UTypedElementQueryObserverCallbackAdapterProcessorBase::GetObservedOperations() const
 {
-	return Operation;
+	return static_cast<EMassObservedOperationFlags>(ObservedOperations);
 }
 
 bool UTypedElementQueryObserverCallbackAdapterProcessorBase::ConfigureQueryCallback(
@@ -1894,10 +1894,10 @@ bool UTypedElementQueryObserverCallbackAdapterProcessorBase::ConfigureQueryCallb
 	switch (Query.Description.Callback.Type)
 	{
 	case EQueryCallbackType::ObserveAdd:
-		Operation = EMassObservedOperation::Add;
+		ObservedOperations = EMassObservedOperationFlags::Add;
 		break;
 	case EQueryCallbackType::ObserveRemove:
-		Operation = EMassObservedOperation::Remove;
+		ObservedOperations = EMassObservedOperationFlags::Remove;
 		break;
 	default:
 		checkf(false, TEXT("Query type %i is not supported from the observer processor adapter."),
@@ -1933,15 +1933,19 @@ void UTypedElementQueryObserverCallbackAdapterProcessorBase::DebugOutputDescript
 {
 #if WITH_MASSENTITY_DEBUG
 	UMassObserverProcessor::DebugOutputDescription(Ar, Indent);
-	EMassObservedOperation ObservationType = GetObservedOperation();
-	if (ObservationType == EMassObservedOperation::Add)
-	{
-		Ar.Logf(TEXT("\n%*sType: Editor Add Observer"), Indent, TEXT(""));
-	}
-	else if (ObservationType == EMassObservedOperation::Remove)
+	const uint8 ObservationTypes = static_cast<uint8>(GetObservedOperations());
+
+	if (ObservationTypes)
 	{
-		Ar.Logf(TEXT("\n%*sType: Editor Remove Observer"), Indent, TEXT(""));
-	}
+		for (uint8 OperationIndex = 0; OperationIndex < static_cast<uint8>(EMassObservedOperation::MAX); ++OperationIndex)
+		{
+			if ((ObservationTypes & (1 << OperationIndex)))
+			{
+				Ar.Logf(TEXT("\n%*sType: Editor %s Observer"), Indent, TEXT("")
+					, *UEnum::GetValueAsString(static_cast<EMassObservedOperation>(OperationIndex)));
+			}
+		}
+	}	
 	else
 	{
 		Ar.Logf(TEXT("\n%*sType: Editor <Unknown> Observer"), Indent, TEXT(""));
diff --git a/Engine/Plugins/Experimental/EditorDataStorage/Source/TedsCore/Private/Processors/TypedElementProcessorAdaptors.h b/Engine/Plugins/Experimental/EditorDataStorage/Source/TedsCore/Private/Processors/TypedElementProcessorAdaptors.h
--- a/Engine/Plugins/Experimental/EditorDataStorage/Source/TedsCore/Private/Processors/TypedElementProcessorAdaptors.h
+++ b/Engine/Plugins/Experimental/EditorDataStorage/Source/TedsCore/Private/Processors/TypedElementProcessorAdaptors.h
@@ -280,7 +280,7 @@ public:
 
 	FMassEntityQuery& GetQuery();
 	const UScriptStruct* GetObservedType() const;
-	EMassObservedOperation GetObservedOperation() const;
+	EMassObservedOperationFlags GetObservedOperations() const;
 	virtual bool ConfigureQueryCallback(
 		UE::Editor::DataStorage::FExtendedQuery& Query,
 		UE::Editor::DataStorage::FExtendedQueryStore::Handle QueryHandle,
diff --git a/Engine/Plugins/Experimental/EditorDataStorage/Source/TedsCore/Private/Queries/TypedElementExtendedQueryStore.cpp b/Engine/Plugins/Experimental/EditorDataStorage/Source/TedsCore/Private/Queries/TypedElementExtendedQueryStore.cpp
--- a/Engine/Plugins/Experimental/EditorDataStorage/Source/TedsCore/Private/Queries/TypedElementExtendedQueryStore.cpp
+++ b/Engine/Plugins/Experimental/EditorDataStorage/Source/TedsCore/Private/Queries/TypedElementExtendedQueryStore.cpp
@@ -1076,7 +1076,7 @@ namespace UE::Editor::DataStorage
 					static_cast<UTypedElementQueryObserverCallbackAdapterProcessorBase*>(StoredQuery.Processor.Get()))
 				{
 					Observer->ConfigureQueryCallback(StoredQuery, QueryHandle, *this, Environment);
-					EntityManager.GetObserverManager().AddObserverInstance(*Observer->GetObservedType(), Observer->GetObservedOperation(), *Observer);
+					EntityManager.GetObserverManager().AddObserverInstance(Observer->GetObservedType(), Observer->GetObservedOperations(), Observer);
 				}
 				else
 				{
@@ -1219,7 +1219,7 @@ namespace UE::Editor::DataStorage
 					static_cast<UTypedElementQueryObserverCallbackAdapterProcessorBase*>(QueryData.Processor.Get());
 				if (ensure(Observer))
 				{
-					EntityManager.GetObserverManager().RemoveObserverInstance(*Observer->GetObservedType(), Observer->GetObservedOperation(), *Observer);
+					EntityManager.GetObserverManager().RemoveObserverInstance(Observer->GetObservedType(), Observer->GetObservedOperations(), Observer);
 				}
 			}
 			else
diff --git a/Engine/Plugins/Runtime/MassGameplay/Source/MassGameplayDebug/Private/AssignDebugVisProcessor.cpp b/Engine/Plugins/Runtime/MassGameplay/Source/MassGameplayDebug/Private/AssignDebugVisProcessor.cpp
--- a/Engine/Plugins/Runtime/MassGameplay/Source/MassGameplayDebug/Private/AssignDebugVisProcessor.cpp
+++ b/Engine/Plugins/Runtime/MassGameplay/Source/MassGameplayDebug/Private/AssignDebugVisProcessor.cpp
@@ -15,7 +15,7 @@ UAssignDebugVisProcessor::UAssignDebugVisProcessor()
 	bAutoRegisterWithProcessingPhases = false;
 	bRequiresGameThreadExecution = true; // due to UMassDebuggerSubsystem
 	ObservedType = FSimDebugVisFragment::StaticStruct();
-	Operation = EMassObservedOperation::Add;
+	ObservedOperations = EMassObservedOperationFlags::Add;
 }
 
 void UAssignDebugVisProcessor::ConfigureQueries(const TSharedRef<FMassEntityManager>& EntityManager)
diff --git a/Engine/Plugins/Runtime/MassGameplay/Source/MassMovement/Private/MassVelocityRandomizerTrait.cpp b/Engine/Plugins/Runtime/MassGameplay/Source/MassMovement/Private/MassVelocityRandomizerTrait.cpp
--- a/Engine/Plugins/Runtime/MassGameplay/Source/MassMovement/Private/MassVelocityRandomizerTrait.cpp
+++ b/Engine/Plugins/Runtime/MassGameplay/Source/MassMovement/Private/MassVelocityRandomizerTrait.cpp
@@ -27,7 +27,7 @@ UMassRandomVelocityInitializer::UMassRandomVelocityInitializer()
 	: EntityQuery(*this)
 {
 	ObservedType = FMassVelocityFragment::StaticStruct();
-	Operation = EMassObservedOperation::Add;
+	ObservedOperations = EMassObservedOperationFlags::Add;
 	RandomStream.Initialize(UE::Mass::Utils::GenerateRandomSeed());
 }
 
diff --git a/Engine/Plugins/Runtime/MassGameplay/Source/MassReplication/Private/MassReplicationFragments.cpp b/Engine/Plugins/Runtime/MassGameplay/Source/MassReplication/Private/MassReplicationFragments.cpp
--- a/Engine/Plugins/Runtime/MassGameplay/Source/MassReplication/Private/MassReplicationFragments.cpp
+++ b/Engine/Plugins/Runtime/MassGameplay/Source/MassReplication/Private/MassReplicationFragments.cpp
@@ -14,7 +14,7 @@ UMassNetworkIDFragmentInitializer::UMassNetworkIDFragmentInitializer()
 {
 	ExecutionFlags = (int32)(EProcessorExecutionFlags::Standalone | EProcessorExecutionFlags::Server);
 	ObservedType = FMassNetworkIDFragment::StaticStruct();
-	Operation = EMassObservedOperation::Add;
+	ObservedOperations = EMassObservedOperationFlags::Add;
 }
 
 void UMassNetworkIDFragmentInitializer::ConfigureQueries(const TSharedRef<FMassEntityManager>& EntityManager)
diff --git a/Engine/Plugins/Runtime/MassGameplay/Source/MassReplication/Private/MassReplicationGridProcessor.cpp b/Engine/Plugins/Runtime/MassGameplay/Source/MassReplication/Private/MassReplicationGridProcessor.cpp
--- a/Engine/Plugins/Runtime/MassGameplay/Source/MassReplication/Private/MassReplicationGridProcessor.cpp
+++ b/Engine/Plugins/Runtime/MassGameplay/Source/MassReplication/Private/MassReplicationGridProcessor.cpp
@@ -131,7 +131,7 @@ UMassReplicationGridRemoverProcessor::UMassReplicationGridRemoverProcessor()
 	: EntityQuery(*this)
 {
 	ObservedType = FMassReplicationGridCellLocationFragment::StaticStruct();
-	Operation = EMassObservedOperation::Remove;
+	ObservedOperations = EMassObservedOperationFlags::Remove;
 	ExecutionFlags = (int32)EProcessorExecutionFlags::AllNetModes;
 }
 
diff --git a/Engine/Plugins/Runtime/MassGameplay/Source/MassRepresentation/Private/MassRepresentationProcessor.cpp b/Engine/Plugins/Runtime/MassGameplay/Source/MassRepresentation/Private/MassRepresentationProcessor.cpp
--- a/Engine/Plugins/Runtime/MassGameplay/Source/MassRepresentation/Private/MassRepresentationProcessor.cpp
+++ b/Engine/Plugins/Runtime/MassGameplay/Source/MassRepresentation/Private/MassRepresentationProcessor.cpp
@@ -444,7 +444,7 @@ UMassRepresentationFragmentDestructor::UMassRepresentationFragmentDestructor()
 	: EntityQuery(*this)
 {
 	ObservedType = FMassRepresentationFragment::StaticStruct();
-	Operation = EMassObservedOperation::Remove;
+	ObservedOperations = EMassObservedOperationFlags::Remove;
 	ExecutionFlags = (int32)EProcessorExecutionFlags::AllNetModes;
 	bRequiresGameThreadExecution = true; // due to FMassRepresentationSubsystemSharedFragment.RepresentationSubsystem use
 }
diff --git a/Engine/Plugins/Runtime/MassGameplay/Source/MassRepresentation/Private/MassStationaryISMRepresentationFragmentDestructor.cpp b/Engine/Plugins/Runtime/MassGameplay/Source/MassRepresentation/Private/MassStationaryISMRepresentationFragmentDestructor.cpp
--- a/Engine/Plugins/Runtime/MassGameplay/Source/MassRepresentation/Private/MassStationaryISMRepresentationFragmentDestructor.cpp
+++ b/Engine/Plugins/Runtime/MassGameplay/Source/MassRepresentation/Private/MassStationaryISMRepresentationFragmentDestructor.cpp
@@ -12,7 +12,7 @@ UMassStationaryISMRepresentationFragmentDestructor::UMassStationaryISMRepresenta
 	: EntityQuery(*this)
 {
 	ObservedType = FMassRepresentationFragment::StaticStruct();
-	Operation = EMassObservedOperation::Remove;
+	ObservedOperations = EMassObservedOperationFlags::Remove;
 	ExecutionFlags = (int32)EProcessorExecutionFlags::AllWorldModes;
 	bRequiresGameThreadExecution = true; // not sure about this
 }
diff --git a/Engine/Plugins/Runtime/MassGameplay/Source/MassSmartObjects/Private/MassSmartObjectProcessor.cpp b/Engine/Plugins/Runtime/MassGameplay/Source/MassSmartObjects/Private/MassSmartObjectProcessor.cpp
--- a/Engine/Plugins/Runtime/MassGameplay/Source/MassSmartObjects/Private/MassSmartObjectProcessor.cpp
+++ b/Engine/Plugins/Runtime/MassGameplay/Source/MassSmartObjects/Private/MassSmartObjectProcessor.cpp
@@ -386,7 +386,7 @@ UMassSmartObjectUserFragmentDeinitializer::UMassSmartObjectUserFragmentDeinitial
 	: EntityQuery(*this)
 {
 	ObservedType = FMassSmartObjectUserFragment::StaticStruct();
-	Operation = EMassObservedOperation::Remove;
+	ObservedOperations = EMassObservedOperationFlags::Remove;
 	ExecutionFlags = (int32)EProcessorExecutionFlags::AllNetModes;
 }
 
diff --git a/Engine/Plugins/Runtime/MassGameplay/Source/MassSmartObjects/Private/MassSmartObjectRegistration.cpp b/Engine/Plugins/Runtime/MassGameplay/Source/MassSmartObjects/Private/MassSmartObjectRegistration.cpp
--- a/Engine/Plugins/Runtime/MassGameplay/Source/MassSmartObjects/Private/MassSmartObjectRegistration.cpp
+++ b/Engine/Plugins/Runtime/MassGameplay/Source/MassSmartObjects/Private/MassSmartObjectRegistration.cpp
@@ -20,7 +20,7 @@ UMassSmartObjectInitializerBase::UMassSmartObjectInitializerBase()
 	: EntityQuery(*this)
 {
 	ExecutionFlags = static_cast<int32>(EProcessorExecutionFlags::Standalone | EProcessorExecutionFlags::Server);
-	Operation = EMassObservedOperation::Add;
+	ObservedOperations = EMassObservedOperationFlags::Add;
 }
 
 void UMassSmartObjectInitializerBase::ConfigureQueries(const TSharedRef<FMassEntityManager>& EntityManager)
@@ -78,7 +78,7 @@ UMassSmartObjectDeinitializerBase::UMassSmartObjectDeinitializerBase()
 	: EntityQuery(*this)
 {
 	ExecutionFlags = static_cast<int32>(EProcessorExecutionFlags::Standalone | EProcessorExecutionFlags::Server);
-	Operation = EMassObservedOperation::Remove;
+	ObservedOperations = EMassObservedOperationFlags::Remove;
 }
 
 void UMassSmartObjectDeinitializerBase::ConfigureQueries(const TSharedRef<FMassEntityManager>& EntityManager)
diff --git a/Engine/Source/Developer/MassEntityTestSuite/Private/MassCommandsTest.cpp b/Engine/Source/Developer/MassEntityTestSuite/Private/MassCommandsTest.cpp
--- a/Engine/Source/Developer/MassEntityTestSuite/Private/MassCommandsTest.cpp
+++ b/Engine/Source/Developer/MassEntityTestSuite/Private/MassCommandsTest.cpp
@@ -291,7 +291,7 @@ struct FCommands_PushWhileFlushing : FEntityTestBase
 				}
 			};
 		FMassObserverManager& ObserverManager = EntityManager->GetObserverManager();
-		ObserverManager.AddObserverInstance(*FTestTag_A::StaticStruct(), EMassObservedOperation::Add, *ObserverProcessor);
+		ObserverManager.AddObserverInstance(FTestTag_A::StaticStruct(), EMassObservedOperationFlags::Add, ObserverProcessor);
 
 		EntityManager->Defer().PushCommand<FMassCommandAddTag<FTestTag_A>>(Entities);
 		for (const FMassEntityHandle& EntityHandle : Entities)
diff --git a/Engine/Source/Developer/MassEntityTestSuite/Private/MassObserverTest.cpp b/Engine/Source/Developer/MassEntityTestSuite/Private/MassObserverTest.cpp
--- a/Engine/Source/Developer/MassEntityTestSuite/Private/MassObserverTest.cpp
+++ b/Engine/Source/Developer/MassEntityTestSuite/Private/MassObserverTest.cpp
@@ -29,7 +29,7 @@ struct FTagBaseOperation : FEntityTestBase
 
 	TArray<FMassEntityHandle> AffectedEntities;
 	UMassTestProcessorBase* ObserverProcessor = nullptr;
-	EMassObservedOperation OperationObserved = EMassObservedOperation::MAX;
+	EMassObservedOperationFlags OperationFlagsObserved = EMassObservedOperationFlags::None;
 	TArray<FMassEntityHandle> EntitiesInt;
 	TArray<FMassEntityHandle> EntitiesIntsFloat;
 	TArray<FMassEntityHandle> ExpectedEntities;
@@ -63,7 +63,7 @@ struct FTagBaseOperation : FEntityTestBase
 	virtual bool InstantTest() override
 	{
 		FMassObserverManager& ObserverManager = EntityManager->GetObserverManager();
-		ObserverManager.AddObserverInstance(*FTagStruct::StaticStruct(), OperationObserved, *ObserverProcessor);
+		ObserverManager.AddObserverInstance(FTagStruct::StaticStruct(), OperationFlagsObserved, ObserverProcessor);
 
 		EntityManager->BatchCreateEntities(IntsArchetype, 3, EntitiesInt);
 		EntityManager->BatchCreateEntities(FloatsIntsArchetype, 3, EntitiesIntsFloat);
@@ -89,7 +89,7 @@ struct FTagBaseOperation : FEntityTestBase
 
 struct FSingleEntitySingleArchetypeAdd : FTagBaseOperation
 {
-	FSingleEntitySingleArchetypeAdd() { OperationObserved = EMassObservedOperation::Add; }
+	FSingleEntitySingleArchetypeAdd() { OperationFlagsObserved = EMassObservedOperationFlags::Add; }
 	virtual bool PerformOperation() override 
 	{
 		ExpectedEntities = { EntitiesInt[1] };
@@ -101,7 +101,7 @@ IMPLEMENT_AI_INSTANT_TEST(FSingleEntitySingleArchetypeAdd, "System.Mass.Observer
 
 struct FSingleEntitySingleArchetypeRemove : FTagBaseOperation
 {
-	FSingleEntitySingleArchetypeRemove() { OperationObserved = EMassObservedOperation::Remove; }
+	FSingleEntitySingleArchetypeRemove() { OperationFlagsObserved = EMassObservedOperationFlags::Remove; }
 	virtual bool PerformOperation() override
 	{
 		ExpectedEntities = { EntitiesInt[1] };
@@ -118,7 +118,7 @@ IMPLEMENT_AI_INSTANT_TEST(FSingleEntitySingleArchetypeRemove, "System.Mass.Obser
 
 struct FSingleEntitySingleArchetypeDestroy : FTagBaseOperation
 {
-	FSingleEntitySingleArchetypeDestroy() { OperationObserved = EMassObservedOperation::Remove; }
+	FSingleEntitySingleArchetypeDestroy() { OperationFlagsObserved = EMassObservedOperationFlags::Remove; }
 	virtual bool PerformOperation() override
 	{
 		ExpectedEntities = { EntitiesInt[1] };
@@ -135,7 +135,7 @@ IMPLEMENT_AI_INSTANT_TEST(FSingleEntitySingleArchetypeDestroy, "System.Mass.Obse
 
 struct FMultipleArchetypeAdd : FTagBaseOperation
 {
-	FMultipleArchetypeAdd() { OperationObserved = EMassObservedOperation::Add; }
+	FMultipleArchetypeAdd() { OperationFlagsObserved = EMassObservedOperationFlags::Add; }
 
 	virtual bool PerformOperation() override
 	{
@@ -152,7 +152,7 @@ IMPLEMENT_AI_INSTANT_TEST(FMultipleArchetypeAdd, "System.Mass.Observer.Tag.Multi
 
 struct FMultipleArchetypeAdd_Sync : FTagBaseOperation
 {
-	FMultipleArchetypeAdd_Sync() { OperationObserved = EMassObservedOperation::Add; }
+	FMultipleArchetypeAdd_Sync() { OperationFlagsObserved = EMassObservedOperationFlags::Add; }
 
 	virtual bool PerformOperation() override
 	{
@@ -168,7 +168,7 @@ IMPLEMENT_AI_INSTANT_TEST(FMultipleArchetypeAdd_Sync, "System.Mass.Observer.Tag.
 
 struct FMultipleArchetypeRemove : FTagBaseOperation
 {
-	FMultipleArchetypeRemove() { OperationObserved = EMassObservedOperation::Remove; }
+	FMultipleArchetypeRemove() { OperationFlagsObserved = EMassObservedOperationFlags::Remove; }
 
 	virtual bool PerformOperation() override
 	{
@@ -192,7 +192,7 @@ IMPLEMENT_AI_INSTANT_TEST(FMultipleArchetypeRemove, "System.Mass.Observer.Tag.Mu
 
 struct FMultipleArchetypeRemove_Sync : FTagBaseOperation
 {
-	FMultipleArchetypeRemove_Sync() { OperationObserved = EMassObservedOperation::Remove; }
+	FMultipleArchetypeRemove_Sync() { OperationFlagsObserved = EMassObservedOperationFlags::Remove; }
 
 	virtual bool PerformOperation() override
 	{
@@ -216,7 +216,7 @@ IMPLEMENT_AI_INSTANT_TEST(FMultipleArchetypeRemove_Sync, "System.Mass.Observer.T
 
 struct FMultipleArchetypeDestroy : FTagBaseOperation
 {
-	FMultipleArchetypeDestroy() { OperationObserved = EMassObservedOperation::Remove; }
+	FMultipleArchetypeDestroy() { OperationFlagsObserved = EMassObservedOperationFlags::Remove; }
 
 	virtual bool PerformOperation() override
 	{
@@ -240,7 +240,7 @@ IMPLEMENT_AI_INSTANT_TEST(FMultipleArchetypeDestroy, "System.Mass.Observer.Tag.M
 struct FForbidModifyOnDestroy : FTagBaseOperation
 {
 	using Super = FTagBaseOperation;
-	FForbidModifyOnDestroy() { OperationObserved = EMassObservedOperation::Remove; }
+	FForbidModifyOnDestroy() { OperationFlagsObserved = EMassObservedOperationFlags::Remove; }
 
 	virtual bool SetUp() override
 	{
@@ -293,7 +293,7 @@ IMPLEMENT_AI_INSTANT_TEST(FForbidModifyOnDestroy, "System.Mass.Observer.ForbidMo
 
 struct FMultipleArchetypeSwap : FTagBaseOperation
 {
-	FMultipleArchetypeSwap() { OperationObserved = EMassObservedOperation::Remove; }
+	FMultipleArchetypeSwap() { OperationFlagsObserved = EMassObservedOperationFlags::Remove; }
 
 	virtual bool PerformOperation() override
 	{
@@ -316,14 +316,14 @@ IMPLEMENT_AI_INSTANT_TEST(FMultipleArchetypeSwap, "System.Mass.Observer.Tag.Mult
 
 struct FEntityCreation_Individuals : FTagBaseOperation
 {
-	FEntityCreation_Individuals() { OperationObserved = EMassObservedOperation::Add; }
+	FEntityCreation_Individuals() { OperationFlagsObserved = EMassObservedOperationFlags::Add; }
 
 	virtual bool InstantTest() override
 	{
 		constexpr int32 EntitiesToSpawnCount = 6;
 		
 		FMassObserverManager& ObserverManager = EntityManager->GetObserverManager();
-		ObserverManager.AddObserverInstance(*FTagStruct::StaticStruct(), OperationObserved, *ObserverProcessor);
+		ObserverManager.AddObserverInstance(FTagStruct::StaticStruct(), OperationFlagsObserved, ObserverProcessor);
 
 		int32 ArrayMidPoint = 0;
 		{
@@ -345,14 +345,14 @@ IMPLEMENT_AI_INSTANT_TEST(FEntityCreation_Individuals, "System.Mass.Observer.Cre
 
 struct FEntityCreation_Batched : FTagBaseOperation
 {
-	FEntityCreation_Batched() { OperationObserved = EMassObservedOperation::Add; }
+	FEntityCreation_Batched() { OperationFlagsObserved = EMassObservedOperationFlags::Add; }
 
 	virtual bool InstantTest() override
 	{
 		constexpr int32 EntitiesToSpawnCount = 6;
 
 		FMassObserverManager& ObserverManager = EntityManager->GetObserverManager();
-		ObserverManager.AddObserverInstance(*FTagStruct::StaticStruct(), OperationObserved, *ObserverProcessor);
+		ObserverManager.AddObserverInstance(FTagStruct::StaticStruct(), OperationFlagsObserved, ObserverProcessor);
 
 		{
 			TSharedRef<FMassEntityManager::FEntityCreationContext> CreationContext = EntityManager->BatchCreateEntities(IntsArchetype, EntitiesToSpawnCount, EntitiesInt);
@@ -381,7 +381,8 @@ struct FFragmentTestBase : FEntityTestBase
 
 	TArray<FMassEntityHandle> AffectedEntities;
 	UMassTestProcessorBase* ObserverProcessor = nullptr;
-	EMassObservedOperation OperationObserved = EMassObservedOperation::MAX;
+	//EMassObservedOperation OperationObserved = EMassObservedOperation::MAX;
+	EMassObservedOperationFlags OperationFlagsObserved = EMassObservedOperationFlags::None;
 	TArray<FMassEntityHandle> EntitiesFloats;
 	TArray<FMassEntityHandle> EntitiesInt;
 	TArray<FMassEntityHandle> EntitiesIntsFloat;
@@ -418,7 +419,7 @@ struct FFragmentTestBase : FEntityTestBase
 		EntityManager->BatchCreateEntities(FloatsIntsArchetype, 3, EntitiesIntsFloat);
 
 		FMassObserverManager& ObserverManager = EntityManager->GetObserverManager();
-		ObserverManager.AddObserverInstance(*FFragmentStruct::StaticStruct(), OperationObserved, *ObserverProcessor);
+		ObserverManager.AddObserverInstance(FFragmentStruct::StaticStruct(), OperationFlagsObserved, ObserverProcessor);
 				
 		if (PerformOperation())
 		{
@@ -441,7 +442,7 @@ struct FFragmentTestBase : FEntityTestBase
 
 struct FFragmentTest_SingleEntitySingleArchetypeAdd : FFragmentTestBase
 {
-	FFragmentTest_SingleEntitySingleArchetypeAdd() { OperationObserved = EMassObservedOperation::Add; }
+	FFragmentTest_SingleEntitySingleArchetypeAdd() { OperationFlagsObserved = EMassObservedOperationFlags::Add; }
 	virtual bool PerformOperation() override
 	{
 		ExpectedEntities = { EntitiesInt[1] };
@@ -453,7 +454,7 @@ IMPLEMENT_AI_INSTANT_TEST(FFragmentTest_SingleEntitySingleArchetypeAdd, "System.
 
 struct FFragmentTest_SingleEntitySingleArchetypeRemove : FFragmentTestBase
 {
-	FFragmentTest_SingleEntitySingleArchetypeRemove() { OperationObserved = EMassObservedOperation::Remove; }
+	FFragmentTest_SingleEntitySingleArchetypeRemove() { OperationFlagsObserved = EMassObservedOperationFlags::Remove; }
 	virtual bool PerformOperation() override
 	{
 		ExpectedEntities = { EntitiesInt[1] };
@@ -470,7 +471,7 @@ IMPLEMENT_AI_INSTANT_TEST(FFragmentTest_SingleEntitySingleArchetypeRemove, "Syst
 
 struct FFragmentTest_SingleEntitySingleArchetypeDestroy : FFragmentTestBase
 {
-	FFragmentTest_SingleEntitySingleArchetypeDestroy() { OperationObserved = EMassObservedOperation::Remove; }
+	FFragmentTest_SingleEntitySingleArchetypeDestroy() { OperationFlagsObserved = EMassObservedOperationFlags::Remove; }
 	virtual bool PerformOperation() override
 	{
 		ExpectedEntities = { EntitiesInt[1] };
@@ -486,7 +487,7 @@ IMPLEMENT_AI_INSTANT_TEST(FFragmentTest_SingleEntitySingleArchetypeDestroy, "Sys
 
 struct FFragmentTest_MultipleArchetypeAdd : FFragmentTestBase
 {
-	FFragmentTest_MultipleArchetypeAdd() { OperationObserved = EMassObservedOperation::Add; }
+	FFragmentTest_MultipleArchetypeAdd() { OperationFlagsObserved = EMassObservedOperationFlags::Add; }
 
 	virtual bool PerformOperation() override
 	{
@@ -507,7 +508,7 @@ IMPLEMENT_AI_INSTANT_TEST(FFragmentTest_MultipleArchetypeAdd, "System.Mass.Obser
 
 struct FFragmentTest_MultipleArchetypeRemove : FFragmentTestBase
 {
-	FFragmentTest_MultipleArchetypeRemove() { OperationObserved = EMassObservedOperation::Remove; }
+	FFragmentTest_MultipleArchetypeRemove() { OperationFlagsObserved = EMassObservedOperationFlags::Remove; }
 
 	virtual bool PerformOperation() override
 	{
@@ -530,7 +531,7 @@ IMPLEMENT_AI_INSTANT_TEST(FFragmentTest_MultipleArchetypeRemove, "System.Mass.Ob
 
 struct FFragmentTest_MultipleArchetypeDestroy : FFragmentTestBase
 {
-	FFragmentTest_MultipleArchetypeDestroy() { OperationObserved = EMassObservedOperation::Remove; }
+	FFragmentTest_MultipleArchetypeDestroy() { OperationFlagsObserved = EMassObservedOperationFlags::Remove; }
 
 	virtual bool PerformOperation() override
 	{
@@ -553,7 +554,7 @@ IMPLEMENT_AI_INSTANT_TEST(FFragmentTest_MultipleArchetypeDestroy, "System.Mass.O
 
 struct FFragmentTest_EntityCreation_Individual : FFragmentTestBase
 {
-	FFragmentTest_EntityCreation_Individual() { OperationObserved = EMassObservedOperation::Add; }
+	FFragmentTest_EntityCreation_Individual() { OperationFlagsObserved = EMassObservedOperationFlags::Add; }
 
 	virtual bool InstantTest() override
 	{
@@ -570,7 +571,7 @@ struct FFragmentTest_EntityCreation_Individual : FFragmentTestBase
 			};
 
 		FMassObserverManager& ObserverManager = EntityManager->GetObserverManager();
-		ObserverManager.AddObserverInstance(*FFragmentStruct::StaticStruct(), OperationObserved, *ObserverProcessor);
+		ObserverManager.AddObserverInstance(FFragmentStruct::StaticStruct(), OperationFlagsObserved, ObserverProcessor);
 
 		TArray<FInstancedStruct> FragmentInstanceList = { FInstancedStruct::Make(FFragmentStruct(TestValue)) };
 
@@ -599,14 +600,14 @@ IMPLEMENT_AI_INSTANT_TEST(FFragmentTest_EntityCreation_Individual, "System.Mass.
 
 struct FFragmentTest_EntityCreation_Individuals : FFragmentTestBase
 {
-	FFragmentTest_EntityCreation_Individuals() { OperationObserved = EMassObservedOperation::Add; }
+	FFragmentTest_EntityCreation_Individuals() { OperationFlagsObserved = EMassObservedOperationFlags::Add; }
 
 	virtual bool InstantTest() override
 	{
 		constexpr int32 EntitiesToSpawnCount = 6;
 
 		FMassObserverManager& ObserverManager = EntityManager->GetObserverManager();
-		ObserverManager.AddObserverInstance(*FFragmentStruct::StaticStruct(), OperationObserved, *ObserverProcessor);
+		ObserverManager.AddObserverInstance(FFragmentStruct::StaticStruct(), OperationFlagsObserved, ObserverProcessor);
 
 		int32 ArrayMidPoint = 0;
 		{
@@ -631,7 +632,7 @@ struct FObserverChangingComposition_Sync : FFragmentTestBase
 {
 	FObserverChangingComposition_Sync()
 	{
-		OperationObserved = EMassObservedOperation::Add;
+		OperationFlagsObserved = EMassObservedOperationFlags::Add;
 	}
 
 	virtual bool InstantTest() override
@@ -640,7 +641,7 @@ struct FObserverChangingComposition_Sync : FFragmentTestBase
 		const FMassArchetypeHandle OriginalArchetype = FloatsArchetype;
 
 		FMassObserverManager& ObserverManager = EntityManager->GetObserverManager();
-		ObserverManager.AddObserverInstance(*FFragmentStruct::StaticStruct(), OperationObserved, *ObserverProcessor);
+		ObserverManager.AddObserverInstance(FFragmentStruct::StaticStruct(), OperationFlagsObserved, ObserverProcessor);
 
 		{
 			AITEST_SCOPED_CHECK("Use asynchronous API instead", 1);
@@ -678,7 +679,7 @@ struct FObserverChangingComposition_Deferred : FFragmentTestBase
 {
 	FObserverChangingComposition_Deferred()
 	{
-		OperationObserved = EMassObservedOperation::Add;
+		OperationFlagsObserved = EMassObservedOperationFlags::Add;
 	}
 
 	virtual bool InstantTest() override
@@ -692,7 +693,7 @@ struct FObserverChangingComposition_Deferred : FFragmentTestBase
 			};
 
 		FMassObserverManager& ObserverManager = EntityManager->GetObserverManager();
-		ObserverManager.AddObserverInstance(*FFragmentStruct::StaticStruct(), OperationObserved, *ObserverProcessor);
+		ObserverManager.AddObserverInstance(FFragmentStruct::StaticStruct(), OperationFlagsObserved, ObserverProcessor);
 
 		EntityManager->BatchCreateEntities(OriginalArchetype, EntitiesToSpawn, EntitiesInt);
 
@@ -711,7 +712,7 @@ IMPLEMENT_AI_INSTANT_TEST(FObserverChangingComposition_Deferred, "System.Mass.Ob
  */
 struct FModificationsWhileCreationContextActive : FTagBaseOperation
 {
-	FModificationsWhileCreationContextActive() { OperationObserved = EMassObservedOperation::Add; }
+	FModificationsWhileCreationContextActive() { OperationFlagsObserved = EMassObservedOperationFlags::Add; }
 
 	virtual bool InstantTest() override
 	{
@@ -722,7 +723,7 @@ struct FModificationsWhileCreationContextActive : FTagBaseOperation
 		FMassArchetypeEntityCollection InitialEntitiesCollection(IntsArchetype, EntitiesInt, FMassArchetypeEntityCollection::EDuplicatesHandling::NoDuplicates);
 
 		FMassObserverManager& ObserverManager = EntityManager->GetObserverManager();
-		ObserverManager.AddObserverInstance(*FTagStruct::StaticStruct(), OperationObserved, *ObserverProcessor);
+		ObserverManager.AddObserverInstance(FTagStruct::StaticStruct(), OperationFlagsObserved, ObserverProcessor);
 
 		{
 			TSharedRef<FMassObserverManager::FObserverLock> ObserversLock = EntityManager->GetOrMakeObserversLock();
@@ -751,7 +752,7 @@ struct FCreationOperationOrder : FFragmentTestBase
 
 	FCreationOperationOrder()
 	{
-		OperationObserved = EMassObservedOperation::Add;
+		OperationFlagsObserved = EMassObservedOperationFlags::Add;
 	}
 
 	virtual void BuildScenario(TArray<FMassEntityHandle>& PreExistingEntities, TArray<FMassEntityHandle>& NewEntities) = 0;
@@ -772,7 +773,7 @@ struct FCreationOperationOrder : FFragmentTestBase
 		EntityManager->BatchCreateEntities(IntsArchetype, 2, PreExistingEntities);
 
 		FMassObserverManager& ObserverManager = EntityManager->GetObserverManager();
-		ObserverManager.AddObserverInstance(*FFragmentStruct::StaticStruct(), OperationObserved, *ObserverProcessor);
+		ObserverManager.AddObserverInstance(FFragmentStruct::StaticStruct(), OperationFlagsObserved, ObserverProcessor);
 
 		BuildScenario(PreExistingEntities, NewEntities);
 
@@ -857,14 +858,14 @@ IMPLEMENT_AI_INSTANT_TEST(FCreationOperationOrder_Individual, "System.Mass.Obser
 // 3. Release the creation context - we expect Float observers to trigger
 struct FCreatedEntitiesUnobservedCompositionChange : FFragmentTestBase
 {
-	FCreatedEntitiesUnobservedCompositionChange() { OperationObserved = EMassObservedOperation::Add; }
+	FCreatedEntitiesUnobservedCompositionChange() { OperationFlagsObserved = EMassObservedOperationFlags::Add; }
 
 	virtual bool InstantTest() override
 	{
 		constexpr int32 EntitiesToSpawnCount = 6;
 
 		FMassObserverManager& ObserverManager = EntityManager->GetObserverManager();
-		ObserverManager.AddObserverInstance(*FFragmentStruct::StaticStruct(), OperationObserved, *ObserverProcessor);
+		ObserverManager.AddObserverInstance(FFragmentStruct::StaticStruct(), OperationFlagsObserved, ObserverProcessor);
 		{
 			TSharedRef<FMassEntityManager::FEntityCreationContext> CreationContext = EntityManager->BatchCreateEntities(FloatsArchetype, EntitiesToSpawnCount, EntitiesFloats);
 
@@ -905,29 +906,29 @@ struct FMoveToAnotherArchetype_SingleEntity : FEntityTestBase
 			return ObserverProcessor;
 		};
 
-		ObserverManager.AddObserverInstance(*FTestTag_A::StaticStruct(), EMassObservedOperation::Add
-			, *CreateObserver([this](FMassExecutionContext& Context)
+		ObserverManager.AddObserverInstance(FTestTag_A::StaticStruct(), EMassObservedOperationFlags::Add
+			, CreateObserver([this](FMassExecutionContext& Context)
 				{
 					bTagAdded = true;
 				}
 			));
 
-		ObserverManager.AddObserverInstance(*FTestTag_A::StaticStruct(), EMassObservedOperation::Remove
-			, *CreateObserver([this](FMassExecutionContext& Context)
+		ObserverManager.AddObserverInstance(FTestTag_A::StaticStruct(), EMassObservedOperationFlags::Remove
+			, CreateObserver([this](FMassExecutionContext& Context)
 				{
 					bTagRemoved = true;
 				}
 			));
 
-		ObserverManager.AddObserverInstance(*FTestFragment_Float::StaticStruct(), EMassObservedOperation::Add
-			, *CreateObserver([this](FMassExecutionContext& Context)
+		ObserverManager.AddObserverInstance(FTestFragment_Float::StaticStruct(), EMassObservedOperationFlags::Add
+			, CreateObserver([this](FMassExecutionContext& Context)
 				{
 					bFloatAdded = true;
 				}
 			));
 
-		ObserverManager.AddObserverInstance(*FTestFragment_Float::StaticStruct(), EMassObservedOperation::Remove
-			, *CreateObserver([this](FMassExecutionContext& Context)
+		ObserverManager.AddObserverInstance(FTestFragment_Float::StaticStruct(), EMassObservedOperationFlags::Remove
+			, CreateObserver([this](FMassExecutionContext& Context)
 				{
 					bFloatRemoved = true;
 				}
diff --git a/Engine/Source/Runtime/MassEntity/Private/MassEntityManager.cpp b/Engine/Source/Runtime/MassEntity/Private/MassEntityManager.cpp
--- a/Engine/Source/Runtime/MassEntity/Private/MassEntityManager.cpp
+++ b/Engine/Source/Runtime/MassEntity/Private/MassEntityManager.cpp
@@ -1513,11 +1513,11 @@ void FMassEntityManager::BatchChangeTagsForEntities(TConstArrayView<FMassArchety
 		if (ensure(CurrentArchetype) && CurrentArchetype->GetTagBitSet() != NewTagComposition)
 		{
 			FMassTagBitSet TagsAdded = TagsToAdd - CurrentArchetype->GetTagBitSet();
-			const bool bTagsAddedAreObserved = ObserverManager.HasObserversForBitSet(TagsAdded, EMassObservedOperation::Add);
+			const bool bTagsAddedAreObserved = ObserverManager.HasObserversForBitSet(TagsAdded, EMassObservedOperation::AddElement);
 			FMassTagBitSet TagsRemoved = TagsToRemove.GetOverlap(CurrentArchetype->GetTagBitSet());
 			if (TagsRemoved.IsEmpty() == false)
 			{
-				ObserverManager.OnCompositionChanged(Collection, MoveTemp(TagsRemoved), EMassObservedOperation::Remove);
+				ObserverManager.OnCompositionChanged(Collection, MoveTemp(TagsRemoved), EMassObservedOperation::RemoveElement);
 			}
 			
 			FMassArchetypeHandle NewArchetypeHandle = InternalCreateSimilarArchetype(Collection.GetArchetype().DataPtr, NewTagComposition);
@@ -1541,7 +1541,7 @@ void FMassEntityManager::BatchChangeTagsForEntities(TConstArrayView<FMassArchety
 				ObserverManager.OnCompositionChanged(
 					FMassArchetypeEntityCollection(NewArchetypeHandle, MoveTemp(NewArchetypeEntityRanges))
 					, MoveTemp(TagsAdded)
-					, EMassObservedOperation::Add);
+					, EMassObservedOperation::AddElement);
 			}
 		}
 	}
@@ -1565,12 +1565,12 @@ void FMassEntityManager::BatchChangeFragmentCompositionForEntities(TConstArrayVi
 			if (CurrentArchetype->GetFragmentBitSet() != NewFragmentComposition)
 			{
 				FMassFragmentBitSet FragmentsAdded = FragmentsToAdd - CurrentArchetype->GetFragmentBitSet();
-				const bool bFragmentsAddedAreObserved = ObserverManager.HasObserversForBitSet(FragmentsAdded, EMassObservedOperation::Add);
+				const bool bFragmentsAddedAreObserved = ObserverManager.HasObserversForBitSet(FragmentsAdded, EMassObservedOperation::AddElement);
 				FMassFragmentBitSet FragmentsRemoved = FragmentsToRemove.GetOverlap(CurrentArchetype->GetFragmentBitSet());
 				
 				if (FragmentsRemoved.IsEmpty() == false)
 				{
-					ObserverManager.OnCompositionChanged(Collection, MoveTemp(FragmentsRemoved), EMassObservedOperation::Remove);
+					ObserverManager.OnCompositionChanged(Collection, MoveTemp(FragmentsRemoved), EMassObservedOperation::RemoveElement);
 				}
 
 				FMassArchetypeHandle NewArchetypeHandle = InternalCreateSimilarArchetype(Collection.GetArchetype().DataPtr, NewFragmentComposition);
@@ -1594,7 +1594,7 @@ void FMassEntityManager::BatchChangeFragmentCompositionForEntities(TConstArrayVi
 					ObserverManager.OnCompositionChanged(
 						FMassArchetypeEntityCollection(NewArchetypeHandle, MoveTemp(NewArchetypeEntityRanges))
 						, MoveTemp(FragmentsAdded)
-						, EMassObservedOperation::Add);
+						, EMassObservedOperation::AddElement);
 				}
 			}
 		}
@@ -1639,7 +1639,7 @@ void FMassEntityManager::BatchAddFragmentInstancesForEntities(TConstArrayView<FM
 			if (CurrentArchetype->GetFragmentBitSet() != NewFragmentComposition)
 			{
 				FragmentsAdded = FragmentsAffected - CurrentArchetype->GetFragmentBitSet();
-				bFragmentsAddedAreObserved = ObserverManager.HasObserversForBitSet(FragmentsAdded, EMassObservedOperation::Add);
+				bFragmentsAddedAreObserved = ObserverManager.HasObserversForBitSet(FragmentsAdded, EMassObservedOperation::AddElement);
 
 				FMassArchetypeHandle NewArchetypeHandle = InternalCreateSimilarArchetype(TargetArchetypeHandle.DataPtr, NewFragmentComposition);
 				checkSlow(NewArchetypeHandle.IsValid());
@@ -1673,7 +1673,7 @@ void FMassEntityManager::BatchAddFragmentInstancesForEntities(TConstArrayView<FM
 				ObserverManager.OnCompositionChanged(
 					FMassArchetypeEntityCollection(TargetArchetypeHandle, MoveTemp(TargetArchetypeEntityRanges))
 					, MoveTemp(FragmentsAdded)
-					, EMassObservedOperation::Add);
+					, EMassObservedOperation::AddElement);
 			}
 		}
 		else 
@@ -1735,13 +1735,13 @@ void FMassEntityManager::MoveEntityToAnotherArchetype(FMassEntityHandle Entity,
 	check(CurrentArchetype);
 
 	FMassArchetypeCompositionDescriptor CompositionRemoved = CurrentArchetype->GetCompositionDescriptor().CalculateDifference(NewArchetype.GetCompositionDescriptor());
-	ObserverManager.OnCompositionChanged(Entity, MoveTemp(CompositionRemoved), EMassObservedOperation::Remove);
+	ObserverManager.OnCompositionChanged(Entity, MoveTemp(CompositionRemoved), EMassObservedOperation::RemoveElement);
 	
 	CurrentArchetype->MoveEntityToAnotherArchetype(Entity, NewArchetype, SharedFragmentValuesOverride);
 	GetEntityStorageInterface().SetArchetypeFromShared(Entity.Index, NewArchetypeHandle.DataPtr);
 
 	FMassArchetypeCompositionDescriptor CompositionAdded = NewArchetype.GetCompositionDescriptor().CalculateDifference(CurrentArchetype->GetCompositionDescriptor());
-	ObserverManager.OnCompositionChanged(Entity, MoveTemp(CompositionAdded), EMassObservedOperation::Add);
+	ObserverManager.OnCompositionChanged(Entity, MoveTemp(CompositionAdded), EMassObservedOperation::AddElement);
 }
 
 void FMassEntityManager::SetEntityFragmentValues(FMassEntityHandle Entity, TArrayView<const FInstancedStruct> FragmentInstanceList)
diff --git a/Engine/Source/Runtime/MassEntity/Private/MassObserverManager.cpp b/Engine/Source/Runtime/MassEntity/Private/MassObserverManager.cpp
--- a/Engine/Source/Runtime/MassEntity/Private/MassObserverManager.cpp
+++ b/Engine/Source/Runtime/MassEntity/Private/MassObserverManager.cpp
@@ -10,6 +10,7 @@
 #include "MassEntityCollection.h"
 #include "MassProcessingContext.h"
 #include "MassObserverNotificationTypes.h"
+#include "MassObserverProcessor.h"
 #include "VisualLogger/VisualLogger.h"
 
 
@@ -36,7 +37,8 @@ namespace UE::Mass::ObserverManager
 		// a helper function to reduce code duplication in FMassObserverManager::Initialize
 		template<typename TBitSet>
 		void AddRegisteredObserverProcessorInstances(FMassEntityManager& EntityManager, const EProcessorExecutionFlags WorldExecutionFlags, UObject& Owner
-			, const UMassObserverRegistry::FObserverClassesMap& RegisteredObserverTypes, TBitSet& InOutObservedBitSet, FMassObserversMap& Observers)
+			, const UMassObserverRegistry::FObserverClassesMap& RegisteredObserverTypes, TBitSet& InOutObservedBitSet, FMassObserversMap& Observers
+			, TMap<const TSubclassOf<UMassProcessor>, UMassProcessor*>& InOutExistingProcessorsMap)
 		{
 			for (auto It : RegisteredObserverTypes)
 			{
@@ -48,40 +50,76 @@ namespace UE::Mass::ObserverManager
 				}
 
 				const UScriptStruct* StructType = CastChecked<const UScriptStruct>(Type);
-				InOutObservedBitSet.Add(*StructType);
-				FMassRuntimePipeline& Pipeline = (*Observers).FindOrAdd(StructType);
+				TArray<TObjectPtr<UMassProcessor>> ObserverProcessors;
 
 				for (const FSoftClassPath& SoftProcessorClass : It.Value)
 				{
 					if (TSubclassOf<UMassProcessor> ProcessorClass = SoftProcessorClass.ResolveClass())
 					{
-						if (ProcessorClass->GetDefaultObject<UMassProcessor>()->ShouldExecute(WorldExecutionFlags))
+						if (UMassProcessor** ExistingProcessor = InOutExistingProcessorsMap.Find(ProcessorClass))
 						{
-							Pipeline.AppendProcessor(ProcessorClass, Owner);
+							ObserverProcessors.AddUnique(*ExistingProcessor);
+						}
+						else if (ProcessorClass->GetDefaultObject<UMassProcessor>()->ShouldExecute(WorldExecutionFlags))
+						{
+							UMassProcessor* ObserverInstance = NewObject<UMassProcessor>(&Owner, ProcessorClass);							
+							ObserverProcessors.Add(ObserverInstance);
+							InOutExistingProcessorsMap.Add(ProcessorClass, ObserverInstance);
 						}
 					}
 				}
-				Pipeline.Initialize(Owner, EntityManager.AsShared());
+
+				if (ObserverProcessors.Num() > 0)
+				{
+					FMassRuntimePipeline& Pipeline = (*Observers).FindOrAdd(StructType);
+					Pipeline.AppendProcessors(MoveTemp(ObserverProcessors));
+					Pipeline.Initialize(Owner, EntityManager.AsShared());
+					
+					InOutObservedBitSet.Add(*StructType);
+				}
 			}
 		};
 	} // Private
 
 	struct FDeprecationHelper
 	{
+		/** determines which observed operation type the given HandlersContainerPtr implements */
+		static EMassObservedOperation DetermineOperationType(const FMassObserversMap* HandlersContainerPtr, const FMassObserversMap* MapArray)
+		{
+			if (HandlersContainerPtr == &MapArray[static_cast<uint8>(EMassObservedOperation::AddElement)])
+			{
+				return EMassObservedOperation::AddElement;
+			}
+			if (HandlersContainerPtr == &MapArray[static_cast<uint8>(EMassObservedOperation::RemoveElement)])
+			{
+				return EMassObservedOperation::RemoveElement;
+			}
+			return EMassObservedOperation::MAX;
+		};
+
 		static void HandleSingleElement(TNotNull<FMassObserverManager*> ObserverManager, const UScriptStruct& ElementType, const FMassArchetypeEntityCollection& EntityCollection, FMassObserversMap& HandlersContainer)
 		{
-			UE::Mass::FProcessingContext ProcessingContext(ObserverManager->EntityManager);
-			ProcessingContext.AuxData.InitializeAs(&ElementType);
-			FMassObserverManager::HandleElementsImpl(ProcessingContext, { EntityCollection }, MakeArrayView((const UScriptStruct**)&ElementType, 1), HandlersContainer);
+			const bool bIsFragment = UE::Mass::IsA<FMassFragment>(&ElementType);
+			check(bIsFragment || UE::Mass::IsA<FMassTag>(&ElementType));
+
+			const EMassObservedOperation Operation = bIsFragment
+				? DetermineOperationType(&HandlersContainer, ObserverManager->EntityManager.GetObserverManager().FragmentObservers)
+				: DetermineOperationType(&HandlersContainer, ObserverManager->EntityManager.GetObserverManager().TagObservers);
+
+			HandleSingleElement(ObserverManager, ElementType, EntityCollection, Operation);
 		}
 
 		static void HandleSingleElement(TNotNull<FMassObserverManager*> ObserverManager, const UScriptStruct& ElementType, const FMassArchetypeEntityCollection& EntityCollection, const EMassObservedOperation Operation)
 		{
-			const int32 OperationIndex = int32(Operation);
+			const int32 OperationIndex = static_cast<int32>(Operation);
 			const bool bIsFragment = UE::Mass::IsA<FMassFragment>(&ElementType);
 			check(bIsFragment || UE::Mass::IsA<FMassTag>(&ElementType));
 
-			HandleSingleElement(ObserverManager, ElementType, EntityCollection
+			const UScriptStruct* ElementTypePtr = &ElementType;
+
+			UE::Mass::FProcessingContext ProcessingContext(ObserverManager->EntityManager);
+			FMassObserverManager::HandleElementsImpl(ProcessingContext, MakeArrayView(&EntityCollection, 1)
+				, {Operation, MakeArrayView(&ElementTypePtr, 1)}
 				, bIsFragment ? ObserverManager->FragmentObservers[OperationIndex] : ObserverManager->TagObservers[OperationIndex]);
 		}
 	};
@@ -94,8 +132,8 @@ namespace UE::Mass::ObserverManager
 
 	struct FBufferedNotificationExecutioner 
 	{
-		FBufferedNotificationExecutioner(FNotificationContext& InNotificationContext, const EMassObservedOperation InOpType)
-			: NotificationContext(InNotificationContext), OpType(InOpType)
+		FBufferedNotificationExecutioner(FNotificationContext& InNotificationContext)
+			: NotificationContext(InNotificationContext)
 		{}
 
 		template<typename TEntities>
@@ -134,7 +172,7 @@ namespace UE::Mass::ObserverManager
 			Change.ExportTypes(ObservedTypesOverlap);
 			FMassObserverManager::HandleElementsImpl(NotificationContext.ProcessingContext
 				, Entities.GetUpToDatePerArchetypeCollections(NotificationContext.ObserverManager.EntityManager)
-				, ObservedTypesOverlap
+				, {OpType, ObservedTypesOverlap}
 				, NotificationContext.ObserverManager.FragmentObservers[static_cast<uint8>(OpType)]);
 		}
 
@@ -142,11 +180,13 @@ namespace UE::Mass::ObserverManager
 		{
 			ObservedTypesOverlap.Reset();
 			Change.ExportTypes(ObservedTypesOverlap);
-
+			
 			FMassArchetypeHandle ArchetypeHandle = NotificationContext.ObserverManager.GetEntityManager().GetArchetypeForEntity(EntityHandle);
-			NotificationContext.ObserverManager.HandleFragmentsImpl(NotificationContext.ProcessingContext
-				, FMassArchetypeEntityCollection(MoveTemp(ArchetypeHandle), EntityHandle)
-				, ObservedTypesOverlap
+			FMassArchetypeEntityCollection AsEntityCollection(MoveTemp(ArchetypeHandle), EntityHandle);
+
+			NotificationContext.ObserverManager.HandleElementsImpl(NotificationContext.ProcessingContext
+				, MakeArrayView(&AsEntityCollection, 1)
+				, {OpType, ObservedTypesOverlap}
 				, NotificationContext.ObserverManager.FragmentObservers[static_cast<uint8>(OpType)]);
 		}
 
@@ -156,7 +196,7 @@ namespace UE::Mass::ObserverManager
 			Change.ExportTypes(ObservedTypesOverlap);
 			FMassObserverManager::HandleElementsImpl(NotificationContext.ProcessingContext
 				, Entities.GetUpToDatePerArchetypeCollections(NotificationContext.ObserverManager.EntityManager)
-				, ObservedTypesOverlap
+				, {OpType, ObservedTypesOverlap}
 				, NotificationContext.ObserverManager.TagObservers[static_cast<uint8>(OpType)]);
 		}
 
@@ -166,9 +206,11 @@ namespace UE::Mass::ObserverManager
 			Change.ExportTypes(ObservedTypesOverlap);
 
 			FMassArchetypeHandle ArchetypeHandle = NotificationContext.ObserverManager.GetEntityManager().GetArchetypeForEntity(EntityHandle);
-			NotificationContext.ObserverManager.HandleFragmentsImpl(NotificationContext.ProcessingContext
-				, FMassArchetypeEntityCollection(MoveTemp(ArchetypeHandle), EntityHandle)
-				, ObservedTypesOverlap
+			FMassArchetypeEntityCollection AsEntityCollection(MoveTemp(ArchetypeHandle), EntityHandle);
+
+			NotificationContext.ObserverManager.HandleElementsImpl(NotificationContext.ProcessingContext
+				, MakeArrayView(&AsEntityCollection, 1)
+				, {OpType, ObservedTypesOverlap}
 				, NotificationContext.ObserverManager.TagObservers[static_cast<uint8>(OpType)]);
 		}
 
@@ -193,10 +235,33 @@ namespace UE::Mass::ObserverManager
 		{
 			const FMassArchetypeHandle ArchetypeHandle = NotificationContext.ObserverManager.GetEntityManager().GetArchetypeForEntity(EntityHandle);
 			const FMassArchetypeCompositionDescriptor& ArchetypeComposition = NotificationContext.ProcessingContext.GetEntityManager()->GetArchetypeComposition(ArchetypeHandle);
-			NotificationContext.ObserverManager.OnCompositionChanged(EntityHandle, ArchetypeComposition, EMassObservedOperation::Add, &NotificationContext.ProcessingContext);
+			NotificationContext.ObserverManager.OnCompositionChanged(EntityHandle, ArchetypeComposition, EMassObservedOperation::CreateEntity, &NotificationContext.ProcessingContext);
 		}
 		FNotificationContext& NotificationContext;
 	};
+
+	/** Trivial type with a sole responsibility of initializing and incrementing FMassObserverExecutionContext.CurrentTypeIndex */
+	struct FObserverContextIterator
+	{
+		FObserverContextIterator(FMassObserverExecutionContext& InRuntimeContext)
+			: RuntimeContext(InRuntimeContext)
+		{
+			RuntimeContext.CurrentTypeIndex = 0;
+		}
+
+		int32 operator++()
+		{
+			return ++RuntimeContext.CurrentTypeIndex;
+		}
+
+		operator bool() const
+		{
+			return RuntimeContext.IsValid();
+		}
+
+	private:
+		FMassObserverExecutionContext& RuntimeContext;
+	};
 } // UE::Mass::ObserverManager
 
 //----------------------------------------------------------------------//
@@ -254,11 +319,12 @@ void FMassObserverManager::Initialize()
 	const UWorld* World = Owner->GetWorld();
 	const EProcessorExecutionFlags WorldExecutionFlags = UE::Mass::Utils::DetermineProcessorExecutionFlags(World);
 
+	TMap<const TSubclassOf<UMassProcessor>, UMassProcessor*> TransientProcessorsMap;
 	using UE::Mass::ObserverManager::Private::AddRegisteredObserverProcessorInstances;
 	for (int i = 0; i < (int)EMassObservedOperation::MAX; ++i)
 	{
-		AddRegisteredObserverProcessorInstances(EntityManager, WorldExecutionFlags, *Owner, Registry.FragmentObserverMaps[i], ObservedFragments[i], FragmentObservers[i]);
-		AddRegisteredObserverProcessorInstances(EntityManager, WorldExecutionFlags, *Owner, Registry.TagObserverMaps[i], ObservedTags[i], TagObservers[i]);
+		AddRegisteredObserverProcessorInstances(EntityManager, WorldExecutionFlags, *Owner, Registry.FragmentObserverMaps[i], ObservedFragments[i], FragmentObservers[i], TransientProcessorsMap);
+		AddRegisteredObserverProcessorInstances(EntityManager, WorldExecutionFlags, *Owner, Registry.TagObserverMaps[i], ObservedTags[i], TagObservers[i], TransientProcessorsMap);
 	}
 
 #if WITH_MASSENTITY_DEBUG
@@ -295,7 +361,7 @@ bool FMassObserverManager::OnPostEntitiesCreated(const FMassArchetypeEntityColle
 	}
 
 	const FMassArchetypeCompositionDescriptor& ArchetypeComposition = EntityManager.GetArchetypeComposition(EntityCollection.GetArchetype());
-	return OnCompositionChanged(EntityCollection, ArchetypeComposition, EMassObservedOperation::Add);
+	return OnCompositionChanged(EntityCollection, ArchetypeComposition, EMassObservedOperation::CreateEntity);
 }
 
 bool FMassObserverManager::OnPostEntityCreated(FMassEntityHandle EntityHandle, const FMassArchetypeCompositionDescriptor& Composition)
@@ -313,30 +379,30 @@ bool FMassObserverManager::OnPostEntityCreated(FMassEntityHandle EntityHandle, c
 	{
 		const FMassArchetypeHandle ArchetypeHandle = EntityManager.GetArchetypeForEntity(EntityHandle);
 		const FMassArchetypeCompositionDescriptor& ArchetypeComposition = EntityManager.GetArchetypeComposition(ArchetypeHandle);
-		return OnCompositionChanged(EntityHandle, ArchetypeComposition, EMassObservedOperation::Add);
+		return OnCompositionChanged(EntityHandle, ArchetypeComposition, EMassObservedOperation::CreateEntity);
 	}
 
-	return OnCompositionChanged(EntityHandle, Composition, EMassObservedOperation::Add);
+	return OnCompositionChanged(EntityHandle, Composition, EMassObservedOperation::CreateEntity);
 }
 
 bool FMassObserverManager::OnPreEntitiesDestroyed(const FMassArchetypeEntityCollection& EntityCollection)
 {
 	TRACE_CPUPROFILER_EVENT_SCOPE_STR("OnPreEntitiesDestroyed")
 	const FMassArchetypeCompositionDescriptor& ArchetypeComposition = EntityManager.GetArchetypeComposition(EntityCollection.GetArchetype());
-	return OnCompositionChanged(EntityCollection, ArchetypeComposition, EMassObservedOperation::Remove);
+	return OnCompositionChanged(EntityCollection, ArchetypeComposition, EMassObservedOperation::DestroyEntity);
 }
 
 bool FMassObserverManager::OnPreEntitiesDestroyed(UE::Mass::FProcessingContext& ProcessingContext, const FMassArchetypeEntityCollection& EntityCollection)
 {
 	TRACE_CPUPROFILER_EVENT_SCOPE_STR("OnPreEntitiesDestroyed")
 	const FMassArchetypeCompositionDescriptor& ArchetypeComposition = EntityManager.GetArchetypeComposition(EntityCollection.GetArchetype());	
-	return OnCompositionChanged(EntityCollection, ArchetypeComposition, EMassObservedOperation::Remove, &ProcessingContext);
+	return OnCompositionChanged(EntityCollection, ArchetypeComposition, EMassObservedOperation::DestroyEntity, &ProcessingContext);
 }
 
 bool FMassObserverManager::OnPreEntityDestroyed(const FMassArchetypeCompositionDescriptor& ArchetypeComposition, const FMassEntityHandle Entity)
 {
 	TRACE_CPUPROFILER_EVENT_SCOPE_STR("OnPreEntityDestroyed")
-	return OnCompositionChanged(Entity, ArchetypeComposition, EMassObservedOperation::Remove);
+	return OnCompositionChanged(Entity, ArchetypeComposition, EMassObservedOperation::DestroyEntity);
 }
 
 bool FMassObserverManager::OnCompositionChanged(FCollectionRefOrHandle&& EntityCollection, const FMassArchetypeCompositionDescriptor& CompositionDelta
@@ -365,7 +431,8 @@ bool FMassObserverManager::OnCompositionChanged(FCollectionRefOrHandle&& EntityC
 			CreationContext->MarkDirty();
 			return false;
 		}
-		UE_CVLOG_UELOG(Operation == EMassObservedOperation::Remove, EntityManager.GetOwner(), LogMass, Log
+		UE_CVLOG_UELOG(Operation == EMassObservedOperation::RemoveElement || Operation == EMassObservedOperation::DestroyEntity
+			, EntityManager.GetOwner(), LogMass, Log
 			, TEXT("%hs: Remove operation while observers are locked - the remove-observer will be executed after the data has already been removed."), __FUNCTION__);
 	}
 
@@ -380,7 +447,6 @@ bool FMassObserverManager::OnCompositionChanged(FCollectionRefOrHandle&& EntityC
 	{
 		if (LocksCount > 0)
 		{
-			const UE::Mass::ObserverManager::EObservedOperationNotification NotificationType = static_cast<UE::Mass::ObserverManager::EObservedOperationNotification>(Operation);
 			TSharedRef<FObserverLock> ObserverLockRef = ActiveObserverLock.Pin().ToSharedRef();
 
 			// UE::Mass::FEntityCollection(EntityCollection) OR EntityHandle
@@ -390,13 +456,13 @@ bool FMassObserverManager::OnCompositionChanged(FCollectionRefOrHandle&& EntityC
 
 				if (EntityCollection.EntityHandle.IsSet())
 				{
-					ObserverLockRef->AddNotification(NotificationType, EntityCollection.EntityHandle
+					ObserverLockRef->AddNotification(Operation, EntityCollection.EntityHandle
 						, bHasFragmentsOverlap, MoveTemp(FragmentOverlap)
 						, bHasTagsOverlap, MoveTemp(TagOverlap));
 				}
 				else
 				{
-					ObserverLockRef->AddNotification(NotificationType, EntityCollection.EntityCollection
+					ObserverLockRef->AddNotification(Operation, EntityCollection.EntityCollection
 						, bHasFragmentsOverlap, MoveTemp(FragmentOverlap) 
 						, bHasTagsOverlap, MoveTemp(TagOverlap));
 				}
@@ -418,15 +484,15 @@ bool FMassObserverManager::OnCompositionChanged(FCollectionRefOrHandle&& EntityC
 				if (bHasFragmentsOverlap && bHasTagsOverlap)
 				{
 					FMassArchetypeCompositionDescriptor ChangedComposition(MoveTemp(FragmentOverlap), MoveTemp(TagOverlap), {}, {}, {});
-					ObserverLockRef->BufferedNotifications.Emplace(NotificationType, MoveTemp(ChangedComposition), MoveTemp(Entities));
+					ObserverLockRef->BufferedNotifications.Emplace(Operation, MoveTemp(ChangedComposition), MoveTemp(Entities));
 				}
 				else if (bHasFragmentsOverlap)
 				{
-					ObserverLockRef->BufferedNotifications.Emplace(NotificationType, MoveTemp(FragmentOverlap), MoveTemp(Entities));
+					ObserverLockRef->BufferedNotifications.Emplace(Operation, MoveTemp(FragmentOverlap), MoveTemp(Entities));
 				}
 				else // bHasTagsOverlap
 				{
-					ObserverLockRef->BufferedNotifications.Emplace(NotificationType, MoveTemp(TagOverlap), MoveTemp(Entities));
+					ObserverLockRef->BufferedNotifications.Emplace(Operation, MoveTemp(TagOverlap), MoveTemp(Entities));
 				}
 			}
 		}
@@ -445,7 +511,7 @@ bool FMassObserverManager::OnCompositionChanged(FCollectionRefOrHandle&& EntityC
 				{
 					FragmentOverlap.ExportTypes(ObservedTypesOverlap);
 
-					HandleElementsImpl(*LocalProcessingContext, {Collection}, ObservedTypesOverlap, FragmentObservers[OperationIndex]);
+					HandleElementsImpl(*LocalProcessingContext, {Collection}, {Operation, ObservedTypesOverlap}, FragmentObservers[OperationIndex]);
 				}
 
 				if (bHasTagsOverlap)
@@ -453,7 +519,7 @@ bool FMassObserverManager::OnCompositionChanged(FCollectionRefOrHandle&& EntityC
 					ObservedTypesOverlap.Reset();
 					TagOverlap.ExportTypes(ObservedTypesOverlap);
 
-					HandleElementsImpl(*LocalProcessingContext, {Collection}, ObservedTypesOverlap, TagObservers[OperationIndex]);
+					HandleElementsImpl(*LocalProcessingContext, {Collection}, {Operation, ObservedTypesOverlap}, TagObservers[OperationIndex]);
 				}
 
 				if (ProcessingContext == nullptr)
@@ -487,7 +553,8 @@ bool FMassObserverManager::OnCollectionsCreatedImpl(UE::Mass::FProcessingContext
 
 	check(LocksCount == 0);
 
-	constexpr int32 OperationIndex = static_cast<int32>(EMassObservedOperation::Add);
+	constexpr EMassObservedOperation Operation = EMassObservedOperation::CreateEntity;
+	constexpr int32 OperationIndex = static_cast<int32>(Operation);
 
 	FMassFragmentBitSet FragmentOverlap;
 	FMassTagBitSet TagOverlap;
@@ -512,14 +579,14 @@ bool FMassObserverManager::OnCollectionsCreatedImpl(UE::Mass::FProcessingContext
 		if (bHasFragmentsOverlap)
 		{
 			FragmentOverlap.ExportTypes(ObservedTypesOverlap);
-			HandleElementsImpl(ProcessingContext, EntityCollections, ObservedTypesOverlap, FragmentObservers[OperationIndex]);
+			HandleElementsImpl(ProcessingContext, EntityCollections, {Operation, ObservedTypesOverlap}, FragmentObservers[OperationIndex]);
 		}
 
 		if (bHasTagsOverlap)
 		{
 			ObservedTypesOverlap.Reset();
 			TagOverlap.ExportTypes(ObservedTypesOverlap);
-			HandleElementsImpl(ProcessingContext, EntityCollections, ObservedTypesOverlap, TagObservers[OperationIndex]);
+			HandleElementsImpl(ProcessingContext, EntityCollections, {Operation, ObservedTypesOverlap}, TagObservers[OperationIndex]);
 		}
 
 		return true;
@@ -528,87 +595,129 @@ bool FMassObserverManager::OnCollectionsCreatedImpl(UE::Mass::FProcessingContext
 }
 
 void FMassObserverManager::HandleElementsImpl(UE::Mass::FProcessingContext& ProcessingContext, TConstArrayView<FMassArchetypeEntityCollection> EntityCollections
-		, TArrayView<const UScriptStruct*> ObservedTypes, FMassObserversMap& HandlersContainer)
+		, FMassObserverExecutionContext&& ObserverContext, FMassObserversMap& HandlersContainer)
 {	
-	TRACE_CPUPROFILER_EVENT_SCOPE(MassObserver_HandleFragmentsImpl);
+	TRACE_CPUPROFILER_EVENT_SCOPE(MassObserver_HandleElementsImpl);
 
-	check(ObservedTypes.Num());
+	check(ObserverContext.GetTypesInOperation().Num());
 	ensureMsgf(EntityCollections.Num(), TEXT("Empty collections array is unexpected at this point. Nothing bad will happen, but it's a waste of perf."));
 
 	FMassEntityManager::FScopedProcessing ProcessingScope = ProcessingContext.EntityManager->NewProcessingScope();
 
-	for (const UScriptStruct* Type : ObservedTypes)
-	{		
-		ProcessingContext.AuxData.InitializeAs(Type);
-		FMassRuntimePipeline& Pipeline = (*HandlersContainer).FindChecked(Type);
+	// @todo maybe we could make this type configurable, so that project-specific code can extend it? 
+	FMassObserverExecutionContext& RuntimeObserverContext = ProcessingContext.AuxData.InitializeAs<FMassObserverExecutionContext>(MoveTemp(ObserverContext));
+
+	for (UE::Mass::ObserverManager::FObserverContextIterator ContextIterator(RuntimeObserverContext); ContextIterator; ++ContextIterator)
+	{
+		FMassRuntimePipeline& Pipeline = (*HandlersContainer).FindChecked(RuntimeObserverContext.GetCurrentType());
 
 		UE::Mass::Executor::RunProcessorsView(Pipeline.GetMutableProcessors(), ProcessingContext, EntityCollections);
 	}
 }
 
-void FMassObserverManager::AddObserverInstance(const UScriptStruct& ElementType, const EMassObservedOperation Operation, UMassProcessor& ObserverProcessor)
+void FMassObserverManager::AddObserverInstance(TNotNull<const UScriptStruct*> ElementType, const EMassObservedOperationFlags OperationFlags, TNotNull<UMassProcessor*> ObserverProcessor)
 {
-	const bool bIsFragment = UE::Mass::IsA<FMassFragment>(&ElementType);
-	checkSlow(bIsFragment || UE::Mass::IsA<FMassTag>(&ElementType));
+	const bool bIsFragment = UE::Mass::IsA<FMassFragment>(ElementType);
+	checkSlow(bIsFragment || UE::Mass::IsA<FMassTag>(ElementType));
 
-	FMassRuntimePipeline* Pipeline = nullptr;
+	FMassObserversMap* ObserversMap = bIsFragment
+		? FragmentObservers
+		: TagObservers;
 
-	if (bIsFragment)
-	{
-		Pipeline = &(*FragmentObservers[static_cast<uint8>(Operation)]).FindOrAdd(&ElementType);
-		ObservedFragments[static_cast<uint8>(Operation)].Add(ElementType);
-	}
-	else
-	{
-		Pipeline = &(*TagObservers[static_cast<uint8>(Operation)]).FindOrAdd(&ElementType);
-		ObservedTags[static_cast<uint8>(Operation)].Add(ElementType);
-	}
+	bool bInitializeCalled = false;
 
-	// AppendUniqueProcessor will return true only if ObserverProcessor has not been a part of Pipeline yet.
-	// Otherwise, we don't need to CallInitialize.
-	if (Pipeline->AppendUniqueProcessor(ObserverProcessor))
+	for (uint8 OperationIndex = 0; OperationIndex < static_cast<uint8>(EMassObservedOperation::MAX); ++OperationIndex)
 	{
-		// calling initialize to ensure the given processor is related to the same EntityManager
-		if (UObject* Owner = EntityManager.GetOwner())
-		{	
-			ObserverProcessor.CallInitialize(Owner, EntityManager.AsShared());
+		if (!!(OperationFlags & static_cast<EMassObservedOperationFlags>(1 << OperationIndex)))
+		{
+			if (bIsFragment)
+			{
+				ObservedFragments[OperationIndex].Add(*ElementType);
+			}
+			else
+			{
+				ObservedTags[OperationIndex].Add(*ElementType);
+			}
+
+			FMassRuntimePipeline& Pipeline = (*ObserversMap[OperationIndex]).FindOrAdd(ElementType);
+			if (Pipeline.AppendUniqueProcessor(*ObserverProcessor) && !bInitializeCalled)
+			{
+				if (UObject* Owner = EntityManager.GetOwner())
+				{	
+					ObserverProcessor->CallInitialize(Owner, EntityManager.AsShared());
+					bInitializeCalled = true;
+				}
+			}
 		}
 	}
 }
 
-void FMassObserverManager::RemoveObserverInstance(const UScriptStruct& ElementType, const EMassObservedOperation Operation, UMassProcessor& ObserverProcessor)
+void FMassObserverManager::AddObserverInstance(const UScriptStruct& ElementType, const EMassObservedOperation Operation, UMassProcessor& ObserverProcessor)
 {
-	const bool bIsFragmentObserver = UE::Mass::IsA<FMassFragment>(&ElementType);
+	ensureAlways(Operation < EMassObservedOperation::MAX);
+	check(Operation != EMassObservedOperation::MAX);
 
-	if (!ensure(bIsFragmentObserver || UE::Mass::IsA<FMassTag>(&ElementType)))
-	{
-		return;
-	}
+	EMassObservedOperationFlags OperationFlag = (Operation < EMassObservedOperation::MAX)
+		? static_cast<EMassObservedOperationFlags>(1 << static_cast<uint8>(Operation))
+		: (Operation == EMassObservedOperation::Add
+			? EMassObservedOperationFlags::Add
+			: EMassObservedOperationFlags::Remove);
+
+	AddObserverInstance(&ElementType, OperationFlag, &ObserverProcessor);
+}
 
-	TMap<TObjectPtr<const UScriptStruct>, FMassRuntimePipeline>& ObserversMap =
-		bIsFragmentObserver ? *FragmentObservers[(uint8)Operation] : *TagObservers[(uint8)Operation];
+void FMassObserverManager::RemoveObserverInstance(TNotNull<const UScriptStruct*> ElementType, const EMassObservedOperationFlags OperationFlags, TNotNull<UMassProcessor*> ObserverProcessor)
+{
+	const bool bIsFragment = UE::Mass::IsA<FMassFragment>(ElementType);
 
-	FMassRuntimePipeline* Pipeline = ObserversMap.Find(&ElementType);
-	if (!ensureMsgf(Pipeline, TEXT("Trying to remove an observer for a fragment/tag that does not seem to be observed.")))
+	if (!ensure(bIsFragment || UE::Mass::IsA<FMassTag>(ElementType)))
 	{
 		return;
 	}
-	Pipeline->RemoveProcessor(ObserverProcessor);
 
-	if (Pipeline->Num() == 0)
+	FMassObserversMap* ObserversMap = bIsFragment
+		? FragmentObservers
+		: TagObservers;
+
+	for (uint8 OperationIndex = 0; OperationIndex < static_cast<uint8>(EMassObservedOperation::MAX); ++OperationIndex)
 	{
-		ObserversMap.Remove(&ElementType);
-		if (bIsFragmentObserver)
+		if (!!(OperationFlags & static_cast<EMassObservedOperationFlags>(1 << OperationIndex)))
 		{
-			ObservedFragments[(uint8)Operation].Remove(ElementType);
-		}
-		else
-		{
-			ObservedTags[(uint8)Operation].Remove(ElementType);
+			FMassRuntimePipeline* Pipeline = (*ObserversMap[OperationIndex]).Find(ElementType);
+			if (ensureMsgf(Pipeline, TEXT("Trying to remove an observer for a fragment/tag that does not seem to be observed.")))
+			{
+				Pipeline->RemoveProcessor(*ObserverProcessor);
+				if (Pipeline->Num() == 0)
+				{
+					(*ObserversMap[OperationIndex]).Remove(ElementType);
+					if (bIsFragment)
+					{
+						ObservedFragments[OperationIndex].Remove(*ElementType);
+					}
+					else
+					{
+						ObservedTags[OperationIndex].Remove(*ElementType);
+					}
+				}
+			}
 		}
 	}
 }
 
+void FMassObserverManager::RemoveObserverInstance(const UScriptStruct& ElementType, const EMassObservedOperation Operation, UMassProcessor& ObserverProcessor)
+{
+	ensureAlways(Operation < EMassObservedOperation::MAX);
+	check(Operation != EMassObservedOperation::MAX);
+
+	EMassObservedOperationFlags OperationFlag = (Operation < EMassObservedOperation::MAX)
+		? static_cast<EMassObservedOperationFlags>(1 << static_cast<uint8>(Operation))
+		: (Operation == EMassObservedOperation::Add
+			? EMassObservedOperationFlags::Add
+			: EMassObservedOperationFlags::Remove);
+
+	RemoveObserverInstance(&ElementType, OperationFlag, &ObserverProcessor);
+}
+
 TSharedRef<FMassObserverManager::FObserverLock> FMassObserverManager::GetOrMakeObserverLock()
 {
 	if (ActiveObserverLock.IsValid())
@@ -690,25 +799,20 @@ void FMassObserverManager::ResumeExecution(FObserverLock& LockBeingReleased)
 		UE::Mass::FProcessingContext ProcessingContext(EntityManager);
 
 		FNotificationContext NotificationContext{*this, ProcessingContext};
-		FBufferedNotificationExecutioner AddNotification(NotificationContext, EMassObservedOperation::Add);
-		FBufferedNotificationExecutioner RemoveNotification(NotificationContext, EMassObservedOperation::Remove);
+
+		FBufferedNotificationExecutioner RegularOperationNotifier(NotificationContext);
 		FBufferedCreationNotificationExecutioner CreationOpExecutioner(NotificationContext);
 
 		for (FBufferedNotification& Op : LockBeingReleased.BufferedNotifications)
 		{
-			switch (Op.Type)
+			if (Op.Type == EMassObservedOperation::CreateEntity)
 			{
-				case EObservedOperationNotification::Add:
-					Visit(AddNotification, Op.CompositionChange, Op.AffectedEntities);
-					break;
-				case EObservedOperationNotification::Remove:
-					Visit(RemoveNotification, Op.CompositionChange, Op.AffectedEntities);
-					break;
-				case EObservedOperationNotification::Create:
-					Visit(CreationOpExecutioner, MoveTemp(Op.AffectedEntities));
-					break;
-				default:
-					ensureMsgf(false, TEXT("%hs: Unhandled EObservedOperationNotification value"), __FUNCTION__);
+				Visit(CreationOpExecutioner, MoveTemp(Op.AffectedEntities));
+			}
+			else
+			{
+				RegularOperationNotifier.OpType = Op.Type;
+				Visit(RegularOperationNotifier, Op.CompositionChange, Op.AffectedEntities);
 			}
 		}
 #if WITH_MASSENTITY_DEBUG
@@ -750,7 +854,7 @@ bool FMassObserverManager::OnPostEntitiesCreated(UE::Mass::FProcessingContext&,
 	for (const FMassArchetypeEntityCollection& Collection : EntityCollections)
 	{
 		const FMassArchetypeCompositionDescriptor& ArchetypeComposition = EntityManager.GetArchetypeComposition(Collection.GetArchetype());
-		bReturnValue |= OnCompositionChanged(Collection, ArchetypeComposition, EMassObservedOperation::Add);
+		bReturnValue |= OnCompositionChanged(Collection, ArchetypeComposition, EMassObservedOperation::CreateEntity);
 	}
 
 	return bReturnValue;
@@ -769,12 +873,12 @@ void FMassObserverManager::HandleSingleEntityImpl(const UScriptStruct& FragmentT
 
 void FMassObserverManager::OnPostFragmentOrTagAdded(const UScriptStruct& FragmentOrTagType, const FMassArchetypeEntityCollection& EntityCollection)
 {
-	UE::Mass::ObserverManager::FDeprecationHelper::HandleSingleElement(this, FragmentOrTagType, EntityCollection, EMassObservedOperation::Add);
+	UE::Mass::ObserverManager::FDeprecationHelper::HandleSingleElement(this, FragmentOrTagType, EntityCollection, EMassObservedOperation::AddElement);
 }
 
 void FMassObserverManager::OnPreFragmentOrTagRemoved(const UScriptStruct& FragmentOrTagType, const FMassArchetypeEntityCollection& EntityCollection)
 {
-	UE::Mass::ObserverManager::FDeprecationHelper::HandleSingleElement(this, FragmentOrTagType, EntityCollection, EMassObservedOperation::Remove);
+	UE::Mass::ObserverManager::FDeprecationHelper::HandleSingleElement(this, FragmentOrTagType, EntityCollection, EMassObservedOperation::RemoveElement);
 }
 
 void FMassObserverManager::OnFragmentOrTagOperation(const UScriptStruct& FragmentOrTagType, const FMassArchetypeEntityCollection& EntityCollection, const EMassObservedOperation Operation)
@@ -814,3 +918,34 @@ const FMassArchetypeEntityCollection& FMassObserverManager::FCreationContext::Ge
 	static FMassArchetypeEntityCollection DummyInstance;
 	return DummyInstance;
 }
+
+void FMassObserverManager::HandleElementsImpl(UE::Mass::FProcessingContext& ProcessingContext, TConstArrayView<FMassArchetypeEntityCollection> EntityCollections
+		, TArrayView<const UScriptStruct*> ObservedTypes, FMassObserversMap& HandlersContainer)
+{
+	if (ObservedTypes.Num() == 0 || ObservedTypes[0] == nullptr)
+	{
+		return;
+	}
+
+	// Legacy support: we need to figure out which operation this is
+	// and to do that we need to know whether we're handling fragments or tags
+	const bool bIsFragment = UE::Mass::IsA<FMassFragment>(ObservedTypes[0]);
+	check(bIsFragment || UE::Mass::IsA<FMassTag>(ObservedTypes[0]));
+
+	using UE::Mass::ObserverManager::FDeprecationHelper;
+	const EMassObservedOperation Operation = bIsFragment
+		? FDeprecationHelper::DetermineOperationType(&HandlersContainer, ProcessingContext.GetEntityManager()->GetObserverManager().FragmentObservers)
+		: FDeprecationHelper::DetermineOperationType(&HandlersContainer, ProcessingContext.GetEntityManager()->GetObserverManager().TagObservers);
+
+	checkf(Operation < EMassObservedOperation::MAX, TEXT("Unable to determine the legacy operation type"));
+
+	HandleElementsImpl(ProcessingContext, EntityCollections, {Operation, ObservedTypes}, HandlersContainer);
+}
+
+void FMassObserverManager::HandleFragmentsImpl(UE::Mass::FProcessingContext& ProcessingContext, const FMassArchetypeEntityCollection& EntityCollection
+		, TArrayView<const UScriptStruct*> ObservedTypes, FMassObserversMap& HandlersContainer)
+{
+PRAGMA_DISABLE_DEPRECATION_WARNINGS
+	HandleElementsImpl(ProcessingContext, {EntityCollection}, ObservedTypes, HandlersContainer);
+PRAGMA_ENABLE_DEPRECATION_WARNINGS
+}
diff --git a/Engine/Source/Runtime/MassEntity/Private/MassObserverProcessor.cpp b/Engine/Source/Runtime/MassEntity/Private/MassObserverProcessor.cpp
--- a/Engine/Source/Runtime/MassEntity/Private/MassObserverProcessor.cpp
+++ b/Engine/Source/Runtime/MassEntity/Private/MassObserverProcessor.cpp
@@ -25,7 +25,17 @@ void UMassObserverProcessor::PostInitProperties()
 
 	if (HasAnyFlags(RF_ClassDefaultObject) && MyClass->HasAnyClassFlags(CLASS_Abstract) == false)
 	{
-		if (ObservedType != nullptr && Operation != EMassObservedOperation::MAX)
+PRAGMA_DISABLE_DEPRECATION_WARNINGS
+		if (Operation != EMassObservedOperation::MAX)
+		{
+			ObservedOperations |= (Operation == EMassObservedOperation::Add)
+				? EMassObservedOperationFlags::Add
+				: EMassObservedOperationFlags::Remove;
+			Operation = EMassObservedOperation::MAX;
+		}
+PRAGMA_DISABLE_DEPRECATION_WARNINGS
+
+		if (ObservedType != nullptr && ObservedOperations != EMassObservedOperationFlags::None)
 		{
 			Register();
 		}
@@ -42,7 +52,7 @@ void UMassObserverProcessor::Register()
 	if (bAutoRegisterWithObserverRegistry)
 	{
 		check(ObservedType);
-		UMassObserverRegistry::GetMutable().RegisterObserver(*ObservedType, Operation, GetClass());
+		UMassObserverRegistry::GetMutable().RegisterObserver(ObservedType, static_cast<uint8>(ObservedOperations), GetClass());
 	}
 }
 
diff --git a/Engine/Source/Runtime/MassEntity/Private/MassObserverRegistry.cpp b/Engine/Source/Runtime/MassEntity/Private/MassObserverRegistry.cpp
--- a/Engine/Source/Runtime/MassEntity/Private/MassObserverRegistry.cpp
+++ b/Engine/Source/Runtime/MassEntity/Private/MassObserverRegistry.cpp
@@ -13,18 +13,26 @@ UMassObserverRegistry::UMassObserverRegistry()
 	check(HasAnyFlags(RF_ClassDefaultObject));
 }
 
-void UMassObserverRegistry::RegisterObserver(const UScriptStruct& ObservedType, const EMassObservedOperation Operation, TSubclassOf<UMassProcessor> ObserverClass)
+void UMassObserverRegistry::RegisterObserver(TNotNull<const UScriptStruct*> ObservedType, const uint8 OperationFlags, TSubclassOf<UMassProcessor> ObserverClass)
 {
 	check(ObserverClass);
-	const bool bIsFragment = UE::Mass::IsA<FMassFragment>(&ObservedType);
-	checkSlow(bIsFragment || UE::Mass::IsA<FMassTag>(&ObservedType));
+	const bool bIsFragment = UE::Mass::IsA<FMassFragment>(ObservedType);
+	checkSlow(bIsFragment || UE::Mass::IsA<FMassTag>(ObservedType));
 
-	if (bIsFragment)
-	{
-		FragmentObserverMaps[static_cast<uint8>(Operation)].FindOrAdd(&ObservedType).AddUnique(FSoftClassPath(ObserverClass));
-	}
-	else
+	FObserverClassesMap* ObserversMap = bIsFragment
+		? FragmentObserverMaps
+		: TagObserverMaps; 
+
+	for (uint8 OperationIndex = 0; OperationIndex < static_cast<uint8>(EMassObservedOperation::MAX); ++OperationIndex)
 	{
-		TagObserverMaps[static_cast<uint8>(Operation)].FindOrAdd(&ObservedType).AddUnique(FSoftClassPath(ObserverClass));
+		if ((OperationFlags & (1 << OperationIndex)))
+		{
+			ObserversMap[OperationIndex].FindOrAdd(ObservedType).AddUnique(FSoftClassPath(ObserverClass));
+		}
 	}
 }
+
+void UMassObserverRegistry::RegisterObserver(const UScriptStruct& ObservedType, const EMassObservedOperation Operation, TSubclassOf<UMassProcessor> ObserverClass)
+{
+	RegisterObserver(&ObservedType, static_cast<uint8>(1 << static_cast<uint8>(Operation)), ObserverClass);
+}
diff --git a/Engine/Source/Runtime/MassEntity/Public/MassEntityTypes.h b/Engine/Source/Runtime/MassEntity/Public/MassEntityTypes.h
--- a/Engine/Source/Runtime/MassEntity/Public/MassEntityTypes.h
+++ b/Engine/Source/Runtime/MassEntity/Public/MassEntityTypes.h
@@ -513,18 +513,41 @@ public:
 	}
 };
 
+/**
+ * The enum is used to categorize any operation an entity can be a subject to.
+ */
 UENUM()
 enum class EMassObservedOperation : uint8
 {
-	Add,
-	Remove,
-	// @todo Keeping this here as an indication of design intent. For now, we handle entity destruction like removal, but 
-	// there might be computationally expensive cases where we might want to avoid for soon-to-be-dead entities. 
-	// Destroy,
+	AddElement,	// when an element (a fragment, tag...) is added to an existing entity
+	RemoveElement,	// when an element (a fragment, tag...) is removed from an existing entity
+	DestroyEntity,	// when an entity is destroyed, which is a special case of RemoveElement, because the entity gets all of its elements removed
+	CreateEntity,	// when an entity is created, which is a special case of AddElement, because the entity gets all of its elements added
+
 	// @todo another planned supported operation type
 	// Touch,
-	MAX
+	// -- new operations above this line -- //
+
+	MAX,
+
+	// the following values are deprecated. Use one of the values above 
+	Add UMETA(Deprecated, DisplayName="DEPRECATED_Add"),
+	Remove UMETA(Deprecated, DisplayName="DEPRECATED_Remove")
+};
+
+enum class EMassObservedOperationFlags : uint8
+{
+	None = 0,
+	AddElement = 1 << static_cast<uint8>(EMassObservedOperation::AddElement),
+	RemoveElement = 1 << static_cast<uint8>(EMassObservedOperation::RemoveElement),
+	CreateEntity = 1 << static_cast<uint8>(EMassObservedOperation::CreateEntity),
+	DestroyEntity = 1 << static_cast<uint8>(EMassObservedOperation::DestroyEntity),
+	
+	Add = AddElement | CreateEntity,
+	Remove = RemoveElement | DestroyEntity,
+	All = Add | Remove,
 };
+ENUM_CLASS_FLAGS(EMassObservedOperationFlags);
 
 enum class EMassExecutionContextType : uint8
 {
diff --git a/Engine/Source/Runtime/MassEntity/Public/MassObserverManager.h b/Engine/Source/Runtime/MassEntity/Public/MassObserverManager.h
--- a/Engine/Source/Runtime/MassEntity/Public/MassObserverManager.h
+++ b/Engine/Source/Runtime/MassEntity/Public/MassObserverManager.h
@@ -15,6 +15,7 @@
 #include "MassObserverManager.generated.h"
 
 
+struct FMassObserverExecutionContext;
 struct FMassEntityManager;
 struct FMassArchetypeEntityCollection;
 class UMassProcessor;
@@ -120,12 +121,12 @@ struct FMassObserverManager
 	/** @return whether there are observers watching affected elements */
 	bool OnPostCompositionAdded(const FMassEntityHandle Entity, const FMassArchetypeCompositionDescriptor& Composition)
 	{
-		return OnCompositionChanged(Entity, Composition, EMassObservedOperation::Add);
+		return OnCompositionChanged(Entity, Composition, EMassObservedOperation::AddElement);
 	}
 	/** @return whether there are observers watching affected elements */
 	bool OnPreCompositionRemoved(const FMassEntityHandle Entity, const FMassArchetypeCompositionDescriptor& Composition)
 	{
-		return OnCompositionChanged(Entity, Composition, EMassObservedOperation::Remove);
+		return OnCompositionChanged(Entity, Composition, EMassObservedOperation::RemoveElement);
 	}
 
 protected:
@@ -186,8 +187,10 @@ public:
 		return OnCompositionChanged(Entity, FMassArchetypeCompositionDescriptor(MoveTemp(BitSet)), Operation, ProcessingContext);
 	}
 
-	MASSENTITY_API void AddObserverInstance(const UScriptStruct& ElementType, const EMassObservedOperation Operation, UMassProcessor& ObserverProcessor);
-	MASSENTITY_API void RemoveObserverInstance(const UScriptStruct& ElementType, const EMassObservedOperation Operation, UMassProcessor& ObserverProcessor);
+	MASSENTITY_API void AddObserverInstance(TNotNull<const UScriptStruct*> ElementType, EMassObservedOperationFlags OperationFlags, TNotNull<UMassProcessor*> ObserverProcessor);
+	MASSENTITY_API void AddObserverInstance(const UScriptStruct& ElementType, EMassObservedOperation Operation, UMassProcessor& ObserverProcessor);
+	MASSENTITY_API void RemoveObserverInstance(TNotNull<const UScriptStruct*> ElementType, EMassObservedOperationFlags OperationFlags, TNotNull<UMassProcessor*> ObserverProcessor);
+	MASSENTITY_API void RemoveObserverInstance(const UScriptStruct& ElementType, EMassObservedOperation Operation, UMassProcessor& ObserverProcessor);
 
 	MASSENTITY_API void ReleaseCreationHandle(UE::Mass::ObserverManager::FCreationNotificationHandle InCreationNotificationHandle);
 
@@ -204,13 +207,7 @@ protected:
 	friend UE::Mass::ObserverManager::FBufferedCreationNotificationExecutioner;
 
 	static MASSENTITY_API void HandleElementsImpl(UE::Mass::FProcessingContext& ProcessingContext, TConstArrayView<FMassArchetypeEntityCollection> EntityCollections
-		, TArrayView<const UScriptStruct*> ObservedTypes, FMassObserversMap& HandlersContainer);
-
-	static void HandleFragmentsImpl(UE::Mass::FProcessingContext& ProcessingContext, const FMassArchetypeEntityCollection& EntityCollection
-		, TArrayView<const UScriptStruct*> ObservedTypes, FMassObserversMap& HandlersContainer)
-	{
-		HandleElementsImpl(ProcessingContext, {EntityCollection}, ObservedTypes, HandlersContainer);
-	}
+		, FMassObserverExecutionContext&& ObserverContext, FMassObserversMap& HandlersContainer);
 
 	/** Coalesces all the elements observed in all the collections and executes all the observers at once */
 	MASSENTITY_API bool OnCollectionsCreatedImpl(UE::Mass::FProcessingContext& ProcessingContext, TArray<FMassArchetypeEntityCollection>&& EntityCollections);
@@ -286,7 +283,13 @@ public:
 protected:
 	UE_DEPRECATED(5.6, "Use HandleElementsImpl instead.")
 	MASSENTITY_API void HandleSingleEntityImpl(const UScriptStruct& FragmentType, const FMassArchetypeEntityCollection& EntityCollection, FMassObserversMap& HandlersContainer);
-	
+	UE_DEPRECATED(5.7, "Use the implementation with the FMassObserverExecutionContext parameter")
+	static void HandleElementsImpl(UE::Mass::FProcessingContext& ProcessingContext, TConstArrayView<FMassArchetypeEntityCollection> EntityCollections
+		, TArrayView<const UScriptStruct*> ObservedTypes, FMassObserversMap& HandlersContainer);
+	UE_DEPRECATED(5.7, "Use the HandleElementsImpl implementation with the FMassObserverExecutionContext parameter")
+	static void HandleFragmentsImpl(UE::Mass::FProcessingContext& ProcessingContext, const FMassArchetypeEntityCollection& EntityCollection
+		, TArrayView<const UScriptStruct*> ObservedTypes, FMassObserversMap& HandlersContainer);
+
 	friend UE::Mass::ObserverManager::FDeprecationHelper;
 };
 
diff --git a/Engine/Source/Runtime/MassEntity/Public/MassObserverNotificationTypes.h b/Engine/Source/Runtime/MassEntity/Public/MassObserverNotificationTypes.h
--- a/Engine/Source/Runtime/MassEntity/Public/MassObserverNotificationTypes.h
+++ b/Engine/Source/Runtime/MassEntity/Public/MassObserverNotificationTypes.h
@@ -17,14 +17,15 @@ namespace UE::Mass
 	struct FProcessingContext;
 	namespace ObserverManager
 	{
-		enum class EObservedOperationNotification : uint8
+		enum class
+		UE_DEPRECATED(5.7, "The type is no longer being used, replaced by EMassObservedOperation")
+		EObservedOperationNotification : uint8
 		{
 			Add = static_cast<uint8>(EMassObservedOperation::Add),
 			Remove = static_cast<uint8>(EMassObservedOperation::Remove),
 			Create
 		};
 
-
 		/**
 		 * The type represents a single "operation", as observed by the registered observers, an operation that has been performed
 		 * while the FObserverLock was active. Instances of FBufferedNotification contain all the information needed
@@ -39,14 +40,14 @@ namespace UE::Mass
 			struct FEmptyComposition
 			{
 			};
-			EObservedOperationNotification Type;
+			EMassObservedOperation Type;
 			using FCompositionDescription = TVariant<FEmptyComposition, FMassArchetypeCompositionDescriptor, FMassFragmentBitSet, FMassTagBitSet>;
 			FCompositionDescription CompositionChange;
 			using FEntitiesContainer = TVariant<FEntityCollection, FMassEntityHandle>;
 			FEntitiesContainer AffectedEntities;
 
 			template<typename TComposition>
-			FBufferedNotification(const EObservedOperationNotification InType, TComposition&& Composition, FEntitiesContainer&& Entities)
+			FBufferedNotification(const EMassObservedOperation InType, TComposition&& Composition, FEntitiesContainer&& Entities)
 				: Type(InType)
 				, CompositionChange(TInPlaceType<typename TDecay<TComposition>::Type>(), Forward<TComposition>(Composition))
 				, AffectedEntities(MoveTemp(Entities))
@@ -54,7 +55,7 @@ namespace UE::Mass
 			}
 
 			template<typename TComposition, typename TEntities>
-			FBufferedNotification(const EObservedOperationNotification InType, TComposition&& Composition, TEntities&& Entities)
+			FBufferedNotification(const EMassObservedOperation InType, TComposition&& Composition, TEntities&& Entities)
 				: Type(InType)
 				, CompositionChange(TInPlaceType<typename TDecay<TComposition>::Type>(), Forward<TComposition>(Composition))
 				, AffectedEntities(TInPlaceType<typename TDecay<TEntities>::Type>(), Forward<TEntities>(Entities))
@@ -62,14 +63,14 @@ namespace UE::Mass
 			}
 
 			template<typename TEntities>
-			FBufferedNotification(const EObservedOperationNotification InType, FCompositionDescription&& Composition, TEntities&& Entities)
+			FBufferedNotification(const EMassObservedOperation InType, FCompositionDescription&& Composition, TEntities&& Entities)
 				: Type(InType)
 				, CompositionChange(MoveTemp(Composition))
 				, AffectedEntities(TInPlaceType<typename TDecay<TEntities>::Type>(), Forward<TEntities>(Entities))
 			{
 			}
 
-			FBufferedNotification(const EObservedOperationNotification InType, FCompositionDescription&& Composition, const FMassArchetypeEntityCollection& Entities)
+			FBufferedNotification(const EMassObservedOperation InType, FCompositionDescription&& Composition, const FMassArchetypeEntityCollection& Entities)
 				: Type(InType)
 				, CompositionChange(MoveTemp(Composition))
 				, AffectedEntities(TInPlaceType<typename TDecay<FEntityCollection>::Type>(), Entities)
@@ -78,7 +79,7 @@ namespace UE::Mass
 
 			bool IsCreationNotification() const
 			{
-				return Type == EObservedOperationNotification::Create;
+				return Type == EMassObservedOperation::CreateEntity;
 			}
 
 			void AddHandle(const FMassEntityHandle EntityHandle)
@@ -251,7 +252,7 @@ namespace UE::Mass
 				if (CreationNotificationIndex == INDEX_NONE)
 				{
 					CreationNotificationIndex = BufferedNotifications.Num();
-					BufferedNotifications.Emplace(EObservedOperationNotification::Create
+					BufferedNotifications.Emplace(EMassObservedOperation::CreateEntity
 						, FBufferedNotification::FEmptyComposition{}
 						, UE::Mass::FEntityCollection());
 				}
@@ -280,7 +281,7 @@ namespace UE::Mass
 				if (CreationNotificationIndex == INDEX_NONE)
 				{
 					CreationNotificationIndex = BufferedNotifications.Num();
-					BufferedNotifications.Emplace(EObservedOperationNotification::Create
+					BufferedNotifications.Emplace(EMassObservedOperation::CreateEntity
 						, FBufferedNotification::FEmptyComposition{}
 						, CreatedEntity);
 				}
@@ -298,7 +299,7 @@ namespace UE::Mass
 				if (CreationNotificationIndex == INDEX_NONE)
 				{
 					CreationNotificationIndex = BufferedNotifications.Num();
-					BufferedNotifications.Emplace(EObservedOperationNotification::Create
+					BufferedNotifications.Emplace(EMassObservedOperation::CreateEntity
 						, FBufferedNotification::FEmptyComposition{}
 						, UE::Mass::FEntityCollection(InCreatedEntities));
 				}
@@ -316,7 +317,7 @@ namespace UE::Mass
 				if (CreationNotificationIndex == INDEX_NONE)
 				{
 					CreationNotificationIndex = BufferedNotifications.Num();
-					BufferedNotifications.Emplace(EObservedOperationNotification::Create, FBufferedNotification::FEmptyComposition{}
+					BufferedNotifications.Emplace(EMassObservedOperation::CreateEntity, FBufferedNotification::FEmptyComposition{}
 						, UE::Mass::FEntityCollection(InCreatedEntities, Forward<FMassArchetypeEntityCollection>(InEntityCollection)));
 				}
 				else
@@ -334,7 +335,7 @@ namespace UE::Mass
 				if (CreationNotificationIndex == INDEX_NONE)
 				{
 					CreationNotificationIndex = BufferedNotifications.Num();
-					BufferedNotifications.Emplace(EObservedOperationNotification::Create
+					BufferedNotifications.Emplace(EMassObservedOperation::CreateEntity
 						, FBufferedNotification::FEmptyComposition{}
 						, UE::Mass::FEntityCollection(Forward<T>(InEntityCollection)));
 				}
@@ -347,7 +348,7 @@ namespace UE::Mass
 			}
 
 			template<typename TEntities>
-			void AddNotification(const EObservedOperationNotification NotificationType
+			void AddNotification(const EMassObservedOperation OperationType
 				, TEntities&& Entities
 				, const bool bHasFragmentsOverlap, FMassFragmentBitSet&& FragmentOverlap
 				, const bool bHasTagsOverlap, FMassTagBitSet&& TagOverlap)
@@ -360,14 +361,14 @@ namespace UE::Mass
 					: FBufferedNotification::FCompositionDescription(TInPlaceType<FMassArchetypeCompositionDescriptor>(), FMassArchetypeCompositionDescriptor(MoveTemp(FragmentOverlap), MoveTemp(TagOverlap), {}, {}, {}));
 
 				if (BufferedNotifications.Num() 
-					&& BufferedNotifications.Last().Type == NotificationType
+					&& BufferedNotifications.Last().Type == OperationType
 					&& FBufferedNotification::AreCompositionsEqual(BufferedNotifications.Last().CompositionChange, CompositionChange))
 				{
 					BufferedNotifications.Last().AppendEntities(Forward<TEntities>(Entities));
 				}
 				else
 				{
-					BufferedNotifications.Emplace(NotificationType, MoveTemp(CompositionChange), Forward<TEntities>(Entities));
+					BufferedNotifications.Emplace(OperationType, MoveTemp(CompositionChange), Forward<TEntities>(Entities));
 				}
 			}
 
diff --git a/Engine/Source/Runtime/MassEntity/Public/MassObserverProcessor.h b/Engine/Source/Runtime/MassEntity/Public/MassObserverProcessor.h
--- a/Engine/Source/Runtime/MassEntity/Public/MassObserverProcessor.h
+++ b/Engine/Source/Runtime/MassEntity/Public/MassObserverProcessor.h
@@ -8,6 +8,63 @@
 #define UE_API MASSENTITY_API
 
 
+namespace UE::Mass::ObserverManager
+{
+	struct FObserverContextIterator;
+}
+
+/**
+ * Instances of the type will be fed into FMassRuntimePipeline.AuxData and at execution time will
+ * be available to observer processors via FMassExecutionContext.GetAuxData() 
+ */
+USTRUCT()
+struct FMassObserverExecutionContext
+{
+	GENERATED_BODY()
+
+	FMassObserverExecutionContext() = default;
+	FMassObserverExecutionContext(const EMassObservedOperation InOperation, const TArrayView<const UScriptStruct*> InTypesInOperation)
+		: Operation(InOperation), TypesInOperation(InTypesInOperation)
+	{	
+	}
+
+	EMassObservedOperation GetOperationType() const
+	{
+		return Operation;
+	}
+
+	TConstArrayView<const UScriptStruct*> GetTypesInOperation() const
+	{
+		return TypesInOperation;
+	}
+
+	const UScriptStruct* GetCurrentType() const
+	{
+		return TypesInOperation[CurrentTypeIndex];
+	}
+
+	bool IsValid() const
+	{
+		return Operation <  EMassObservedOperation::MAX
+			&& TypesInOperation.IsValidIndex(CurrentTypeIndex);
+	}
+
+private:
+	friend UE::Mass::ObserverManager::FObserverContextIterator;
+	EMassObservedOperation Operation = EMassObservedOperation::MAX;
+	TArrayView<const UScriptStruct*> TypesInOperation;
+	int32 CurrentTypeIndex = INDEX_NONE;
+};
+
+/**
+ * Base class for Processors that are used as "observers" of entity operations.
+ * An observer declares the type of Mass element it cares about (Fragments and Tags supported at the moment) - via
+ * the ObservedType property - and the types of operations it wants to be notified of - via ObservedOperations.
+ *
+ * When an observed operation takes place the processor's regular execution will take place, with
+ * ExecutionContext's "auxiliary data" (obtained by calling GetAuxData) being filled with an instance of FMassObserverExecutionContext,
+ * that can be used to get information about the type being handled and the kind of operation. 
+ */
 UCLASS(MinimalAPI, abstract)
 class UMassObserverProcessor : public UMassProcessor
 {
@@ -20,7 +77,7 @@ protected:
 	UE_API virtual void PostInitProperties() override;
 
 	/** 
-	 * By default registers this class as Operation observer of ObservedType. Override to register for multiple 
+	 * By default, registers this class as Operation observer of ObservedType. Override to register for multiple 
 	 * operations and/or types 
 	 */
 	UE_API virtual void Register();
@@ -33,7 +90,10 @@ protected:
 	UPROPERTY()
 	TObjectPtr<const UScriptStruct> ObservedType = nullptr;
 
+	UE_DEPRECATED(5.7, "UMassObserverProcessor::Operation is deprecated. Use ObservedOperations instead")
 	EMassObservedOperation Operation = EMassObservedOperation::MAX;
+
+	EMassObservedOperationFlags ObservedOperations = EMassObservedOperationFlags::None;
 };
 
 #undef UE_API
diff --git a/Engine/Source/Runtime/MassEntity/Public/MassObserverRegistry.h b/Engine/Source/Runtime/MassEntity/Public/MassObserverRegistry.h
--- a/Engine/Source/Runtime/MassEntity/Public/MassObserverRegistry.h
+++ b/Engine/Source/Runtime/MassEntity/Public/MassObserverRegistry.h
@@ -59,7 +59,14 @@ public:
 	static UMassObserverRegistry& GetMutable() { return *GetMutableDefault<UMassObserverRegistry>(); }
 	static const UMassObserverRegistry& Get() { return *GetDefault<UMassObserverRegistry>(); }
 
-	UE_API void RegisterObserver(const UScriptStruct& ObservedType, const EMassObservedOperation Operation, TSubclassOf<UMassProcessor> ObserverClass);
+	UE_API void RegisterObserver(TNotNull<const UScriptStruct*> ObservedType, uint8 OperationFlags, TSubclassOf<UMassProcessor> ObserverClass);
+
+	void RegisterObserver(TNotNull<const UScriptStruct*> ObservedType, EMassObservedOperationFlags OperationFlags, TSubclassOf<UMassProcessor> ObserverClass)
+	{
+		RegisterObserver(ObservedType, static_cast<uint8>(OperationFlags), ObserverClass);
+	}
+
+	UE_API void RegisterObserver(const UScriptStruct& ObservedType, EMassObservedOperation Operation, TSubclassOf<UMassProcessor> ObserverClass);
 
 	using FObserverClassesMap = TMap<TObjectKey<const UScriptStruct>, TArray<FSoftClassPath>>;
 protected:
