[Mass] Deprecated Rvalue ref flavor of OnCollectionsCreatedImpl, since it didn't do anything special.

Based on discussion in review 42878312
#rnx

diff --git a/Engine/Source/Runtime/MassEntity/Private/MassObserverManager.cpp b/Engine/Source/Runtime/MassEntity/Private/MassObserverManager.cpp
--- a/Engine/Source/Runtime/MassEntity/Private/MassObserverManager.cpp
+++ b/Engine/Source/Runtime/MassEntity/Private/MassObserverManager.cpp
@@ -225,10 +225,10 @@ namespace UE::Mass::ObserverManager
 			: NotificationContext(InNotificationContext)
 		{}
 
-		void operator()(FEntityCollection&& Entities) const
+		void operator()(const FEntityCollection& Entities) const
 		{
 			NotificationContext.ObserverManager.OnCollectionsCreatedImpl(NotificationContext.ProcessingContext
-				, MoveTemp(Entities).ConsumeArchetypeCollections(NotificationContext.ObserverManager.EntityManager));
+				, Entities.GetUpToDatePerArchetypeCollections(NotificationContext.ObserverManager.EntityManager));
 		}
 
 		void operator()(const FMassEntityHandle EntityHandle) const
@@ -545,7 +545,7 @@ bool FMassObserverManager::OnCompositionChanged(FCollectionRefOrHandle&& EntityC
 	return false;
 }
 
-bool FMassObserverManager::OnCollectionsCreatedImpl(UE::Mass::FProcessingContext& ProcessingContext, TArray<FMassArchetypeEntityCollection>&& EntityCollections)
+bool FMassObserverManager::OnCollectionsCreatedImpl(UE::Mass::FProcessingContext& ProcessingContext, const TConstArrayView<FMassArchetypeEntityCollection> EntityCollections)
 {
 	using UE::Mass::ObserverManager::Tweakables::InlineAllocatorElementsForOverlapTypes;
 
@@ -807,7 +807,7 @@ void FMassObserverManager::ResumeExecution(FObserverLock& LockBeingReleased)
 		{
 			if (Op.Type == EMassObservedOperation::CreateEntity)
 			{
-				Visit(CreationOpExecutioner, MoveTemp(Op.AffectedEntities));
+				Visit(CreationOpExecutioner, Op.AffectedEntities);
 			}
 			else
 			{
@@ -949,3 +949,8 @@ PRAGMA_DISABLE_DEPRECATION_WARNINGS
 	HandleElementsImpl(ProcessingContext, {EntityCollection}, ObservedTypes, HandlersContainer);
 PRAGMA_ENABLE_DEPRECATION_WARNINGS
 }
+
+bool FMassObserverManager::OnCollectionsCreatedImpl(UE::Mass::FProcessingContext& ProcessingContext, TArray<FMassArchetypeEntityCollection>&& EntityCollections)
+{
+	return OnCollectionsCreatedImpl(ProcessingContext, MakeArrayView(EntityCollections));
+}
\ No newline at end of file
diff --git a/Engine/Source/Runtime/MassEntity/Public/MassObserverManager.h b/Engine/Source/Runtime/MassEntity/Public/MassObserverManager.h
--- a/Engine/Source/Runtime/MassEntity/Public/MassObserverManager.h
+++ b/Engine/Source/Runtime/MassEntity/Public/MassObserverManager.h
@@ -210,7 +210,7 @@ protected:
 		, FMassObserverExecutionContext&& ObserverContext, FMassObserversMap& HandlersContainer);
 
 	/** Coalesces all the elements observed in all the collections and executes all the observers at once */
-	MASSENTITY_API bool OnCollectionsCreatedImpl(UE::Mass::FProcessingContext& ProcessingContext, TArray<FMassArchetypeEntityCollection>&& EntityCollections);
+	MASSENTITY_API bool OnCollectionsCreatedImpl(UE::Mass::FProcessingContext& ProcessingContext, TConstArrayView<FMassArchetypeEntityCollection> EntityCollections);
 
 	MASSENTITY_API void OnPostFork(EForkProcessRole);
 
@@ -290,6 +290,9 @@ protected:
 	static void HandleFragmentsImpl(UE::Mass::FProcessingContext& ProcessingContext, const FMassArchetypeEntityCollection& EntityCollection
 		, TArrayView<const UScriptStruct*> ObservedTypes, FMassObserversMap& HandlersContainer);
 
+	UE_DEPRECATED(5.7, "Rvalue ref implementation is no longer available, it didn't do anything special. Please use the other OnCollectionsCreatedImpl implementation instead.")
+	MASSENTITY_API bool OnCollectionsCreatedImpl(UE::Mass::FProcessingContext& ProcessingContext, TArray<FMassArchetypeEntityCollection>&& EntityCollections);
+
 	friend UE::Mass::ObserverManager::FDeprecationHelper;
 };
 
