CIS Issue 916558: fixed mismatching PRAGMA_DISABLE_DEPRECATION_WARNINGS uses

#rnx

diff --git a/Engine/Source/Runtime/MassEntity/Private/MassObserverProcessor.cpp b/Engine/Source/Runtime/MassEntity/Private/MassObserverProcessor.cpp
--- a/Engine/Source/Runtime/MassEntity/Private/MassObserverProcessor.cpp
+++ b/Engine/Source/Runtime/MassEntity/Private/MassObserverProcessor.cpp
@@ -33,7 +33,7 @@ PRAGMA_DISABLE_DEPRECATION_WARNINGS
 				: EMassObservedOperationFlags::Remove;
 			Operation = EMassObservedOperation::MAX;
 		}
-PRAGMA_DISABLE_DEPRECATION_WARNINGS
+PRAGMA_ENABLE_DEPRECATION_WARNINGS
 
 		if (ObservedType != nullptr && ObservedOperations != EMassObservedOperationFlags::None)
 		{
@@ -42,7 +42,7 @@ PRAGMA_DISABLE_DEPRECATION_WARNINGS
 		else if (bAutoRegisterWithObserverRegistry)
 		{
 			UE_LOG(LogMass, Error, TEXT("%s attempting to register %s while it\'s misconfigured, Type: %s, Operation: %s")
-				, ANSI_TO_TCHAR(__FUNCTION__), *MyClass->GetName(), *GetNameSafe(ObservedType), *UEnum::GetValueAsString(Operation));
+				, ANSI_TO_TCHAR(__FUNCTION__), *MyClass->GetName(), *GetNameSafe(ObservedType), *UEnum::GetValueAsString(ObservedOperations));
 		}
 	}
 }
