CIS Issue 916564: no longer trying to use UEnum::GetValueAsString for EMassObservedOperationFlags values

#rnx

diff --git a/Engine/Source/Runtime/MassEntity/Private/MassObserverProcessor.cpp b/Engine/Source/Runtime/MassEntity/Private/MassObserverProcessor.cpp
--- a/Engine/Source/Runtime/MassEntity/Private/MassObserverProcessor.cpp
+++ b/Engine/Source/Runtime/MassEntity/Private/MassObserverProcessor.cpp
@@ -41,8 +41,8 @@ PRAGMA_ENABLE_DEPRECATION_WARNINGS
 		}
 		else if (bAutoRegisterWithObserverRegistry)
 		{
-			UE_LOG(LogMass, Error, TEXT("%s attempting to register %s while it\'s misconfigured, Type: %s, Operation: %s")
-				, ANSI_TO_TCHAR(__FUNCTION__), *MyClass->GetName(), *GetNameSafe(ObservedType), *UEnum::GetValueAsString(ObservedOperations));
+			UE_LOG(LogMass, Error, TEXT("%hs attempting to register %s while it\'s misconfigured, Type: %s, OperationFlags: %#x")
+				, __FUNCTION__, *MyClass->GetName(), *GetNameSafe(ObservedType), ObservedOperations);
 		}
 	}
 }
