[Mass] FMassArchetypeHandle inline functions styling refactor, matching unreal coding standard. Also added a Reset function

#rnx

diff --git a/Engine/Source/Runtime/MassEntity/Public/MassArchetypeTypes.h b/Engine/Source/Runtime/MassEntity/Public/MassArchetypeTypes.h
--- a/Engine/Source/Runtime/MassEntity/Public/MassArchetypeTypes.h
+++ b/Engine/Source/Runtime/MassEntity/Public/MassArchetypeTypes.h
@@ -38,16 +38,18 @@ using FMassChunkConditionFunction = TFunction< bool(const FMassExecutionContext&
 struct FMassArchetypeHandle final
 {
 	FMassArchetypeHandle() = default;
-	bool IsValid() const { return DataPtr.IsValid(); }
+	bool IsValid() const;
 
-	bool operator==(const FMassArchetypeHandle& Other) const { return DataPtr == Other.DataPtr; }
-	bool operator!=(const FMassArchetypeHandle& Other) const { return DataPtr != Other.DataPtr; }
+	bool operator==(const FMassArchetypeHandle& Other) const;
+	bool operator!=(const FMassArchetypeHandle& Other) const;
 
 	MASSENTITY_API friend uint32 GetTypeHash(const FMassArchetypeHandle& Instance);
+
+	void Reset();
+
 private:
-	FMassArchetypeHandle(const TSharedPtr<FMassArchetypeData>& InDataPtr)
-	: DataPtr(InDataPtr)
-	{}
+	FMassArchetypeHandle(const TSharedPtr<FMassArchetypeData>& InDataPtr);
+
 	TSharedPtr<FMassArchetypeData> DataPtr;
 
 	friend FMassArchetypeHelper;
@@ -358,6 +360,32 @@ inline FMassArchetypeVersionedHandle::operator FMassArchetypeHandle() const
 	return ArchetypeHandle;
 }
 
+inline bool FMassArchetypeHandle::IsValid() const
+{
+	return DataPtr.IsValid();
+}
+
+inline bool FMassArchetypeHandle::operator==(const FMassArchetypeHandle& Other) const
+{
+	return DataPtr == Other.DataPtr;
+}
+
+inline bool FMassArchetypeHandle::operator!=(const FMassArchetypeHandle& Other) const
+{
+	return DataPtr != Other.DataPtr;
+}
+
+inline void FMassArchetypeHandle::Reset()
+{
+	DataPtr.Reset();
+}
+
+inline FMassArchetypeHandle::FMassArchetypeHandle(const TSharedPtr<FMassArchetypeData>& InDataPtr)
+	: DataPtr(InDataPtr)
+{
+	
+}
+
 inline FMassArchetypeEntityCollection::FConstEntityRangeArrayView FMassArchetypeEntityCollection::GetRanges() const
 {
 	return Ranges;
