[AutoRTFM] Use a transactionally-safe mutex in the MassDebugger.

diff --git a/Engine/Source/Runtime/MassEntity/Private/MassDebugger.cpp b/Engine/Source/Runtime/MassEntity/Private/MassDebugger.cpp
--- a/Engine/Source/Runtime/MassEntity/Private/MassDebugger.cpp
+++ b/Engine/Source/Runtime/MassEntity/Private/MassDebugger.cpp
@@ -480,7 +480,7 @@ FMassDebugger::FOnMassEntityManagerEvent FMassDebugger::OnEntityManagerDeinitial
 FMassDebugger::FOnEnvironmentEvent FMassDebugger::OnProcessorProviderRegistered;
 FMassDebugger::FOnDebugEvent FMassDebugger::OnDebugEvent;
 TArray<FMassDebugger::FEnvironment> FMassDebugger::ActiveEnvironments;
-UE::FSpinLock FMassDebugger::EntityManagerRegistrationLock;
+UE::FTransactionallySafeMutex FMassDebugger::EntityManagerRegistrationLock;
 bool FMassDebugger::bHasBreakpoint = false;
 TMap<FName, const UScriptStruct*> FMassDebugger::FragmentsByName;
 
@@ -870,7 +870,7 @@ int32 FMassDebugger::RegisterEntityManager(FMassEntityManager& EntityManager)
 {
 	int32 NewEnvironmentIndex = INDEX_NONE;
 	{
-		UE::TScopeLock<UE::FSpinLock> ScopeLock(EntityManagerRegistrationLock);
+		UE::TScopeLock ScopeLock(EntityManagerRegistrationLock);
 		NewEnvironmentIndex = ActiveEnvironments.Emplace(EntityManager);
 	}
 	OnEntityManagerInitialized.Broadcast(EntityManager);
@@ -881,7 +881,7 @@ void FMassDebugger::UnregisterEntityManager(FMassEntityManager& EntityManager)
 {
 	if (EntityManager.DoesSharedInstanceExist())
 	{
-		UE::TScopeLock<UE::FSpinLock> ScopeLock(EntityManagerRegistrationLock);
+		UE::TScopeLock ScopeLock(EntityManagerRegistrationLock);
 		const int32 Index = ActiveEnvironments.IndexOfByPredicate([WeakManager = EntityManager.AsWeak()](const FEnvironment& Element) 
 		{
 			return Element.EntityManager == WeakManager;
@@ -893,7 +893,7 @@ void FMassDebugger::UnregisterEntityManager(FMassEntityManager& EntityManager)
 	}
 	else
 	{
-		UE::TScopeLock<UE::FSpinLock> ScopeLock(EntityManagerRegistrationLock);
+		UE::TScopeLock ScopeLock(EntityManagerRegistrationLock);
 		ActiveEnvironments.RemoveAll([](const FEnvironment& Item)
 			{
 				return Item.IsValid() == false;
@@ -904,7 +904,7 @@ void FMassDebugger::UnregisterEntityManager(FMassEntityManager& EntityManager)
 
 void FMassDebugger::RegisterProcessorDataProvider(FName ProviderName, const TSharedRef<FMassEntityManager>& EntityManager, const UE::Mass::Debug::FProcessorProviderFunction& ProviderFunction)
 {
-	UE::TScopeLock<UE::FSpinLock> ScopeLock(EntityManagerRegistrationLock);
+	UE::TScopeLock ScopeLock(EntityManagerRegistrationLock);
 	int32 Index = ActiveEnvironments.IndexOfByPredicate([WeakEntityManager = EntityManager->AsWeak()](const FEnvironment& Element) 
 	{
 		return Element.EntityManager == WeakEntityManager;
diff --git a/Engine/Source/Runtime/MassEntity/Public/MassDebugger.h b/Engine/Source/Runtime/MassEntity/Public/MassDebugger.h
--- a/Engine/Source/Runtime/MassEntity/Public/MassDebugger.h
+++ b/Engine/Source/Runtime/MassEntity/Public/MassDebugger.h
@@ -8,7 +8,7 @@
 #include "Containers/ContainersFwd.h"
 #include "MassEntityQuery.h"
 #include "MassProcessor.h"
-#include "Misc/SpinLock.h"
+#include "Async/TransactionallySafeMutex.h"
 #include "Logging/TokenizedMessage.h"
 #include "UObject/ObjectKey.h"
 #if UE_ENABLE_INCLUDE_ORDER_DEPRECATED_IN_5_6
@@ -386,7 +386,7 @@ struct FMassDebugger
 	
 private:
 	static MASSENTITY_API TArray<FEnvironment> ActiveEnvironments;
-	static MASSENTITY_API UE::FSpinLock EntityManagerRegistrationLock;
+	static MASSENTITY_API UE::FTransactionallySafeMutex EntityManagerRegistrationLock;
 	static MASSENTITY_API bool bHasBreakpoint;
 	static MASSENTITY_API TMap<FName, const UScriptStruct*> FragmentsByName;
 
