[Mass]
Release entity manager registration lock before broadcasting ProcessorDataProvider registration event

diff --git a/Engine/Source/Runtime/MassEntity/Private/MassDebugger.cpp b/Engine/Source/Runtime/MassEntity/Private/MassDebugger.cpp
--- a/Engine/Source/Runtime/MassEntity/Private/MassDebugger.cpp
+++ b/Engine/Source/Runtime/MassEntity/Private/MassDebugger.cpp
@@ -904,19 +904,21 @@ void FMassDebugger::UnregisterEntityManager(FMassEntityManager& EntityManager)
 
 void FMassDebugger::RegisterProcessorDataProvider(FName ProviderName, const TSharedRef<FMassEntityManager>& EntityManager, const UE::Mass::Debug::FProcessorProviderFunction& ProviderFunction)
 {
-	UE::TScopeLock ScopeLock(EntityManagerRegistrationLock);
-	int32 Index = ActiveEnvironments.IndexOfByPredicate([WeakEntityManager = EntityManager->AsWeak()](const FEnvironment& Element) 
+	int32 Index;
 	{
-		return Element.EntityManager == WeakEntityManager;
-	});
+		UE::TScopeLock ScopeLock(EntityManagerRegistrationLock);
+		Index = ActiveEnvironments.IndexOfByPredicate([WeakEntityManager = EntityManager->AsWeak()](const FEnvironment& Element) 
+		{
+			return Element.EntityManager == WeakEntityManager;
+		});
 
-	if (Index == INDEX_NONE)
-	{
-		Index = RegisterEntityManager(*EntityManager);
-	}
+		if (Index == INDEX_NONE)
+		{
+			Index = RegisterEntityManager(*EntityManager);
+		}
 	
-	ActiveEnvironments[Index].ProcessorProviders.FindOrAdd(ProviderName, ProviderFunction);
-
+		ActiveEnvironments[Index].ProcessorProviders.FindOrAdd(ProviderName, ProviderFunction);
+	}
 	OnProcessorProviderRegistered.Broadcast(ActiveEnvironments[Index]);
 }
 
