[AutoRTFM] Use transactionally-safe threading primitives in MassCommandBuffer.

diff --git a/Engine/Source/Runtime/MassEntity/Private/MassCommandBuffer.cpp b/Engine/Source/Runtime/MassEntity/Private/MassCommandBuffer.cpp
--- a/Engine/Source/Runtime/MassEntity/Private/MassCommandBuffer.cpp
+++ b/Engine/Source/Runtime/MassEntity/Private/MassCommandBuffer.cpp
@@ -234,7 +234,7 @@ void FMassCommandBuffer::MoveAppend(FMassCommandBuffer& Other)
 	UE_MT_SCOPED_READ_ACCESS(Other.PendingBatchCommandsDetector);
 	if (Other.HasPendingCommands())
 	{
-		FScopeLock Lock(&AppendingCommandsCS);
+		UE::TScopeLock Lock(AppendingCommandsCS);
 		UE_MT_SCOPED_WRITE_ACCESS(PendingBatchCommandsDetector);
 		AppendedCommandInstances.Append(MoveTemp(Other.CommandInstances));
 		AppendedCommandInstances.Append(MoveTemp(Other.AppendedCommandInstances));
diff --git a/Engine/Source/Runtime/MassEntity/Public/MassCommandBuffer.h b/Engine/Source/Runtime/MassEntity/Public/MassCommandBuffer.h
--- a/Engine/Source/Runtime/MassEntity/Public/MassCommandBuffer.h
+++ b/Engine/Source/Runtime/MassEntity/Public/MassCommandBuffer.h
@@ -6,7 +6,8 @@
 #if UE_ENABLE_INCLUDE_ORDER_DEPRECATED_IN_5_6
 #include "MassEntityManager.h"
 #endif // UE_ENABLE_INCLUDE_ORDER_DEPRECATED_IN_5_6
-#include "Misc/MTAccessDetector.h"
+#include "Misc/MTTransactionallySafeAccessDetector.h"
+#include "Misc/TransactionallySafeCriticalSection.h"
 #include "MassEntityUtils.h"
 #include "MassCommands.h"
 
@@ -217,9 +218,9 @@ private:
 	bool Flush(FMassEntityManager& EntityManager);
 	MASSENTITY_API void CleanUp();
 
-	FCriticalSection AppendingCommandsCS;
+	FTransactionallySafeCriticalSection AppendingCommandsCS;
 
-	UE_MT_DECLARE_RW_ACCESS_DETECTOR(PendingBatchCommandsDetector);
+	UE_MT_DECLARE_TS_RW_ACCESS_DETECTOR(PendingBatchCommandsDetector);
 	/** 
 	 * Commands created for this specific command buffer. All commands in the array are unique (by type) and reusable 
 	 * with subsequent PushCommand calls
