[AutoRTFM] Update the MassEntityManagerStorage's serial number generator for transactional safety.

diff --git a/Engine/Source/Runtime/MassEntity/Private/MassEntityManagerStorage.cpp b/Engine/Source/Runtime/MassEntity/Private/MassEntityManagerStorage.cpp
--- a/Engine/Source/Runtime/MassEntity/Private/MassEntityManagerStorage.cpp
+++ b/Engine/Source/Runtime/MassEntity/Private/MassEntityManagerStorage.cpp
@@ -112,7 +112,7 @@ namespace UE::Mass
 	FMassEntityHandle FSingleThreadedEntityStorage::AcquireOne()
 	{
 		LLM_SCOPE_BYNAME(TEXT("Mass/SingleThreadedStorage"));
-		const int32 SerialNumber = SerialNumberGenerator.fetch_add(1);
+		const int32 SerialNumber = GenerateSerialNumber();
 		const int32 Index = (EntityFreeIndexList.Num() > 0) ? EntityFreeIndexList.Pop(EAllowShrinking::No) : Entities.Add();
 		Entities[Index].SerialNumber = SerialNumber;
 
@@ -127,7 +127,7 @@ namespace UE::Mass
 		LLM_SCOPE_BYNAME(TEXT("Mass/SingleThreadedStorage"));
 		const int32 NumToAdd = OutEntityHandles.Num();
 
-		const int32 SerialNumber = SerialNumberGenerator.fetch_add(1);
+		const int32 SerialNumber = GenerateSerialNumber();
 
 		int32 NumAdded = 0;
 		int32 CurrentEntityHandleIndex = 0;
diff --git a/Engine/Source/Runtime/MassEntity/Public/MassEntityManagerStorage.h b/Engine/Source/Runtime/MassEntity/Public/MassEntityManagerStorage.h
--- a/Engine/Source/Runtime/MassEntity/Public/MassEntityManagerStorage.h
+++ b/Engine/Source/Runtime/MassEntity/Public/MassEntityManagerStorage.h
@@ -157,8 +157,8 @@ namespace UE::Mass
 		virtual int32 ForceReleaseOne(FMassEntityHandle Handle) override;
 		virtual int32 Num() const override;
 		virtual int32 ComputeFreeSize() const override;
-	private:
 
+	private:
 		struct FEntityData
 		{
 			TSharedPtr<FMassArchetypeData> CurrentArchetype;
@@ -169,7 +169,15 @@ namespace UE::Mass
 			bool IsValid() const;
 		};
 
-		std::atomic<int32> SerialNumberGenerator = 0;
+		std::atomic<int32> NextSerialNumber = 0;
+
+		UE_AUTORTFM_ALWAYS_OPEN
+		int32 GenerateSerialNumber()
+		{
+			// The serial number only needs to be unique; it doesn't need to be rolled back if an AutoRTFM transaction fails.
+			return NextSerialNumber.fetch_add(1);
+		}
+
 		TChunkedArray<FEntityData> Entities;
 		TArray<int32> EntityFreeIndexList;
 	};
