[Mass]
Added ShutdownModule handler to clean up debug environments before static destruction

diff --git a/Engine/Source/Runtime/MassEntity/Private/MassDebugger.cpp b/Engine/Source/Runtime/MassEntity/Private/MassDebugger.cpp
--- a/Engine/Source/Runtime/MassEntity/Private/MassDebugger.cpp
+++ b/Engine/Source/Runtime/MassEntity/Private/MassDebugger.cpp
@@ -1177,6 +1177,13 @@ void FMassDebugger::ClearAllBreakpoints()
 	OnBreakpointsChangedDelegate.Broadcast();
 }
 
+void FMassDebugger::ShutdownDebugger()
+{
+	ClearAllBreakpoints();
+	ActiveEnvironments.Empty();
+	FragmentsByName.Empty();
+}
+
 const UScriptStruct* FMassDebugger::GetFragmentTypeFromName(FName FragmentName)
 {
 	const UScriptStruct** FoundType = FragmentsByName.Find(FragmentName);
diff --git a/Engine/Source/Runtime/MassEntity/Private/MassEntityModule.cpp b/Engine/Source/Runtime/MassEntity/Private/MassEntityModule.cpp
--- a/Engine/Source/Runtime/MassEntity/Private/MassEntityModule.cpp
+++ b/Engine/Source/Runtime/MassEntity/Private/MassEntityModule.cpp
@@ -2,6 +2,7 @@
 
 #include "Modules/ModuleInterface.h"
 #include "MassEntityTypes.h"
+#include "MassDebugger.h"
 #include "MassProcessingTypes.h"
 #include "Modules/ModuleManager.h"
 
@@ -11,6 +12,7 @@ class FMassEntityModule : public IModuleInterface
 {
 	/** IModuleInterface implementation */
 	virtual void StartupModule() override;
+	virtual void ShutdownModule() override;
 };
 
 IMPLEMENT_MODULE(FMassEntityModule, MassEntity)
@@ -28,4 +30,11 @@ void FMassEntityModule::StartupModule()
 #endif // MASS_DO_PARALLEL
 }
 
+void FMassEntityModule::ShutdownModule()
+{
+#if WITH_MASSENTITY_DEBUG
+	FMassDebugger::ShutdownDebugger();
+#endif
+}
+
 #undef LOCTEXT_NAMESPACE 
diff --git a/Engine/Source/Runtime/MassEntity/Public/MassDebugger.h b/Engine/Source/Runtime/MassEntity/Public/MassDebugger.h
--- a/Engine/Source/Runtime/MassEntity/Public/MassDebugger.h
+++ b/Engine/Source/Runtime/MassEntity/Public/MassDebugger.h
@@ -383,8 +383,14 @@ struct FMassDebugger
 	 * Clears all breakpoints in all environments.
 	 */
 	static MASSENTITY_API void ClearAllBreakpoints();
+
+	/**
+	 * Destroy all debug environments and shutdown cleanly.
+	 */
+	static void ShutdownDebugger();
 	
 private:
+	// @todo: contain static data in an instanced struct and setup/teardown in startup/shutdown module
 	static MASSENTITY_API TArray<FEnvironment> ActiveEnvironments;
 	static MASSENTITY_API UE::FTransactionallySafeMutex EntityManagerRegistrationLock;
 	static MASSENTITY_API bool bHasBreakpoint;
