[Mass]
Clean up trace start handler weak pointer usage

diff --git a/Engine/Source/Runtime/MassEntity/Private/MassDebugger.cpp b/Engine/Source/Runtime/MassEntity/Private/MassDebugger.cpp
--- a/Engine/Source/Runtime/MassEntity/Private/MassDebugger.cpp
+++ b/Engine/Source/Runtime/MassEntity/Private/MassDebugger.cpp
@@ -1405,18 +1405,13 @@ FMassDebugger::FEnvironment::FEnvironment(FMassEntityManager& InEntityManager)
 #if UE_MASS_TRACE_ENABLED
 	TraceStartedDelegateHandle = FTraceAuxiliary::OnTraceStarted.AddLambda([WeakEntityManager = EntityManager](FTraceAuxiliary::EConnectionType TraceType, const FString& TraceDestination)
 	{
-		if (!WeakEntityManager.IsValid())
+		if (TSharedPtr<const FMassEntityManager> Manager = WeakEntityManager.Pin())
 		{
-			return;
+			ForEachArchetype(*Manager, [](FMassArchetypeHandle ArchetypeHandle)
+			{
+				UE_TRACE_MASS_ARCHETYPE_CREATED(ArchetypeHandle)
+			});
 		}
-
-		const FMassEntityManager& Manager = *WeakEntityManager.Pin();
-
-		ForEachArchetype(Manager, [](FMassArchetypeHandle ArchetypeHandle)
-		{
-			UE_TRACE_MASS_ARCHETYPE_CREATED(ArchetypeHandle)
-		});
-
 	});
 #endif
 }
