[Mass] Command Buffer API extension:
* PushUniqueCommand - the function lets callers add a hand-crafted command. Use it to push commands that you don't want to get merged with other commands of the same type
* AddTag and RemoveTag - added a versions that take an array view of entity handles
* DestroyEntities - added a version with rvalue ref parameter.

diff --git a/Engine/Source/Runtime/MassEntity/Private/MassCommandBuffer.cpp b/Engine/Source/Runtime/MassEntity/Private/MassCommandBuffer.cpp
--- a/Engine/Source/Runtime/MassEntity/Private/MassCommandBuffer.cpp
+++ b/Engine/Source/Runtime/MassEntity/Private/MassCommandBuffer.cpp
@@ -230,7 +230,7 @@ void FMassCommandBuffer::MoveAppend(FMassCommandBuffer& Other)
 {
 	TRACE_CPUPROFILER_EVENT_SCOPE(MassCommandBuffer_MoveAppend);
 
-	// @todo optimize, there surely a way to do faster then this.
+	// @todo optimize, there's surely a way to do this faster than this.
 	UE_MT_SCOPED_READ_ACCESS(Other.PendingBatchCommandsDetector);
 	if (Other.HasPendingCommands())
 	{
diff --git a/Engine/Source/Runtime/MassEntity/Public/MassCommandBuffer.h b/Engine/Source/Runtime/MassEntity/Public/MassCommandBuffer.h
--- a/Engine/Source/Runtime/MassEntity/Public/MassCommandBuffer.h
+++ b/Engine/Source/Runtime/MassEntity/Public/MassCommandBuffer.h
@@ -71,6 +71,26 @@ public:
 		++ActiveCommandsCounter;
 	}
 
+	/**
+	 * Ordinary PushCommand calls try to reuse existing command instances (as stored in CommandInstances)
+	 * based on their type.
+	 * This command lets callers add a manually configured command instance, that 
+	 * might not be distinguishable from other commands based solely on its type.
+	 * For example, when you implement a command type that has member properties that 
+	 * control how the given command works or what exactly it does.
+	 */
+	void PushUniqueCommand(TUniquePtr<FMassBatchedCommand>&& CommandInstance)
+	{
+		COMMAND_PUSHING_CHECK();
+
+		LLM_SCOPE_BYNAME(TEXT("Mass/PushCommand"));
+		UE::TScopeLock Lock(AppendingCommandsCS);
+		UE_MT_SCOPED_WRITE_ACCESS(PendingBatchCommandsDetector);
+		AppendedCommandInstances.Add(MoveTemp(CommandInstance));
+
+		++ActiveCommandsCounter;
+	}
+
 	template<typename T>
 	void AddFragment(FMassEntityHandle Entity)
 	{
@@ -107,6 +127,13 @@ public:
 		PushCommand<FMassCommandAddTagsInternal<EMassCommandCheckTime::CompileTimeCheck, T>>(Entity);
 	}
 
+	template<typename T>
+	void AddTag(TConstArrayView<FMassEntityHandle> Entities)
+	{
+		static_assert(UE::Mass::CTag<T>, "Given struct type is not a valid tag type.");
+		PushCommand<FMassCommandAddTagsInternal<EMassCommandCheckTime::CompileTimeCheck, T>>(Entities);
+	}
+
 	template<typename T>
 	void AddTag_RuntimeCheck(FMassEntityHandle Entity)
 	{
@@ -122,6 +149,13 @@ public:
 		PushCommand<FMassCommandRemoveTagsInternal<EMassCommandCheckTime::CompileTimeCheck, T>>(Entity);
 	}
 
+	template<typename T>
+	void RemoveTag(TConstArrayView<FMassEntityHandle> Entities)
+	{
+		static_assert(UE::Mass::CTag<T>, "Given struct type is not a valid tag type.");
+		PushCommand<FMassCommandRemoveTagsInternal<EMassCommandCheckTime::CompileTimeCheck, T>>(Entities);
+	}
+
 	template<typename T>
 	void RemoveTag_RuntimeCheck(FMassEntityHandle Entity)
 	{
@@ -156,6 +190,11 @@ public:
 		PushCommand<FMassCommandDestroyEntities>(InEntitiesToDestroy);
 	}
 
+	void DestroyEntities(TArray<FMassEntityHandle>&& InEntitiesToDestroy)
+	{
+		PushCommand<FMassCommandDestroyEntities>(MoveTemp(InEntitiesToDestroy));
+	}
+
 	MASSENTITY_API SIZE_T GetAllocatedSize() const;
 
 	/** 
