[Mass] Added FMassFragmentRequirements::AddElementRequirement that can be called at runtime with UScriptStruct* and it will add the provided type to Fragment or Tag bitsets, depending on the type represented.

diff --git a/Engine/Source/Runtime/MassEntity/Public/MassRequirements.h b/Engine/Source/Runtime/MassEntity/Public/MassRequirements.h
--- a/Engine/Source/Runtime/MassEntity/Public/MassRequirements.h
+++ b/Engine/Source/Runtime/MassEntity/Public/MassRequirements.h
@@ -167,6 +167,15 @@ struct FMassFragmentRequirements
 
 	MASSENTITY_API void Initialize(const TSharedRef<FMassEntityManager>& EntityManager);
 
+	FMassFragmentRequirements& AddElementRequirement(TNotNull<const UScriptStruct*> ElementType, const EMassFragmentAccess AccessMode, const EMassFragmentPresence Presence = EMassFragmentPresence::All)
+	{
+		if (UE::Mass::IsA<FMassFragment>(ElementType))
+		{
+			return AddRequirement(ElementType, AccessMode, Presence);
+		}
+		return AddTagRequirement(ElementType, Presence);
+	}
+
 	FMassFragmentRequirements& AddRequirement(const UScriptStruct* FragmentType, const EMassFragmentAccess AccessMode, const EMassFragmentPresence Presence = EMassFragmentPresence::All)
 	{
 		checkf(bInitialized, TEXT("Modifying requirements before initialization is not supported."));
@@ -233,26 +242,32 @@ struct FMassFragmentRequirements
 		return *this;
 	}
 
-	void AddTagRequirement(const UScriptStruct& TagType, const EMassFragmentPresence Presence)
+	FMassFragmentRequirements& AddTagRequirement(TNotNull<const UScriptStruct*> TagType, const EMassFragmentPresence Presence)
 	{
 		checkf(bInitialized, TEXT("Modifying requirements before initialization is not supported."));
 		checkf(int(Presence) != int(EMassFragmentPresence::MAX), TEXT("MAX presence is not a valid value for AddTagRequirement"));
 		switch (Presence)
 		{
 		case EMassFragmentPresence::All:
-			RequiredAllTags.Add(TagType);
+			RequiredAllTags.Add(*TagType);
 			break;
 		case EMassFragmentPresence::Any:
-			RequiredAnyTags.Add(TagType);
+			RequiredAnyTags.Add(*TagType);
 			break;
 		case EMassFragmentPresence::None:
-			RequiredNoneTags.Add(TagType);
+			RequiredNoneTags.Add(*TagType);
 			break;
 		case EMassFragmentPresence::Optional:
-			RequiredOptionalTags.Add(TagType);
+			RequiredOptionalTags.Add(*TagType);
 			break;
 		}
 		IncrementChangeCounter();
+		return *this;
+	}
+
+	void AddTagRequirement(const UScriptStruct& TagType, const EMassFragmentPresence Presence)
+	{
+		AddTagRequirement(&TagType, Presence);
 	}
 
 	template<typename T>
