[Mass] added LexToString(EMassObservedOperationFlags)

diff --git a/Engine/Source/Runtime/MassEntity/Private/MassEntityTypes.cpp b/Engine/Source/Runtime/MassEntity/Private/MassEntityTypes.cpp
--- a/Engine/Source/Runtime/MassEntity/Private/MassEntityTypes.cpp
+++ b/Engine/Source/Runtime/MassEntity/Private/MassEntityTypes.cpp
@@ -15,6 +15,32 @@ DEFINE_TYPEBITSET(FMassSharedFragmentBitSet);
 DEFINE_TYPEBITSET(FMassConstSharedFragmentBitSet);
 DEFINE_TYPEBITSET(FMassExternalSubsystemBitSet);
 
+
+FString LexToString(const EMassObservedOperationFlags Value)
+{
+	switch (Value)
+	{
+	case EMassObservedOperationFlags::None:
+		return ("None");
+	case EMassObservedOperationFlags::AddElement:
+		return ("AddElement");
+	case EMassObservedOperationFlags::RemoveElement:
+		return ("RemoveElement");
+	case EMassObservedOperationFlags::CreateEntity:
+		return ("CreateEntity");
+	case EMassObservedOperationFlags::DestroyEntity:
+		return ("DestroyEntity");
+	case EMassObservedOperationFlags::Add:
+		return ("AddElement | CreateEntity");
+	case EMassObservedOperationFlags::Remove:
+		return ("RemoveElement | DestroyEntity");
+	case EMassObservedOperationFlags::All:
+		return ("Add | Remove");
+	default:
+		return "UNDEFINED";
+	}
+}
+
 //-----------------------------------------------------------------------------
 // FMassArchetypeCompositionDescriptor
 //-----------------------------------------------------------------------------
diff --git a/Engine/Source/Runtime/MassEntity/Public/MassEntityTypes.h b/Engine/Source/Runtime/MassEntity/Public/MassEntityTypes.h
--- a/Engine/Source/Runtime/MassEntity/Public/MassEntityTypes.h
+++ b/Engine/Source/Runtime/MassEntity/Public/MassEntityTypes.h
@@ -548,6 +548,7 @@ enum class EMassObservedOperationFlags : uint8
 	All = Add | Remove,
 };
 ENUM_CLASS_FLAGS(EMassObservedOperationFlags);
+MASSENTITY_API FString LexToString(const EMassObservedOperationFlags Value);
 
 enum class EMassExecutionContextType : uint8
 {
