[Mass] Renamed EntityManager.DebugGetEntityIndexHandle to CreateEntityIndexHandle, since we're going to be using in production too

#rnx

diff --git a/Engine/Plugins/Runtime/MassGameplay/Source/MassGameplayDebug/Private/MassGameplayDebugTypes.cpp b/Engine/Plugins/Runtime/MassGameplay/Source/MassGameplayDebug/Private/MassGameplayDebugTypes.cpp
--- a/Engine/Plugins/Runtime/MassGameplay/Source/MassGameplayDebug/Private/MassGameplayDebugTypes.cpp
+++ b/Engine/Plugins/Runtime/MassGameplay/Source/MassGameplayDebug/Private/MassGameplayDebugTypes.cpp
@@ -38,7 +38,7 @@ void GetDebugEntitiesAndLocations(const FMassEntityManager& EntityManager, TArra
 
 FMassEntityHandle ConvertEntityIndexToHandle(const FMassEntityManager& EntityManager, const int32 EntityIndex)
 {
-	return EntityManager.DebugGetEntityIndexHandle(EntityIndex);
+	return EntityManager.CreateEntityIndexHandle(EntityIndex);
 }
 
 } // namespace UE::Mass::Debug
diff --git a/Engine/Source/Runtime/MassEntity/Private/MassDebugger.cpp b/Engine/Source/Runtime/MassEntity/Private/MassDebugger.cpp
--- a/Engine/Source/Runtime/MassEntity/Private/MassDebugger.cpp
+++ b/Engine/Source/Runtime/MassEntity/Private/MassDebugger.cpp
@@ -325,7 +325,7 @@ namespace UE::Mass::Debug
 		}
 
 		FMassEntityManager& EntityManager = UE::Mass::Utils::GetEntityManagerChecked(*World);
-		FMassEntityHandle EntityToDestroy = EntityManager.DebugGetEntityIndexHandle(ID);
+		FMassEntityHandle EntityToDestroy = EntityManager.CreateEntityIndexHandle(ID);
 		if (!EntityToDestroy.IsSet())
 		{
 			UE_LOG(LogConsoleResponse, Display, TEXT("Error: cannot find entity for this index"));
@@ -357,7 +357,7 @@ namespace UE::Mass::Debug
 			SetDebugEntityRange(ID, ID);
 
 			FMassEntityManager& EntityManager = UE::Mass::Utils::GetEntityManagerChecked(*World);
-			FMassEntityHandle EntityToDebug = EntityManager.DebugGetEntityIndexHandle(ID);
+			FMassEntityHandle EntityToDebug = EntityManager.CreateEntityIndexHandle(ID);
 			if (!EntityToDebug.IsSet() && ID != INDEX_NONE)
 			{
 				UE_LOG(LogConsoleResponse, Display, TEXT("Cannot find entity for this index, clearing current selection"));
diff --git a/Engine/Source/Runtime/MassEntity/Private/MassEntityManager.cpp b/Engine/Source/Runtime/MassEntity/Private/MassEntityManager.cpp
--- a/Engine/Source/Runtime/MassEntity/Private/MassEntityManager.cpp
+++ b/Engine/Source/Runtime/MassEntity/Private/MassEntityManager.cpp
@@ -1857,6 +1857,14 @@ void FMassEntityManager::CheckIfEntityIsActive(FMassEntityHandle Entity) const
 	checkf(IsEntityBuilt(Entity), TEXT("Entity not yet created(ID: %d, SN:%d)"), Entity.Index, Entity.SerialNumber);
 }
 
+FMassEntityHandle FMassEntityManager::CreateEntityIndexHandle(const int32 EntityIndex) const
+{
+	return (GetEntityStorageInterface().IsValidIndex(EntityIndex)
+		&& GetEntityStorageInterface().GetEntityState(EntityIndex) == UE::Mass::IEntityStorageInterface::EEntityState::Created)
+		? FMassEntityHandle(EntityIndex, GetEntityStorageInterface().GetSerialNumber(EntityIndex))
+		: FMassEntityHandle();
+}
+
 void FMassEntityManager::GetMatchingArchetypes(const FMassFragmentRequirements& Requirements, TArray<FMassArchetypeHandle>& OutValidArchetypes, const uint32 FromArchetypeDataVersion) const
 {
 	for (int32 ArchetypeIndex = FromArchetypeDataVersion; ArchetypeIndex < AllArchetypes.Num(); ++ArchetypeIndex)
@@ -2120,10 +2128,7 @@ void FMassEntityManager::DebugGetArchetypeStrings(const FMassArchetypeHandle& Ar
 
 FMassEntityHandle FMassEntityManager::DebugGetEntityIndexHandle(const int32 EntityIndex) const
 {
-	return (GetEntityStorageInterface().IsValidIndex(EntityIndex)
-		&& GetEntityStorageInterface().GetEntityState(EntityIndex) == UE::Mass::IEntityStorageInterface::EEntityState::Created)
-		? FMassEntityHandle(EntityIndex, GetEntityStorageInterface().GetSerialNumber(EntityIndex))
-		: FMassEntityHandle();
+	return CreateEntityIndexHandle(EntityIndex);
 }
 
 const FString& FMassEntityManager::DebugGetName() const
diff --git a/Engine/Source/Runtime/MassEntity/Public/MassEntityManager.h b/Engine/Source/Runtime/MassEntity/Public/MassEntityManager.h
--- a/Engine/Source/Runtime/MassEntity/Public/MassEntityManager.h
+++ b/Engine/Source/Runtime/MassEntity/Public/MassEntityManager.h
@@ -520,6 +520,11 @@ public:
 	/** Asserts that IsEntityBuilt */
 	UE_API void CheckIfEntityIsActive(FMassEntityHandle EntityHandle) const;
 
+	/**
+	 * Generate valid, up-to-date entity handle for the entity at given index.
+	 */
+	UE_API FMassEntityHandle CreateEntityIndexHandle(const int32 EntityIndex) const;
+
 	template<typename FragmentType>
 	FragmentType& GetFragmentDataChecked(FMassEntityHandle EntityHandle) const
 	{
@@ -1023,6 +1028,7 @@ PRAGMA_ENABLE_DEPRECATION_WARNINGS
 	UE_API void DebugRemoveAllEntities();
 	UE_API void DebugForceArchetypeDataVersionBump();
 	UE_API void DebugGetArchetypeStrings(const FMassArchetypeHandle& Archetype, TArray<FName>& OutFragmentNames, TArray<FName>& OutTagNames);
+	UE_DEPRECATED(5.7, "Using CreateEntityIndexHandle instead.")
 	UE_API FMassEntityHandle DebugGetEntityIndexHandle(const int32 EntityIndex) const;
 	UE_API const FString& DebugGetName() const;
 
