[Mass] Added a couple of functions to MassEntityManager, that deal with types passed in as UScriptStruct, and pass the argument to the appropriate function, based on the type (Tag or Fragment).

diff --git a/Engine/Source/Runtime/MassEntity/Private/MassEntityManager.cpp b/Engine/Source/Runtime/MassEntity/Private/MassEntityManager.cpp
--- a/Engine/Source/Runtime/MassEntity/Private/MassEntityManager.cpp
+++ b/Engine/Source/Runtime/MassEntity/Private/MassEntityManager.cpp
@@ -21,6 +21,8 @@
 
 #define CHECK_SYNC_API() testableCheckf(IsProcessing() == false, TEXT("Synchronous API function %hs called during mass processing. Use asynchronous API instead."), __FUNCTION__)
 #define CHECK_SYNC_API_RETURN(ReturnValue) testableCheckfReturn(IsProcessing() == false, return ReturnValue, TEXT("Synchronous API function %hs called during mass processing. Use asynchronous API instead."), __FUNCTION__)
+#define CHECK_ELEMENT(ElementType) checkf(UE::Mass::Private::IsElement(ElementType), TEXT("%hs: Only tags and fragments are considered 'elements', type provided: %s") \
+		, __FUNCTION__ , *ElementType->GetName())
 
 const FMassEntityHandle FMassEntityManager::InvalidEntity;
 
@@ -46,6 +48,12 @@ namespace UE::Mass::Private
 			}
 		}
 	}
+
+	bool IsElement(TNotNull<const UScriptStruct*> ElementType)
+	{
+		return UE::Mass::IsA<FMassFragment>(ElementType)
+			|| UE::Mass::IsA<FMassTag>(ElementType);
+	}
 }
 
 //-----------------------------------------------------------------------------
@@ -1311,6 +1319,83 @@ void FMassEntityManager::RemoveTagFromEntity(FMassEntityHandle Entity, const USc
 	}
 }
 
+bool FMassEntityManager::DoesEntityHaveElement(FMassEntityHandle Entity, TNotNull<const UScriptStruct*> ElementType) const
+{
+	CHECK_ELEMENT(ElementType);
+
+	CheckIfEntityIsActive(Entity);
+	const TSharedPtr<FMassArchetypeData>& CurrentArchetype = GetEntityStorageInterface().GetArchetypeAsShared(Entity.Index);
+	check(CurrentArchetype);
+	if (UE::Mass::IsA<FMassTag>(ElementType))
+	{
+		return CurrentArchetype->GetCompositionDescriptor().Tags.Contains(*ElementType);
+	}
+	
+	return CurrentArchetype->GetCompositionDescriptor().Fragments.Contains(*ElementType);
+}
+
+void FMassEntityManager::AddElementToEntities(TConstArrayView<FMassEntityHandle> Entities, TNotNull<const UScriptStruct*> ElementType)
+{
+	CHECK_ELEMENT(ElementType);
+
+	TArray<FMassArchetypeEntityCollection> EntityCollections;
+	UE::Mass::Utils::CreateEntityCollections(*this, Entities, FMassArchetypeEntityCollection::EDuplicatesHandling::FoldDuplicates, EntityCollections);
+
+	if (UE::Mass::IsA<FMassTag>(ElementType))
+	{
+		BatchChangeTagsForEntities(EntityCollections, FMassTagBitSet(*ElementType), {});
+	}
+	else
+	{
+		BatchChangeFragmentCompositionForEntities(EntityCollections, FMassFragmentBitSet(*ElementType), {});
+	}
+}
+
+void FMassEntityManager::AddElementToEntity(FMassEntityHandle Entity, TNotNull<const UScriptStruct*> ElementType)
+{
+	CHECK_ELEMENT(ElementType);
+
+	if (UE::Mass::IsA<FMassTag>(ElementType))
+	{
+		AddTagToEntity(Entity, ElementType);
+	}
+	else
+	{
+		AddFragmentToEntity(Entity, ElementType);
+	}
+}
+
+void FMassEntityManager::RemoveElementFromEntities(TConstArrayView<FMassEntityHandle> Entities, TNotNull<const UScriptStruct*> ElementType)
+{
+	CHECK_ELEMENT(ElementType);
+
+	TArray<FMassArchetypeEntityCollection> EntityCollections;
+	UE::Mass::Utils::CreateEntityCollections(*this, Entities, FMassArchetypeEntityCollection::EDuplicatesHandling::FoldDuplicates, EntityCollections);
+
+	if (UE::Mass::IsA<FMassTag>(ElementType))
+	{
+		BatchChangeTagsForEntities(EntityCollections, {}, FMassTagBitSet(*ElementType));
+	}
+	else
+	{
+		BatchChangeFragmentCompositionForEntities(EntityCollections, {}, FMassFragmentBitSet(*ElementType));
+	}
+}
+
+void FMassEntityManager::RemoveElementFromEntity(FMassEntityHandle Entity, TNotNull<const UScriptStruct*> ElementType)
+{
+	CHECK_ELEMENT(ElementType);
+
+	if (UE::Mass::IsA<FMassTag>(ElementType))
+	{
+		RemoveTagFromEntity(Entity, ElementType);
+	}
+	else
+	{
+		RemoveFragmentFromEntity(Entity, ElementType);
+	}
+}
+
 bool FMassEntityManager::AddConstSharedFragmentToEntity(const FMassEntityHandle Entity, const FConstSharedStruct& InConstSharedFragment)
 {
 	CHECK_SYNC_API_RETURN(false);
@@ -2197,3 +2282,5 @@ void FMassEntityManager::BatchSetEntityFragmentsValues(TConstArrayView<FMassArch
 }
 
 #undef CHECK_SYNC_API
+#undef CHECK_SYNC_API_RETURN
+#undef CHECK_ELEMENT
\ No newline at end of file
diff --git a/Engine/Source/Runtime/MassEntity/Public/MassEntityManager.h b/Engine/Source/Runtime/MassEntity/Public/MassEntityManager.h
--- a/Engine/Source/Runtime/MassEntity/Public/MassEntityManager.h
+++ b/Engine/Source/Runtime/MassEntity/Public/MassEntityManager.h
@@ -386,6 +386,33 @@ public:
 	UE_API void RemoveTagFromEntity(FMassEntityHandle EntityHandle, const UScriptStruct* TagType);
 	UE_API void SwapTagsForEntity(FMassEntityHandle EntityHandle, const UScriptStruct* FromFragmentType, const UScriptStruct* ToFragmentType);
 
+	/**
+	 * Adds ElementType to the entities, treating it accordingly based on which element type it represents (i.e. Fragment or Tag).
+	 * The function will assert on unhandled types.
+	 */
+	void AddElementToEntities(TConstArrayView<FMassEntityHandle> Entities, TNotNull<const UScriptStruct*> ElementType);
+
+	/**
+	 * Adds ElementType to the entity, treating it accordingly based on which element type it represents (i.e. Fragment or Tag).
+	 * The function will assert on unhandled types.
+	 */
+	void AddElementToEntity(FMassEntityHandle Entity, TNotNull<const UScriptStruct*> ElementType);
+
+	/*
+	 * Removes ElementType from the entities, treating it accordingly based on which element type it represents (i.e. Fragment or Tag).
+	 * The function will assert on unhandled types.
+	 */
+	void RemoveElementFromEntities(TConstArrayView<FMassEntityHandle> Entities, TNotNull<const UScriptStruct*> ElementType);
+
+	/*
+	 * Removes ElementType from the entity, treating it accordingly based on which element type it represents (i.e. Fragment or Tag).
+	 * The function will assert on unhandled types.
+	 */
+	void RemoveElementFromEntity(FMassEntityHandle Entity, TNotNull<const UScriptStruct*> ElementType);
+
+	/** @return whether Entity has an element of ElementType */
+	bool DoesEntityHaveElement(FMassEntityHandle Entity, TNotNull<const UScriptStruct*> ElementType) const;
+
 	/** 
 	 * Adds a new const shared fragment to the given entity. Note that it only works if the given entity doesn't have
 	 * a shared fragment of the given type. The function will give a soft "pass" if the entity has the shared fragment
