[Mass] Added a command that relies on on the recently added `PushUniqueCommand` mechanism. Also, added "command traits" to be able to mark commands that can be pushed only with that mechanism (i.e. it's unsafe or unreliable to use the type-base identification for them).

#rnx

diff --git a/Engine/Source/Runtime/MassEntity/Public/MassCommandBuffer.h b/Engine/Source/Runtime/MassEntity/Public/MassCommandBuffer.h
--- a/Engine/Source/Runtime/MassEntity/Public/MassCommandBuffer.h
+++ b/Engine/Source/Runtime/MassEntity/Public/MassCommandBuffer.h
@@ -235,6 +235,7 @@ private:
 	template<typename T>
 	T& CreateOrAddCommand()
 	{
+		static_assert(!UE::Mass::Command::TCommandTraits<T>::RequiresUniqueHandling, "This command type needs to be added via PushUniqueCommand");
 		const int32 Index = FMassBatchedCommand::GetCommandIndex<T>();
 
 		if (CommandInstances.IsValidIndex(Index) == false)
diff --git a/Engine/Source/Runtime/MassEntity/Public/MassCommands.h b/Engine/Source/Runtime/MassEntity/Public/MassCommands.h
--- a/Engine/Source/Runtime/MassEntity/Public/MassCommands.h
+++ b/Engine/Source/Runtime/MassEntity/Public/MassCommands.h
@@ -75,6 +75,18 @@ namespace UE::Mass::Utils
 	}
 } // namespace UE::Mass::Utils
 
+namespace UE::Mass::Command
+{
+	template<typename T>
+	struct TCommandTraits final
+	{
+		enum
+		{
+			RequiresUniqueHandling = false
+		};
+	};
+};
+
 struct FMassBatchedCommand
 {
 	FMassBatchedCommand() = default;
@@ -388,6 +400,38 @@ protected:
 	const FMassFragmentBitSet FragmentsAffected;
 };
 
+/**
+ * Command capable of adding any element type, be it a fragment or a tag.
+ * Note that this type of command can only be added via PushUniqueCommand
+ */
+struct FMassCommandAddElement : public FMassBatchedEntityCommand
+{
+	using Super = FMassBatchedEntityCommand;
+
+	FMassCommandAddElement(const TNotNull<const UScriptStruct*> InElementType)
+		: Super(EMassCommandOperationType::Add DEBUG_NAME("AddElement"))
+		, ElementType(InElementType)
+	{}
+
+protected:
+	virtual void Execute(FMassEntityManager& EntityManager) const override
+	{
+		TRACE_CPUPROFILER_EVENT_SCOPE(MassCommandAddElement_Execute);
+		EntityManager.AddElementToEntities(TargetEntities, ElementType);
+	}
+
+	TNotNull<const UScriptStruct*> ElementType;
+};
+
+template<>
+struct UE::Mass::Command::TCommandTraits<FMassCommandAddElement> final
+{
+	enum
+	{
+		RequiresUniqueHandling = true
+	};
+};;
+
 template<typename... TOthers>
 struct FMassCommandBuildEntity : public FMassCommandAddFragmentInstances<TOthers...>
 {
