[Mass] Minor MassEntityManager cleanup, including removing some header inclusion

#rnx

diff --git a/Engine/Plugins/AI/MassAI/Source/MassAIBehavior/Public/MassLookAtProcessors.h b/Engine/Plugins/AI/MassAI/Source/MassAIBehavior/Public/MassLookAtProcessors.h
--- a/Engine/Plugins/AI/MassAI/Source/MassAIBehavior/Public/MassLookAtProcessors.h
+++ b/Engine/Plugins/AI/MassAI/Source/MassAIBehavior/Public/MassLookAtProcessors.h
@@ -4,6 +4,7 @@
 
 #include "MassObserverProcessor.h"
 #include "MassProcessor.h"
+#include "MassEntityQuery.h"
 #if UE_ENABLE_INCLUDE_ORDER_DEPRECATED_IN_5_6
 #include "MassSignalProcessorBase.h"
 #include "MassCommonTypes.h"
diff --git a/Engine/Plugins/AI/MassAI/Source/MassAIDebug/Public/MassDebugStateTreeProcessor.h b/Engine/Plugins/AI/MassAI/Source/MassAIDebug/Public/MassDebugStateTreeProcessor.h
--- a/Engine/Plugins/AI/MassAI/Source/MassAIDebug/Public/MassDebugStateTreeProcessor.h
+++ b/Engine/Plugins/AI/MassAI/Source/MassAIDebug/Public/MassDebugStateTreeProcessor.h
@@ -3,6 +3,7 @@
 #pragma once
 
 #include "MassProcessor.h"
+#include "MassEntityQuery.h"
 #include "MassDebugStateTreeProcessor.generated.h"
 
 #define UE_API MASSAIDEBUG_API
diff --git a/Engine/Plugins/AI/MassAI/Source/MassNavMeshNavigation/Public/MassNavMeshNavigationBoundaryProcessor.h b/Engine/Plugins/AI/MassAI/Source/MassNavMeshNavigation/Public/MassNavMeshNavigationBoundaryProcessor.h
--- a/Engine/Plugins/AI/MassAI/Source/MassNavMeshNavigation/Public/MassNavMeshNavigationBoundaryProcessor.h
+++ b/Engine/Plugins/AI/MassAI/Source/MassNavMeshNavigation/Public/MassNavMeshNavigationBoundaryProcessor.h
@@ -3,6 +3,7 @@
 #pragma once
 
 #include "MassProcessor.h"
+#include "MassEntityQuery.h"
 #include "MassNavMeshNavigationBoundaryProcessor.generated.h"
 
 #define UE_API MASSNAVMESHNAVIGATION_API
diff --git a/Engine/Plugins/AI/MassAI/Source/MassNavMeshNavigation/Public/MassNavMeshNavigationProcessors.h b/Engine/Plugins/AI/MassAI/Source/MassNavMeshNavigation/Public/MassNavMeshNavigationProcessors.h
--- a/Engine/Plugins/AI/MassAI/Source/MassNavMeshNavigation/Public/MassNavMeshNavigationProcessors.h
+++ b/Engine/Plugins/AI/MassAI/Source/MassNavMeshNavigation/Public/MassNavMeshNavigationProcessors.h
@@ -3,6 +3,7 @@
 #pragma once
 
 #include "MassProcessor.h"
+#include "MassEntityQuery.h"
 #include "MassNavMeshNavigationProcessors.generated.h"
 
 #define UE_API MASSNAVMESHNAVIGATION_API
diff --git a/Engine/Plugins/AI/MassAI/Source/MassNavigation/Public/Avoidance/MassAvoidanceProcessors.h b/Engine/Plugins/AI/MassAI/Source/MassNavigation/Public/Avoidance/MassAvoidanceProcessors.h
--- a/Engine/Plugins/AI/MassAI/Source/MassNavigation/Public/Avoidance/MassAvoidanceProcessors.h
+++ b/Engine/Plugins/AI/MassAI/Source/MassNavigation/Public/Avoidance/MassAvoidanceProcessors.h
@@ -3,6 +3,7 @@
 #pragma once
 
 #include "MassProcessor.h"
+#include "MassEntityQuery.h"
 #include "MassAvoidanceProcessors.generated.h"
 
 #define UE_API MASSNAVIGATION_API
diff --git a/Engine/Plugins/AI/MassAI/Source/MassZoneGraphNavigation/Public/MassZoneGraphNavigationProcessors.h b/Engine/Plugins/AI/MassAI/Source/MassZoneGraphNavigation/Public/MassZoneGraphNavigationProcessors.h
--- a/Engine/Plugins/AI/MassAI/Source/MassZoneGraphNavigation/Public/MassZoneGraphNavigationProcessors.h
+++ b/Engine/Plugins/AI/MassAI/Source/MassZoneGraphNavigation/Public/MassZoneGraphNavigationProcessors.h
@@ -4,6 +4,7 @@
 
 #include "MassProcessor.h"
 #include "MassObserverProcessor.h"
+#include "MassEntityQuery.h"
 #include "MassZoneGraphNavigationProcessors.generated.h"
 
 #define UE_API MASSZONEGRAPHNAVIGATION_API
diff --git a/Engine/Plugins/Runtime/InstancedActors/Source/InstancedActors/Public/InstancedActorsVisualizationSwitcherProcessor.h b/Engine/Plugins/Runtime/InstancedActors/Source/InstancedActors/Public/InstancedActorsVisualizationSwitcherProcessor.h
--- a/Engine/Plugins/Runtime/InstancedActors/Source/InstancedActors/Public/InstancedActorsVisualizationSwitcherProcessor.h
+++ b/Engine/Plugins/Runtime/InstancedActors/Source/InstancedActors/Public/InstancedActorsVisualizationSwitcherProcessor.h
@@ -4,6 +4,7 @@
 
 #include "MassProcessor.h"
 #include "MassRepresentationTypes.h"
+#include "MassEntityQuery.h"
 #include "InstancedActorsVisualizationSwitcherProcessor.generated.h"
 
 
diff --git a/Engine/Plugins/Runtime/MassGameplay/Source/MassActors/Public/Translators/MassCapsuleComponentTranslators.h b/Engine/Plugins/Runtime/MassGameplay/Source/MassActors/Public/Translators/MassCapsuleComponentTranslators.h
--- a/Engine/Plugins/Runtime/MassGameplay/Source/MassActors/Public/Translators/MassCapsuleComponentTranslators.h
+++ b/Engine/Plugins/Runtime/MassGameplay/Source/MassActors/Public/Translators/MassCapsuleComponentTranslators.h
@@ -4,6 +4,7 @@
 
 #include "MassCommonFragments.h"
 #include "MassTranslator.h"
+#include "MassEntityQuery.h"
 #include "MassCapsuleComponentTranslators.generated.h"
 
 #define UE_API MASSACTORS_API
diff --git a/Engine/Plugins/Runtime/MassGameplay/Source/MassEQS/Public/Processors/MassEnvQueryGeneratorProcessor_MassEntityHandles.h b/Engine/Plugins/Runtime/MassGameplay/Source/MassEQS/Public/Processors/MassEnvQueryGeneratorProcessor_MassEntityHandles.h
--- a/Engine/Plugins/Runtime/MassGameplay/Source/MassEQS/Public/Processors/MassEnvQueryGeneratorProcessor_MassEntityHandles.h
+++ b/Engine/Plugins/Runtime/MassGameplay/Source/MassEQS/Public/Processors/MassEnvQueryGeneratorProcessor_MassEntityHandles.h
@@ -3,6 +3,7 @@
 #pragma once
 
 #include "Processors/MassEnvQueryProcessorBase.h"
+#include "MassEntityQuery.h"
 #include "MassEnvQueryGeneratorProcessor_MassEntityHandles.generated.h"
 
 /** Processor for completing MassEQSSubsystem Requests sent from UMassEnvQueryGenerator_MassEntityHandles */
diff --git a/Engine/Plugins/Runtime/MassGameplay/Source/MassGameplayDebug/Public/AssignDebugVisProcessor.h b/Engine/Plugins/Runtime/MassGameplay/Source/MassGameplayDebug/Public/AssignDebugVisProcessor.h
--- a/Engine/Plugins/Runtime/MassGameplay/Source/MassGameplayDebug/Public/AssignDebugVisProcessor.h
+++ b/Engine/Plugins/Runtime/MassGameplay/Source/MassGameplayDebug/Public/AssignDebugVisProcessor.h
@@ -3,6 +3,7 @@
 #pragma once
 
 #include "MassObserverProcessor.h"
+#include "MassEntityQuery.h"
 #include "AssignDebugVisProcessor.generated.h"
 
 #define UE_API MASSGAMEPLAYDEBUG_API
diff --git a/Engine/Plugins/Runtime/MassGameplay/Source/MassGameplayDebug/Public/DebugVisLocationProcessor.h b/Engine/Plugins/Runtime/MassGameplay/Source/MassGameplayDebug/Public/DebugVisLocationProcessor.h
--- a/Engine/Plugins/Runtime/MassGameplay/Source/MassGameplayDebug/Public/DebugVisLocationProcessor.h
+++ b/Engine/Plugins/Runtime/MassGameplay/Source/MassGameplayDebug/Public/DebugVisLocationProcessor.h
@@ -3,6 +3,7 @@
 #pragma once
 
 #include "MassProcessor.h"
+#include "MassEntityQuery.h"
 #include "DebugVisLocationProcessor.generated.h"
 
 #define UE_API MASSGAMEPLAYDEBUG_API
diff --git a/Engine/Plugins/Runtime/MassGameplay/Source/MassLOD/Public/MassLODCollectorProcessor.h b/Engine/Plugins/Runtime/MassGameplay/Source/MassLOD/Public/MassLODCollectorProcessor.h
--- a/Engine/Plugins/Runtime/MassGameplay/Source/MassLOD/Public/MassLODCollectorProcessor.h
+++ b/Engine/Plugins/Runtime/MassGameplay/Source/MassLOD/Public/MassLODCollectorProcessor.h
@@ -7,6 +7,7 @@
 #endif // UE_ENABLE_INCLUDE_ORDER_DEPRECATED_IN_5_6
 #include "MassProcessor.h"
 #include "MassLODCollector.h"
+#include "MassEntityQuery.h"
 #include "MassLODCollectorProcessor.generated.h"
 
 
diff --git a/Engine/Plugins/Runtime/MassGameplay/Source/MassLOD/Public/MassLODDistanceCollectorProcessor.h b/Engine/Plugins/Runtime/MassGameplay/Source/MassLOD/Public/MassLODDistanceCollectorProcessor.h
--- a/Engine/Plugins/Runtime/MassGameplay/Source/MassLOD/Public/MassLODDistanceCollectorProcessor.h
+++ b/Engine/Plugins/Runtime/MassGameplay/Source/MassLOD/Public/MassLODDistanceCollectorProcessor.h
@@ -8,6 +8,7 @@
 #include "MassProcessor.h"
 #include "MassLODCollector.h"
 #include "MassLODLogic.h"
+#include "MassEntityQuery.h"
 #include "MassLODDistanceCollectorProcessor.generated.h"
 
 /*
diff --git a/Engine/Plugins/Runtime/MassGameplay/Source/MassMovement/Public/Movement/MassMovementProcessors.h b/Engine/Plugins/Runtime/MassGameplay/Source/MassMovement/Public/Movement/MassMovementProcessors.h
--- a/Engine/Plugins/Runtime/MassGameplay/Source/MassMovement/Public/Movement/MassMovementProcessors.h
+++ b/Engine/Plugins/Runtime/MassGameplay/Source/MassMovement/Public/Movement/MassMovementProcessors.h
@@ -4,6 +4,7 @@
 
 #include "MassProcessor.h"
 #include "MassCommonTypes.h" // For WITH_MASSGAMEPLAY_DEBUG
+#include "MassEntityQuery.h"
 #include "MassMovementProcessors.generated.h"
 
 #define UE_API MASSMOVEMENT_API
diff --git a/Engine/Plugins/Runtime/MassGameplay/Source/MassReplication/Public/MassClientBubbleHandler.h b/Engine/Plugins/Runtime/MassGameplay/Source/MassReplication/Public/MassClientBubbleHandler.h
--- a/Engine/Plugins/Runtime/MassGameplay/Source/MassReplication/Public/MassClientBubbleHandler.h
+++ b/Engine/Plugins/Runtime/MassGameplay/Source/MassReplication/Public/MassClientBubbleHandler.h
@@ -12,6 +12,7 @@
 #include "Engine/World.h"
 #include "MassEntityTemplate.h"
 #include "MassEntityView.h"
+#include "MassEntityQuery.h"
 #if UE_ENABLE_INCLUDE_ORDER_DEPRECATED_IN_5_6
 #include "MassEntityManager.h"
 #include "MassSpawnerTypes.h"
diff --git a/Engine/Plugins/Runtime/MassGameplay/Source/MassRepresentation/Private/MassDistanceLODProcessor.cpp b/Engine/Plugins/Runtime/MassGameplay/Source/MassRepresentation/Private/MassDistanceLODProcessor.cpp
--- a/Engine/Plugins/Runtime/MassGameplay/Source/MassRepresentation/Private/MassDistanceLODProcessor.cpp
+++ b/Engine/Plugins/Runtime/MassGameplay/Source/MassRepresentation/Private/MassDistanceLODProcessor.cpp
@@ -6,6 +6,7 @@
 #include "MassCommonFragments.h"
 #include "MassLODFragments.h"
 #include "MassLODTypes.h"
+#include "MassEntityQuery.h"
 #include "MassRepresentationFragments.h"
 
 UMassDistanceLODProcessor::UMassDistanceLODProcessor()
diff --git a/Engine/Plugins/Runtime/MassGameplay/Source/MassRepresentation/Public/MassDistanceLODProcessor.h b/Engine/Plugins/Runtime/MassGameplay/Source/MassRepresentation/Public/MassDistanceLODProcessor.h
--- a/Engine/Plugins/Runtime/MassGameplay/Source/MassRepresentation/Public/MassDistanceLODProcessor.h
+++ b/Engine/Plugins/Runtime/MassGameplay/Source/MassRepresentation/Public/MassDistanceLODProcessor.h
@@ -3,6 +3,7 @@
 #pragma once
 
 #include "MassProcessor.h"
+#include "MassEntityQuery.h"
 #if UE_ENABLE_INCLUDE_ORDER_DEPRECATED_IN_5_6
 #include "MassRepresentationFragments.h"
 #include "MassCommonFragments.h"
diff --git a/Engine/Plugins/Runtime/MassGameplay/Source/MassSignals/Public/MassSignalProcessorBase.h b/Engine/Plugins/Runtime/MassGameplay/Source/MassSignals/Public/MassSignalProcessorBase.h
--- a/Engine/Plugins/Runtime/MassGameplay/Source/MassSignals/Public/MassSignalProcessorBase.h
+++ b/Engine/Plugins/Runtime/MassGameplay/Source/MassSignals/Public/MassSignalProcessorBase.h
@@ -7,6 +7,7 @@
 #include "MassSignalTypes.h"
 #include "Containers/StaticArray.h"
 #include "Misc/TransactionallySafeRWLock.h"
+#include "MassEntityQuery.h"
 #include "MassSignalProcessorBase.generated.h"
 
 #define UE_API MASSSIGNALS_API
diff --git a/Engine/Plugins/Runtime/MassGameplay/Source/MassSmartObjects/Public/MassSmartObjectProcessor.h b/Engine/Plugins/Runtime/MassGameplay/Source/MassSmartObjects/Public/MassSmartObjectProcessor.h
--- a/Engine/Plugins/Runtime/MassGameplay/Source/MassSmartObjects/Public/MassSmartObjectProcessor.h
+++ b/Engine/Plugins/Runtime/MassGameplay/Source/MassSmartObjects/Public/MassSmartObjectProcessor.h
@@ -4,6 +4,7 @@
 
 #include "MassObserverProcessor.h"
 #include "MassProcessor.h"
+#include "MassEntityQuery.h"
 #include "MassSmartObjectProcessor.generated.h"
 
 #define UE_API MASSSMARTOBJECTS_API
diff --git a/Engine/Plugins/Runtime/MassGameplay/Source/MassSmartObjects/Public/MassSmartObjectRegistration.h b/Engine/Plugins/Runtime/MassGameplay/Source/MassSmartObjects/Public/MassSmartObjectRegistration.h
--- a/Engine/Plugins/Runtime/MassGameplay/Source/MassSmartObjects/Public/MassSmartObjectRegistration.h
+++ b/Engine/Plugins/Runtime/MassGameplay/Source/MassSmartObjects/Public/MassSmartObjectRegistration.h
@@ -5,6 +5,7 @@
 #include "MassObserverProcessor.h"
 #include "MassSignalProcessorBase.h"
 #include "SmartObjectTypes.h"
+#include "MassEntityQuery.h"
 #include "MassSmartObjectRegistration.generated.h"
 
 class USmartObjectDefinition;
diff --git a/Engine/Plugins/Runtime/MassGameplay/Source/MassSpawner/Public/MassSpawnLocationProcessor.h b/Engine/Plugins/Runtime/MassGameplay/Source/MassSpawner/Public/MassSpawnLocationProcessor.h
--- a/Engine/Plugins/Runtime/MassGameplay/Source/MassSpawner/Public/MassSpawnLocationProcessor.h
+++ b/Engine/Plugins/Runtime/MassGameplay/Source/MassSpawner/Public/MassSpawnLocationProcessor.h
@@ -3,6 +3,7 @@
 #pragma once 
 
 #include "MassProcessor.h"
+#include "MassEntityQuery.h"
 #include "MassSpawnLocationProcessor.generated.h"
 
 #define UE_API MASSSPAWNER_API
diff --git a/Engine/Source/Runtime/MassEntity/Private/MassEntityManager.cpp b/Engine/Source/Runtime/MassEntity/Private/MassEntityManager.cpp
--- a/Engine/Source/Runtime/MassEntity/Private/MassEntityManager.cpp
+++ b/Engine/Source/Runtime/MassEntity/Private/MassEntityManager.cpp
@@ -180,7 +180,8 @@ void FMassEntityManager::Initialize(const FMassEntityManagerStorageInitParams& I
 		CommandBuffer = MakeShareable(new FMassCommandBuffer());
 	}
 
-	// if we get forked we need to update the command buffer's CurrentThreadID
+	// if we get forked we need to update the CurrentThreadID of things like command buffers
+	// and potentially active creation context
 	if (FForkProcessHelper::IsForkRequested())
 	{
 		OnPostForkHandle = FCoreDelegates::OnPostFork.AddSP(AsShared(), &FMassEntityManager::OnPostFork);
@@ -416,11 +417,8 @@ FMassArchetypeHandle FMassEntityManager::InternalCreateSimilarArchetype(const TS
 {
 	checkSlow(SourceArchetype.IsValid());
 	const FMassArchetypeData& SourceArchetypeRef = *SourceArchetype.Get();
-	FMassArchetypeCompositionDescriptor NewComposition(SourceArchetypeRef.GetFragmentBitSet()
-		, OverrideTags
-		, SourceArchetypeRef.GetChunkFragmentBitSet()
-		, SourceArchetypeRef.GetSharedFragmentBitSet()
-		, SourceArchetypeRef.GetConstSharedFragmentBitSet());
+	FMassArchetypeCompositionDescriptor NewComposition(SourceArchetypeRef.GetCompositionDescriptor());
+	NewComposition.Tags = OverrideTags;
 
 	return InternalCreateSimilarArchetype(SourceArchetypeRef, MoveTemp(NewComposition), SourceArchetypeRef.GetGroups());
 }
@@ -429,11 +427,8 @@ FMassArchetypeHandle FMassEntityManager::InternalCreateSimilarArchetype(const TS
 {
 	checkSlow(SourceArchetype.IsValid());
 	const FMassArchetypeData& SourceArchetypeRef = *SourceArchetype.Get();
-	FMassArchetypeCompositionDescriptor NewComposition(OverrideFragments
-		, SourceArchetypeRef.GetTagBitSet()
-		, SourceArchetypeRef.GetChunkFragmentBitSet()
-		, SourceArchetypeRef.GetSharedFragmentBitSet()
-		, SourceArchetypeRef.GetConstSharedFragmentBitSet());
+	FMassArchetypeCompositionDescriptor NewComposition(SourceArchetypeRef.GetCompositionDescriptor());
+	NewComposition.Fragments = OverrideFragments;
 
 	return InternalCreateSimilarArchetype(SourceArchetypeRef, MoveTemp(NewComposition), SourceArchetypeRef.GetGroups());
 }
@@ -2074,6 +2069,9 @@ UE::Mass::FEntityBuilder FMassEntityManager::MakeEntityBuilder()
 	return UE::Mass::FEntityBuilder(AsShared());
 }
 
+//-----------------------------------------------------------------------------
+// DEBUG stuff
+//-----------------------------------------------------------------------------
 bool FMassEntityManager::DebugDoCollectionsOverlapCreationContext(TConstArrayView<FMassArchetypeEntityCollection> EntityCollections) const
 {
 	if (TSharedPtr<FMassObserverManager::FCreationContext> AsSharedPtr = ObserverManager.GetCreationContext())
diff --git a/Engine/Source/Runtime/MassEntity/Public/MassEntityManager.h b/Engine/Source/Runtime/MassEntity/Public/MassEntityManager.h
--- a/Engine/Source/Runtime/MassEntity/Public/MassEntityManager.h
+++ b/Engine/Source/Runtime/MassEntity/Public/MassEntityManager.h
@@ -6,20 +6,27 @@
 #include "Containers/StaticArray.h"
 #include "MassEntityTypes.h"
 #include "MassProcessingTypes.h"
+#if UE_ENABLE_INCLUDE_ORDER_DEPRECATED_IN_5_7
 #include "MassEntityQuery.h"
+#endif
 #if UE_ENABLE_INCLUDE_ORDER_DEPRECATED_IN_5_5
 #include "StructUtils/InstancedStruct.h"
 #include "StructUtils/StructUtilsTypes.h"
 #endif
 #include "MassObserverManager.h"
 #include "Containers/MpscQueue.h"
+#if UE_ENABLE_INCLUDE_ORDER_DEPRECATED_IN_5_7
 #include "MassRequirementAccessDetector.h"
+#endif // UE_ENABLE_INCLUDE_ORDER_DEPRECATED_IN_5_7
 #include "Templates/FunctionFwd.h"
 #include "MassEntityManagerStorage.h"
 #include "MassArchetypeGroup.h"
 
 #if WITH_MASSENTITY_DEBUG
 #include "Misc/EnumClassFlags.h"
+#include "MassRequirementAccessDetector.h"
+#else
+struct FMassRequirementAccessDetector;
 #endif
 
 #define UE_API MASSENTITY_API
@@ -32,6 +39,7 @@ struct FMassCommandBuffer;
 struct FMassArchetypeEntityCollection;
 class FOutputDevice;
 struct FMassDebugger;
+struct FMassFragmentRequirements;
 enum class EMassFragmentAccess : uint8;
 enum class EForkProcessRole : uint8;
 namespace UE::Mass
@@ -744,6 +752,9 @@ public:
 	}
 #endif
 
+	//-----------------------------------------------------------------------------
+	// Shared fragments
+	//-----------------------------------------------------------------------------
 	template<typename T, typename... TArgs>
 	UE_DEPRECATED(5.5, "This method will no longer be exposed. Use GetOrCreateSharedFragment instead.")
 	const FSharedStruct& GetOrCreateSharedFragmentByHash(const uint32 Hash, TArgs&&... InArgs)
@@ -1026,8 +1037,14 @@ PRAGMA_ENABLE_DEPRECATION_WARNINGS
 	template<UE::Mass::CConstSharedFragment T>
 	void ForEachConstSharedFragmentConditional(TFunctionRef< bool(const T& /*ConstSharedFragment*/) > ConditionFunction, TFunctionRef< void(const T& /*ConstSharedFragment*/) > ExecuteFunction);
 
+	//-----------------------------------------------------------------------------
+	// Entity Builder
+	//-----------------------------------------------------------------------------
 	[[nodiscard]] UE_API UE::Mass::FEntityBuilder MakeEntityBuilder();
 
+	//-----------------------------------------------------------------------------
+	// Sub-Managers
+	//-----------------------------------------------------------------------------
 	const UE::Mass::FTypeManager& GetTypeManager() const;
 	UE::Mass::FTypeManager& GetTypeManager();
 
diff --git a/Engine/Source/Runtime/MassEntity/Public/MassProcessor.h b/Engine/Source/Runtime/MassEntity/Public/MassProcessor.h
--- a/Engine/Source/Runtime/MassEntity/Public/MassProcessor.h
+++ b/Engine/Source/Runtime/MassEntity/Public/MassProcessor.h
@@ -9,8 +9,10 @@
 #include "MassProcessingTypes.h"
 #include "Async/TaskGraphInterfaces.h"
 #include "MassCommandBuffer.h"
+#include "MassRequirements.h"
 #include "MassProcessor.generated.h"
 
+
 #define UE_API MASSENTITY_API
 
 struct FMassProcessingPhaseConfig;
@@ -18,6 +20,9 @@ class UMassCompositeProcessor;
 struct FMassDebugger;
 struct FMassEntityQuery;
 struct FMassEntityManager;
+struct FMassExecutionContext;
+struct FMassExecutionRequirements;
+struct FMassSubsystemRequirements;
 namespace UE::Mass
 {
 	struct FQueryExecutor;
diff --git a/Engine/Source/Runtime/MassEntity/Public/MassQueryExecutor.h b/Engine/Source/Runtime/MassEntity/Public/MassQueryExecutor.h
--- a/Engine/Source/Runtime/MassEntity/Public/MassQueryExecutor.h
+++ b/Engine/Source/Runtime/MassEntity/Public/MassQueryExecutor.h
@@ -3,6 +3,7 @@
 #pragma once
 
 #include "Misc/NotNull.h"
+#include "MassEntityQuery.h"
 #include "MassExecutionContext.h"
 
 #define UE_API MASSENTITY_API
