[Mass] Sorting observer processors to utilize UMassProcessor.ExecutionPriority
Not that anyone requested that, it was just trivial to add so.. there we go.

diff --git a/Engine/Source/Runtime/MassEntity/Private/MassObserverManager.cpp b/Engine/Source/Runtime/MassEntity/Private/MassObserverManager.cpp
--- a/Engine/Source/Runtime/MassEntity/Private/MassObserverManager.cpp
+++ b/Engine/Source/Runtime/MassEntity/Private/MassObserverManager.cpp
@@ -73,6 +73,7 @@ namespace UE::Mass::ObserverManager
 				{
 					FMassRuntimePipeline& Pipeline = (*Observers).FindOrAdd(StructType);
 					Pipeline.AppendProcessors(MoveTemp(ObserverProcessors));
+					Pipeline.SortByExecutionPriority();
 					Pipeline.Initialize(Owner, EntityManager.AsShared());
 					
 					InOutObservedBitSet.Add(*StructType);
@@ -647,6 +648,8 @@ void FMassObserverManager::AddObserverInstance(TNotNull<const UScriptStruct*> El
 					ObserverProcessor->CallInitialize(Owner, EntityManager.AsShared());
 					bInitializeCalled = true;
 				}
+
+				Pipeline.SortByExecutionPriority();
 			}
 		}
 	}
@@ -666,6 +669,11 @@ void FMassObserverManager::AddObserverInstance(const UScriptStruct& ElementType,
 	AddObserverInstance(&ElementType, OperationFlag, &ObserverProcessor);
 }
 
+void FMassObserverManager::AddObserverInstance(TNotNull<UMassObserverProcessor*> ObserverProcessor)
+{
+	AddObserverInstance(ObserverProcessor->GetObservedTypeChecked(), ObserverProcessor->GetObservedOperations(), ObserverProcessor);
+}
+
 void FMassObserverManager::RemoveObserverInstance(TNotNull<const UScriptStruct*> ElementType, const EMassObservedOperationFlags OperationFlags, TNotNull<UMassProcessor*> ObserverProcessor)
 {
 	const bool bIsFragment = UE::Mass::IsA<FMassFragment>(ElementType);
diff --git a/Engine/Source/Runtime/MassEntity/Private/MassProcessingTypes.cpp b/Engine/Source/Runtime/MassEntity/Private/MassProcessingTypes.cpp
--- a/Engine/Source/Runtime/MassEntity/Private/MassProcessingTypes.cpp
+++ b/Engine/Source/Runtime/MassEntity/Private/MassProcessingTypes.cpp
@@ -259,6 +259,23 @@ UMassCompositeProcessor* FMassRuntimePipeline::FindTopLevelGroupByName(FName Gro
 	return nullptr;
 }
 
+void FMassRuntimePipeline::SortByExecutionPriority()
+{
+	if (Processors.IsEmpty())
+	{
+		return;
+	}
+
+	Processors.RemoveAllSwap([](const UMassProcessor* Processor)
+	{ 
+		return Processor == nullptr; 
+	});
+	Processors.Sort([](const UMassProcessor& ProcessorA, const UMassProcessor& ProcessorB)
+	{
+		return ProcessorA.GetExecutionPriority() > ProcessorB.GetExecutionPriority();
+	});
+}
+
 uint32 GetTypeHash(const FMassRuntimePipeline& Instance)
 { 
 	uint32 Hash = 0;
diff --git a/Engine/Source/Runtime/MassEntity/Public/MassObserverManager.h b/Engine/Source/Runtime/MassEntity/Public/MassObserverManager.h
--- a/Engine/Source/Runtime/MassEntity/Public/MassObserverManager.h
+++ b/Engine/Source/Runtime/MassEntity/Public/MassObserverManager.h
@@ -19,6 +19,7 @@ struct FMassObserverExecutionContext;
 struct FMassEntityManager;
 struct FMassArchetypeEntityCollection;
 class UMassProcessor;
+class UMassObserverProcessor;
 namespace UE::Mass
 {
 	struct FProcessingContext;
@@ -189,6 +190,7 @@ public:
 
 	MASSENTITY_API void AddObserverInstance(TNotNull<const UScriptStruct*> ElementType, EMassObservedOperationFlags OperationFlags, TNotNull<UMassProcessor*> ObserverProcessor);
 	MASSENTITY_API void AddObserverInstance(const UScriptStruct& ElementType, EMassObservedOperation Operation, UMassProcessor& ObserverProcessor);
+	MASSENTITY_API void AddObserverInstance(TNotNull<UMassObserverProcessor*> ObserverProcessor);
 	MASSENTITY_API void RemoveObserverInstance(TNotNull<const UScriptStruct*> ElementType, EMassObservedOperationFlags OperationFlags, TNotNull<UMassProcessor*> ObserverProcessor);
 	MASSENTITY_API void RemoveObserverInstance(const UScriptStruct& ElementType, EMassObservedOperation Operation, UMassProcessor& ObserverProcessor);
 
diff --git a/Engine/Source/Runtime/MassEntity/Public/MassObserverProcessor.h b/Engine/Source/Runtime/MassEntity/Public/MassObserverProcessor.h
--- a/Engine/Source/Runtime/MassEntity/Public/MassObserverProcessor.h
+++ b/Engine/Source/Runtime/MassEntity/Public/MassObserverProcessor.h
@@ -73,6 +73,9 @@ class UMassObserverProcessor : public UMassProcessor
 public:
 	UE_API UMassObserverProcessor();
 
+	EMassObservedOperationFlags GetObservedOperations() const;
+	TNotNull<const UScriptStruct*> GetObservedTypeChecked() const;
+
 protected:
 	UE_API virtual void PostInitProperties() override;
 
@@ -96,4 +99,18 @@ protected:
 	EMassObservedOperationFlags ObservedOperations = EMassObservedOperationFlags::None;
 };
 
+
+//----------------------------------------------------------------------//
+// inlines
+//----------------------------------------------------------------------//
+inline EMassObservedOperationFlags UMassObserverProcessor::GetObservedOperations() const
+{
+	return ObservedOperations;
+}
+
+inline TNotNull<const UScriptStruct*> UMassObserverProcessor::GetObservedTypeChecked() const
+{
+	return ObservedType;
+}
+
 #undef UE_API
diff --git a/Engine/Source/Runtime/MassEntity/Public/MassProcessingTypes.h b/Engine/Source/Runtime/MassEntity/Public/MassProcessingTypes.h
--- a/Engine/Source/Runtime/MassEntity/Public/MassProcessingTypes.h
+++ b/Engine/Source/Runtime/MassEntity/Public/MassProcessingTypes.h
@@ -144,6 +144,12 @@ public:
 
 	UE_API friend uint32 GetTypeHash(const FMassRuntimePipeline& Instance);
 
+	/**
+	 * Sorts processors aggregates in Processors array so that the ones with higher ExecutionPriority are executed first
+	 * The function will also remove nullptrs, if any, before sorting.
+	 */
+	UE_API void SortByExecutionPriority();
+
 	//-----------------------------------------------------------------------------
 	// DEPRECATED
 	//-----------------------------------------------------------------------------
