[Mass] shared fragments should not be considered the same base on UObject path so they are now using a CRC based on object instances

diff --git a/Engine/Plugins/Runtime/MassGameplay/Source/MassSpawner/Private/MassEntityTemplate.cpp b/Engine/Plugins/Runtime/MassGameplay/Source/MassSpawner/Private/MassEntityTemplate.cpp
--- a/Engine/Plugins/Runtime/MassGameplay/Source/MassSpawner/Private/MassEntityTemplate.cpp
+++ b/Engine/Plugins/Runtime/MassGameplay/Source/MassSpawner/Private/MassEntityTemplate.cpp
@@ -85,7 +85,7 @@ uint32 GetTypeHash(const FMassEntityTemplateData& Template)
 	// Initial fragment values, this is not part of the archetype as it is the spawner job to set them.
 	for (const FInstancedStruct& Struct : Template.GetInitialFragmentValues())
 	{
-		InitialValuesHash = UE::StructUtils::GetStructCrc32(Struct, InitialValuesHash);
+		InitialValuesHash = UE::StructUtils::GetStructInstanceCrc32(Struct, InitialValuesHash);
 	}
 
 	// These functions will be called to initialize entity's UObject-based fragments
diff --git a/Engine/Source/Runtime/CoreUObject/Private/StructUtils/StructUtilsTypes.cpp b/Engine/Source/Runtime/CoreUObject/Private/StructUtils/StructUtilsTypes.cpp
--- a/Engine/Source/Runtime/CoreUObject/Private/StructUtils/StructUtilsTypes.cpp
+++ b/Engine/Source/Runtime/CoreUObject/Private/StructUtils/StructUtilsTypes.cpp
@@ -6,13 +6,69 @@
 #include "StructUtils/StructView.h"
 #include "StructUtils/SharedStruct.h"
 #include "StructUtils/UserDefinedStruct.h"
+#include "UObject/ObjectKey.h"
 
 
 namespace UE::StructUtils
 {
+	/**
+	 * Calculates a checksum from the input provided to the archive
+	 * with customizable mode regarding UObject serialization.
+	 */
+	class FHashBuilder : public FArchiveCrc32
+	{
+	public:
+		enum class EObjectHashingMode
+		{
+			/** Objects with same name will produce the same hash */
+			ObjectPathBasedCRC,
+
+			/** Generated hash will be unique per object instance */
+			ObjectPtrBasedCRC
+		};
+
+		explicit FHashBuilder(uint32 InCRC, EObjectHashingMode ObjectHashingMode)
+			: FArchiveCrc32(InCRC, /*RootObject*/nullptr)
+			, ObjectHashingMode(ObjectHashingMode)
+		{
+		}
+
+		//~ Begin FArchive Interface
+		virtual FArchive& operator<<(UObject*& Object) override
+		{
+			if (ObjectHashingMode == EObjectHashingMode::ObjectPtrBasedCRC)
+			{
+				uint32 ObjectKeyHash = GetTypeHash(FObjectKey(Object));
+				Serialize(&ObjectKeyHash, sizeof(ObjectKeyHash));
+				return *this;
+			}
+			
+			return FArchiveCrc32::operator<<(Object);
+		}
+
+		virtual FString GetArchiveName() const override
+		{
+			return TEXT("StructUtils_FHashBuilder");
+		}
+		//~ End FArchive Interface
+	private:
+		EObjectHashingMode ObjectHashingMode = EObjectHashingMode::ObjectPathBasedCRC;
+	};
+
+	uint32 GetStructInstanceCrc32(const UScriptStruct& ScriptStruct, const uint8* StructMemory, const uint32 CRC /*= 0*/)
+	{
+		FHashBuilder Ar(HashCombine(CRC, PointerHash(&ScriptStruct)), FHashBuilder::EObjectHashingMode::ObjectPtrBasedCRC);
+		if (StructMemory)
+		{
+			UScriptStruct& NonConstScriptStruct = const_cast<UScriptStruct&>(ScriptStruct);
+			NonConstScriptStruct.SerializeItem(Ar, const_cast<uint8*>(StructMemory), nullptr);
+		}
+		return Ar.GetCrc();
+	}
+
 	uint32 GetStructCrc32(const UScriptStruct& ScriptStruct, const uint8* StructMemory, const uint32 CRC /*= 0*/)
 	{
-		FArchiveCrc32 Ar(HashCombine(CRC, PointerHash(&ScriptStruct)));
+		FHashBuilder Ar(HashCombine(CRC, PointerHash(&ScriptStruct)), FHashBuilder::EObjectHashingMode::ObjectPathBasedCRC);
 		if (StructMemory)
 		{
 			UScriptStruct& NonConstScriptStruct = const_cast<UScriptStruct&>(ScriptStruct);
@@ -22,33 +78,58 @@ namespace UE::StructUtils
 	}
 
 	template<typename T>
-	uint32 GetStructCrc32Helper(const T& Struct, const uint32 CRC /*= 0*/)
+	uint32 GetStructCrc32Helper(const T& Struct, const FHashBuilder::EObjectHashingMode ObjectHashingMode, const uint32 CRC /*= 0*/)
 	{
 		if (const UScriptStruct* ScriptStruct = Struct.GetScriptStruct())
 		{
-			return GetStructCrc32(*ScriptStruct, Struct.GetMemory(), CRC);
+			if (ObjectHashingMode == FHashBuilder::EObjectHashingMode::ObjectPathBasedCRC)
+			{
+				return GetStructCrc32(*ScriptStruct, Struct.GetMemory(), CRC);
+			}
+
+			return GetStructInstanceCrc32(*ScriptStruct, Struct.GetMemory(), CRC);
 		}
 		return 0;
 	}
 
 	uint32 GetStructCrc32(const FStructView& StructView, const uint32 CRC /*= 0*/)
 	{
-		return GetStructCrc32Helper(StructView, CRC);
+		return GetStructCrc32Helper(StructView, FHashBuilder::EObjectHashingMode::ObjectPathBasedCRC, CRC);
 	}
 
 	uint32 GetStructCrc32(const FConstStructView& StructView, const uint32 CRC /*= 0*/)
 	{
-		return GetStructCrc32Helper(StructView, CRC);
+		return GetStructCrc32Helper(StructView, FHashBuilder::EObjectHashingMode::ObjectPathBasedCRC, CRC);
 	}
 
 	uint32 GetStructCrc32(const FSharedStruct& SharedView, const uint32 CRC /*= 0*/)
 	{
-		return GetStructCrc32Helper(SharedView, CRC);
+		return GetStructCrc32Helper(SharedView, FHashBuilder::EObjectHashingMode::ObjectPathBasedCRC, CRC);
 	}
 
 	uint32 GetStructCrc32(const FConstSharedStruct& SharedView, const uint32 CRC /*= 0*/)
 	{
-		return GetStructCrc32Helper(SharedView, CRC);
+		return GetStructCrc32Helper(SharedView, FHashBuilder::EObjectHashingMode::ObjectPathBasedCRC, CRC);
+	}
+
+	uint32 GetStructInstanceCrc32(const FStructView& StructView, const uint32 CRC /*= 0*/)
+	{
+		return GetStructCrc32Helper(StructView, FHashBuilder::EObjectHashingMode::ObjectPtrBasedCRC, CRC);
+	}
+
+	uint32 GetStructInstanceCrc32(const FConstStructView& StructView, const uint32 CRC /*= 0*/)
+	{
+		return GetStructCrc32Helper(StructView, FHashBuilder::EObjectHashingMode::ObjectPtrBasedCRC, CRC);
+	}
+
+	uint32 GetStructInstanceCrc32(const FSharedStruct& SharedView, const uint32 CRC /*= 0*/)
+	{
+		return GetStructCrc32Helper(SharedView, FHashBuilder::EObjectHashingMode::ObjectPtrBasedCRC, CRC);
+	}
+
+	uint32 GetStructInstanceCrc32(const FConstSharedStruct& SharedView, const uint32 CRC /*= 0*/)
+	{
+		return GetStructCrc32Helper(SharedView, FHashBuilder::EObjectHashingMode::ObjectPtrBasedCRC, CRC);
 	}
 
 	class FArchiveCityHash64 : public FArchiveUObject
diff --git a/Engine/Source/Runtime/CoreUObject/Public/StructUtils/StructUtilsTypes.h b/Engine/Source/Runtime/CoreUObject/Public/StructUtils/StructUtilsTypes.h
--- a/Engine/Source/Runtime/CoreUObject/Public/StructUtils/StructUtilsTypes.h
+++ b/Engine/Source/Runtime/CoreUObject/Public/StructUtils/StructUtilsTypes.h
@@ -21,22 +21,36 @@ class UUserDefinedStruct;
 
 namespace UE::StructUtils
 {
-	extern COREUOBJECT_API uint32 GetStructCrc32(const UScriptStruct& ScriptStruct, const uint8* StructMemory, const uint32 CRC = 0);
-	extern COREUOBJECT_API uint32 GetStructCrc32(const FStructView& StructView, const uint32 CRC = 0);
-	extern COREUOBJECT_API uint32 GetStructCrc32(const FConstStructView& StructView, const uint32 CRC = 0);
-	extern COREUOBJECT_API uint32 GetStructCrc32(const FSharedStruct& SharedView, const uint32 CRC = 0);
-	extern COREUOBJECT_API uint32 GetStructCrc32(const FConstSharedStruct& SharedView, const uint32 CRC = 0);
+	/**
+	 * FArchiveCrc32-based struct hashing functions using object path for UObjects.
+	 * i.e., Objects with same name will produce the same hash
+	 */
+	extern UE_API uint32 GetStructCrc32(const UScriptStruct& ScriptStruct, const uint8* StructMemory, const uint32 CRC = 0);
+	extern UE_API uint32 GetStructCrc32(const FStructView& StructView, const uint32 CRC = 0);
+	extern UE_API uint32 GetStructCrc32(const FConstStructView& StructView, const uint32 CRC = 0);
+	extern UE_API uint32 GetStructCrc32(const FSharedStruct& SharedView, const uint32 CRC = 0);
+	extern UE_API uint32 GetStructCrc32(const FConstSharedStruct& SharedView, const uint32 CRC = 0);
+
+	/**
+	 * FArchiveCrc32-based struct hashing functions using object key for UObjects.
+	 * i.e., Generated hash will be unique per object instance
+	 */
+	extern UE_API uint32 GetStructInstanceCrc32(const UScriptStruct& ScriptStruct, const uint8* StructMemory, const uint32 CRC = 0);
+	extern UE_API uint32 GetStructInstanceCrc32(const FStructView& StructView, const uint32 CRC = 0);
+	extern UE_API uint32 GetStructInstanceCrc32(const FConstStructView& StructView, const uint32 CRC = 0);
+	extern UE_API uint32 GetStructInstanceCrc32(const FSharedStruct& SharedView, const uint32 CRC = 0);
+	extern UE_API uint32 GetStructInstanceCrc32(const FConstSharedStruct& SharedView, const uint32 CRC = 0);
 
 	/** 
 	 * CityHash64-based struct hashing functions.
 	 * Note that these are relatively slow due to using either UScriptStruct.GetStructTypeHash (if implemented) or
 	 * a serialization path. 
 	 */
-	extern COREUOBJECT_API uint64 GetStructHash64(const UScriptStruct& ScriptStruct, const uint8* StructMemory);
-	extern COREUOBJECT_API uint64 GetStructHash64(const FStructView& StructView);
-	extern COREUOBJECT_API uint64 GetStructHash64(const FConstStructView& StructView);
-	extern COREUOBJECT_API uint64 GetStructHash64(const FSharedStruct& SharedView);
-	extern COREUOBJECT_API uint64 GetStructHash64(const FConstSharedStruct& SharedView);
+	extern UE_API uint64 GetStructHash64(const UScriptStruct& ScriptStruct, const uint8* StructMemory);
+	extern UE_API uint64 GetStructHash64(const FStructView& StructView);
+	extern UE_API uint64 GetStructHash64(const FConstStructView& StructView);
+	extern UE_API uint64 GetStructHash64(const FSharedStruct& SharedView);
+	extern UE_API uint64 GetStructHash64(const FConstSharedStruct& SharedView);
 
 	template <typename T>
 	typename TEnableIf<!TIsDerivedFrom<T, UObject>::IsDerived, UScriptStruct*>::Type GetAsUStruct()
@@ -128,8 +142,8 @@ namespace UE::StructUtils::Private
 		UObject* OldCurrentReinstantiateOuterObject = nullptr;
 	};
 
-	extern COREUOBJECT_API const UUserDefinedStruct* GetStructureToReinstantiate();
-	extern COREUOBJECT_API UObject* GetCurrentReinstantiationOuterObject();
+	extern UE_API const UUserDefinedStruct* GetStructureToReinstantiate();
+	extern UE_API UObject* GetCurrentReinstantiationOuterObject();
 };
 #endif // WITH_EDITOR
 
diff --git a/Engine/Source/Runtime/MassEntity/Public/MassEntityManager.h b/Engine/Source/Runtime/MassEntity/Public/MassEntityManager.h
--- a/Engine/Source/Runtime/MassEntity/Public/MassEntityManager.h
+++ b/Engine/Source/Runtime/MassEntity/Public/MassEntityManager.h
@@ -767,7 +767,7 @@ public:
 	 * Returns or creates a shared struct associated to a given shared fragment set of values
 	 * identified internally by a CRC.
 	 * Use this overload when an instance of the desired const shared fragment type is available and
-	 * that can be used directly to compute a CRC (i.e., UE::StructUtils::GetStructCrc32)
+	 * that can be used directly to compute a CRC (i.e., UE::StructUtils::GetStructInstanceCrc32)
 	 *	e.g.,
 	 *	USTRUCT()
 	 *	struct FIntConstSharedFragment : public FMassConstSharedFragment
@@ -790,7 +790,7 @@ public:
 	{
 		static_assert(UE::Mass::CConstSharedFragment<T>,
 			"Given struct doesn't represent a valid const shared fragment type. Make sure to inherit from FMassConstSharedFragment or one of its child-types.");
-		const uint32 Hash = UE::StructUtils::GetStructCrc32(FConstStructView::Make(Fragment));
+		const uint32 Hash = UE::StructUtils::GetStructInstanceCrc32(FConstStructView::Make(Fragment));
 PRAGMA_DISABLE_DEPRECATION_WARNINGS
 		return GetOrCreateConstSharedFragmentByHash(Hash, Fragment);
 PRAGMA_ENABLE_DEPRECATION_WARNINGS
@@ -800,7 +800,7 @@ PRAGMA_ENABLE_DEPRECATION_WARNINGS
 	 * Returns or creates a shared struct associated to a given shared fragment set of values
 	 * identified internally by a CRC.
 	 * Use this overload when an instance of the desired shared fragment type is available and
-	 * that can be used directly to compute a CRC (i.e., UE::StructUtils::GetStructCrc32)
+	 * that can be used directly to compute a CRC (i.e., UE::StructUtils::GetStructInstanceCrc32)
 	 *	e.g.,
 	 *	USTRUCT()
 	 *	struct FIntSharedFragment : public FMassSharedFragment
@@ -823,7 +823,7 @@ PRAGMA_ENABLE_DEPRECATION_WARNINGS
 	{
 		static_assert(UE::Mass::CSharedFragment<T>,
 			"Given struct doesn't represent a valid shared fragment type. Make sure to inherit from FMassSharedFragment or one of its child-types.");
-		const uint32 Hash = UE::StructUtils::GetStructCrc32(FConstStructView::Make(Fragment));
+		const uint32 Hash = UE::StructUtils::GetStructInstanceCrc32(FConstStructView::Make(Fragment));
 		return GetOrCreateSharedFragmentByHash(Hash, Fragment);
 	}
 
@@ -831,7 +831,7 @@ PRAGMA_ENABLE_DEPRECATION_WARNINGS
 	 * Returns or creates a shared struct associated to a given shared fragment set of values
 	 * identified internally by a CRC.
 	 * Use this overload when values can be provided as constructor arguments for the desired const shared fragment type and
-	 * that can be used directly to compute a CRC (i.e., UE::StructUtils::GetStructCrc32)
+	 * that can be used directly to compute a CRC (i.e., UE::StructUtils::GetStructInstanceCrc32)
  	 *	e.g.,
 	 *	USTRUCT()
 	 *	struct FIntConstSharedFragment : public FMassConstSharedFragment
@@ -855,7 +855,7 @@ PRAGMA_ENABLE_DEPRECATION_WARNINGS
 		static_assert(UE::Mass::CConstSharedFragment<T>,
 			"Given struct doesn't represent a valid const shared fragment type. Make sure to inherit from FMassConstSharedFragment or one of its child-types.");
 		T Struct(Forward<TArgs>(InArgs)...);
-		const uint32 Hash = UE::StructUtils::GetStructCrc32(FConstStructView::Make(Struct));
+		const uint32 Hash = UE::StructUtils::GetStructInstanceCrc32(FConstStructView::Make(Struct));
 PRAGMA_DISABLE_DEPRECATION_WARNINGS
 		return GetOrCreateConstSharedFragmentByHash(Hash, MoveTemp(Struct));
 PRAGMA_ENABLE_DEPRECATION_WARNINGS
@@ -865,7 +865,7 @@ PRAGMA_ENABLE_DEPRECATION_WARNINGS
 	 * Returns or creates a shared struct associated to a given shared fragment set of values
 	 * identified internally by a CRC.
 	 * Use this overload when values can be provided as constructor arguments for the desired shared fragment type and
-	 * that can be used directly to compute a CRC (i.e., UE::StructUtils::GetStructCrc32)
+	 * that can be used directly to compute a CRC (i.e., UE::StructUtils::GetStructInstanceCrc32)
  	 *	e.g.,
 	 *	USTRUCT()
 	 *	struct FIntSharedFragment : public FMassSharedFragment
@@ -889,7 +889,7 @@ PRAGMA_ENABLE_DEPRECATION_WARNINGS
 		static_assert(UE::Mass::CSharedFragment<T>,
 			"Given struct doesn't represent a valid shared fragment type. Make sure to inherit from FMassSharedFragment or one of its child-types.");
 		T Struct(Forward<TArgs>(InArgs)...);
-		const uint32 Hash = UE::StructUtils::GetStructCrc32(FConstStructView::Make(Struct));
+		const uint32 Hash = UE::StructUtils::GetStructInstanceCrc32(FConstStructView::Make(Struct));
 		return GetOrCreateSharedFragmentByHash(Hash, MoveTemp(Struct));
 	}
 
@@ -897,7 +897,7 @@ PRAGMA_ENABLE_DEPRECATION_WARNINGS
 	 * Returns or creates a shared struct associated to a given shared fragment set of values
 	 * identified internally by a CRC.
 	 * Use this overload when the reflection data and the memory of an instance of the desired const shared fragment type
-	 * is available and that can be used directly to compute a CRC (i.e., UE::StructUtils::GetStructCrc32)
+	 * is available and that can be used directly to compute a CRC (i.e., UE::StructUtils::GetStructInstanceCrc32)
 	 * e.g.,
 	 * FSharedStruct SharedStruct = EntityManager.GetOrCreateConstSharedFragment(*StructView.GetScriptStruct(), StructView.GetMemory());
 	 *
@@ -909,7 +909,7 @@ PRAGMA_ENABLE_DEPRECATION_WARNINGS
 	{
 		checkf(InScriptStruct.IsChildOf(TBaseStructure<FMassConstSharedFragment>::Get()),
 			TEXT("Given struct doesn't represent a valid const shared fragment type. Make sure to inherit from FMassConstSharedFragment or one of its child-types."));
-		const uint32 Hash = UE::StructUtils::GetStructCrc32(InScriptStruct, InStructMemory);
+		const uint32 Hash = UE::StructUtils::GetStructInstanceCrc32(InScriptStruct, InStructMemory);
 		return GetOrCreateConstSharedFragmentByHash(Hash, &InScriptStruct, InStructMemory);
 	}
 
@@ -917,7 +917,7 @@ PRAGMA_ENABLE_DEPRECATION_WARNINGS
 	 * Returns or creates a shared struct associated to a given shared fragment set of values
 	 * identified internally by a CRC.
 	 * Use this overload when the reflection data and the memory of an instance of the desired shared fragment type
-	 * is available and that can be used directly to compute a CRC (i.e., UE::StructUtils::GetStructCrc32)
+	 * is available and that can be used directly to compute a CRC (i.e., UE::StructUtils::GetStructInstanceCrc32)
 	 * e.g.,
 	 * FSharedStruct SharedStruct = EntityManager.GetOrCreateSharedFragment(*StructView.GetScriptStruct(), StructView.GetMemory());
 	 *
@@ -929,14 +929,14 @@ PRAGMA_ENABLE_DEPRECATION_WARNINGS
 	{
 		checkf(InScriptStruct.IsChildOf(TBaseStructure<FMassSharedFragment>::Get()),
 			TEXT("Given struct doesn't represent a valid shared fragment type. Make sure to inherit from FMassSharedFragment or one of its child-types."));
-		const uint32 Hash = UE::StructUtils::GetStructCrc32(InScriptStruct, InStructMemory);
+		const uint32 Hash = UE::StructUtils::GetStructInstanceCrc32(InScriptStruct, InStructMemory);
 		return GetOrCreateSharedFragmentByHash(Hash, &InScriptStruct, InStructMemory);
 	}
 
 	/**
 	 * Returns or creates a shared struct associated to a given shared fragment set of values
 	 * identified internally by a CRC.
-	 * Use this overload when a different struct should be used to compute a CRC (i.e., UE::StructUtils::GetStructCrc32)
+	 * Use this overload when a different struct should be used to compute a CRC (i.e., UE::StructUtils::GetStructInstanceCrc32)
 	 * and values can be provided as constructor arguments for the desired const shared fragment type
 	 *	e.g.,
 	 *
@@ -964,7 +964,7 @@ PRAGMA_ENABLE_DEPRECATION_WARNINGS
 	 *	FIntConstSharedFragmentParams Params(123);
 	 *	const FConstSharedStruct SharedStruct = EntityManager.GetOrCreateConstSharedFragment<FIntConstSharedFragment>(FConstStructView::Make(Params), Params);
 	 *
-	 * @params HashingHelperStruct Struct view passed to UE::StructUtils::GetStructCrc32 to compute the CRC
+	 * @params HashingHelperStruct Struct view passed to UE::StructUtils::GetStructInstanceCrc32 to compute the CRC
 	 * @params InArgs List of arguments provided to the constructor of the desired fragment type
 	 * @return FConstSharedStruct to the matching, or newly created shared fragment
 	 */
@@ -974,7 +974,7 @@ PRAGMA_ENABLE_DEPRECATION_WARNINGS
 		static_assert(UE::Mass::CConstSharedFragment<T>,
 			"Given struct doesn't represent a valid const shared fragment type. Make sure to inherit from FMassConstSharedFragment or one of its child-types.");
 		T Fragment(Forward<TArgs>(InArgs)...);
-		const uint32 Hash = UE::StructUtils::GetStructCrc32(HashingHelperStruct);
+		const uint32 Hash = UE::StructUtils::GetStructInstanceCrc32(HashingHelperStruct);
 PRAGMA_DISABLE_DEPRECATION_WARNINGS
 		return GetOrCreateConstSharedFragmentByHash(Hash, MoveTemp(Fragment));
 PRAGMA_ENABLE_DEPRECATION_WARNINGS
@@ -983,7 +983,7 @@ PRAGMA_ENABLE_DEPRECATION_WARNINGS
 	/**
 	 * Returns or creates a shared struct associated to a given shared fragment set of values
 	 * identified internally by a CRC.
-	 * Use this overload when a different struct should be used to compute a CRC (i.e., UE::StructUtils::GetStructCrc32)
+	 * Use this overload when a different struct should be used to compute a CRC (i.e., UE::StructUtils::GetStructInstanceCrc32)
 	 * and values can be provided as constructor arguments for the desired shared fragment type
 	 *	e.g.,
 	 *
@@ -1011,7 +1011,7 @@ PRAGMA_ENABLE_DEPRECATION_WARNINGS
 	 *	FIntSharedFragmentParams Params(123);
 	 *	const FSharedStruct SharedStruct = EntityManager.GetOrCreateSharedFragment<FIntSharedFragment>(FConstStructView::Make(Params), Params);
 	 *
-	 * @params HashingHelperStruct Struct view passed to UE::StructUtils::GetStructCrc32 to compute the CRC
+	 * @params HashingHelperStruct Struct view passed to UE::StructUtils::GetStructInstanceCrc32 to compute the CRC
 	 * @params InArgs List of arguments provided to the constructor of the desired fragment type
 	 * @return FSharedStruct to the matching, or newly created shared fragment
 	 */
@@ -1021,7 +1021,7 @@ PRAGMA_ENABLE_DEPRECATION_WARNINGS
 		static_assert(UE::Mass::CSharedFragment<T>,
 			"Given struct doesn't represent a valid shared fragment type. Make sure to inherit from FMassSharedFragment or one of its child-types.");
 		T Fragment(Forward<TArgs>(InArgs)...);
-		const uint32 Hash = UE::StructUtils::GetStructCrc32(HashingHelperStruct);
+		const uint32 Hash = UE::StructUtils::GetStructInstanceCrc32(HashingHelperStruct);
 		return GetOrCreateSharedFragmentByHash(Hash, MoveTemp(Fragment));
 	}
 
