Replace some usages of FORCEINLINE with inline in Mass modules.

For context about this change, see this doc: https://docs.google.com/document/d/1LjpT7NpzTecfXdFNyv5KM6sCu67827BRljNi1iaWVSE/edit?tab=t.0#heading=h.fvtz8gjukz7x

Includes modules: MassNavigation, MassActors, MassCommon, MassEQS, MassGameplayDebug, MassLOD, MassRepresentation, MassEntity

Reviewers were automatically determined. Please comment if anyone else should be added.
#rnx

diff --git a/Engine/Plugins/AI/MassAI/Source/MassNavigation/Public/MassNavigationDebug.h b/Engine/Plugins/AI/MassAI/Source/MassNavigation/Public/MassNavigationDebug.h
--- a/Engine/Plugins/AI/MassAI/Source/MassNavigation/Public/MassNavigationDebug.h
+++ b/Engine/Plugins/AI/MassAI/Source/MassNavigation/Public/MassNavigationDebug.h
@@ -26,12 +26,12 @@ namespace UE::MassNavigation::Debug
 			, EntityIndex(InEntityIndex)
 		{}
 
-		FORCEINLINE const UObject* GetLogOwner() const
+		inline const UObject* GetLogOwner() const
 		{
 			return LogContext.GetLogOwner(EntityIndex, FallbackLogOwner);
 		}
 		
-		FORCEINLINE bool ShouldLogEntity(FColor* OutEntityColor = nullptr) const
+		inline bool ShouldLogEntity(FColor* OutEntityColor = nullptr) const
 		{
 			return LogContext.ShouldLogEntity(EntityIndex, OutEntityColor);
 		}
diff --git a/Engine/Plugins/Runtime/MassGameplay/Source/MassActors/Public/MassActorSubsystem.h b/Engine/Plugins/Runtime/MassGameplay/Source/MassActors/Public/MassActorSubsystem.h
--- a/Engine/Plugins/Runtime/MassGameplay/Source/MassActors/Public/MassActorSubsystem.h
+++ b/Engine/Plugins/Runtime/MassGameplay/Source/MassActors/Public/MassActorSubsystem.h
@@ -99,21 +99,21 @@ struct FMassActorFragment : public FObjectWrapperFragment
 
 
 	/** @return none const pointer to the actor	*/
-	FORCEINLINE AActor* GetMutable() { return Actor.Get(); }
+	inline AActor* GetMutable() { return Actor.Get(); }
 	UE_API AActor* GetMutable(EActorAccess Access);
 
 	/** @return none const pointer to the actor	only if owned by mass */
-	FORCEINLINE AActor* GetOwnedByMassMutable() { return bIsOwnedByMass ? Actor.Get() : nullptr; }
+	inline AActor* GetOwnedByMassMutable() { return bIsOwnedByMass ? Actor.Get() : nullptr; }
 
 	/** @return none const pointer to the actor	only if owned by mass */
-	FORCEINLINE const AActor* Get() const { return Actor.Get(); }
+	inline const AActor* Get() const { return Actor.Get(); }
 	UE_API const AActor* Get(EActorAccess Access) const;
 
 	/** @return if the actor is owned by mass */
-	FORCEINLINE bool IsOwnedByMass() const { return bIsOwnedByMass; }
+	inline bool IsOwnedByMass() const { return bIsOwnedByMass; }
 
 	/** @return if the actor is a valid pointer */
-	FORCEINLINE bool IsValid() const { return Actor.IsValid(); }
+	inline bool IsValid() const { return Actor.IsValid(); }
 
 private:
 	// made visible for debugging purposes. It will show up in Mass's gameplay debugger category when viewing fragment details
diff --git a/Engine/Plugins/Runtime/MassGameplay/Source/MassCommon/Public/MassCommonUtils.h b/Engine/Plugins/Runtime/MassGameplay/Source/MassCommon/Public/MassCommonUtils.h
--- a/Engine/Plugins/Runtime/MassGameplay/Source/MassCommon/Public/MassCommonUtils.h
+++ b/Engine/Plugins/Runtime/MassGameplay/Source/MassCommon/Public/MassCommonUtils.h
@@ -40,11 +40,11 @@ namespace UE::Mass::Utils
 
 #else
 
-	FORCEINLINE constexpr bool IsDeterministic() { return false; }
+	inline constexpr bool IsDeterministic() { return false; }
 
-	FORCEINLINE int32 OverrideRandomSeedForTesting(int32 InSeed) { return InSeed; }
+	inline int32 OverrideRandomSeedForTesting(int32 InSeed) { return InSeed; }
 
-	FORCEINLINE int32 GenerateRandomSeed() { return FMath::Rand(); }
+	inline int32 GenerateRandomSeed() { return FMath::Rand(); }
 #endif
 
 } // namespace UE::Mass::Utils
\ No newline at end of file
diff --git a/Engine/Plugins/Runtime/MassGameplay/Source/MassCommon/Public/RandomSequence.h b/Engine/Plugins/Runtime/MassGameplay/Source/MassCommon/Public/RandomSequence.h
--- a/Engine/Plugins/Runtime/MassGameplay/Source/MassCommon/Public/RandomSequence.h
+++ b/Engine/Plugins/Runtime/MassGameplay/Source/MassCommon/Public/RandomSequence.h
@@ -28,7 +28,7 @@ namespace UE::RandomSequence
 	/**
 	 * @return 32 bit fibonacci hash at specified index.
 	 */
-	FORCEINLINE uint32 FibonacciHash(const int32 SeqIndex)
+	inline uint32 FibonacciHash(const int32 SeqIndex)
 	{
 		constexpr uint32 K = 2654435769u; // 2^32 / phi (golden ratio)
 		// Offset the sequence by 1, so that index 0 is not always 0.
@@ -39,7 +39,7 @@ namespace UE::RandomSequence
 	 * Helper function to return random float.
 	 * @return A random number in [0..1]
 	 */
-	FORCEINLINE float FRand(const int32 SeqIndex)
+	inline float FRand(const int32 SeqIndex)
 	{
 		float Result;
 		*(uint32*)&Result = 0x3F800000U | (FibonacciHash(SeqIndex) >> 9);
@@ -50,7 +50,7 @@ namespace UE::RandomSequence
 	 * Helper function to return random int on specified range.
 	 * @return A random number in [0..A)
 	 */
-	FORCEINLINE int32 RandHelper(const int32 SeqIndex, const int32 A)
+	inline int32 RandHelper(const int32 SeqIndex, const int32 A)
 	{
 		return (int32)(((int64)FibonacciHash(SeqIndex) * (int64)A) >> 32);
 	}
@@ -59,7 +59,7 @@ namespace UE::RandomSequence
 	 * Helper function to return random int in specified range.
 	 * @return A random number >= Min and <= Max
 	 */
-	FORCEINLINE int32 RandRange(const int32 SeqIndex, const int32 InMin, const int32 InMax)
+	inline int32 RandRange(const int32 SeqIndex, const int32 InMin, const int32 InMax)
 	{
 		const int32 Range = (InMax - InMin) + 1;
 		return InMin + RandHelper(SeqIndex, Range);
@@ -69,7 +69,7 @@ namespace UE::RandomSequence
 	 * Helper function to return random float in specified range.
 	 * @return A random number >= Min and <= Max
 	 */
-	FORCEINLINE float FRandRange(const int32 SeqIndex, const float InMin, const float InMax)
+	inline float FRandRange(const int32 SeqIndex, const float InMin, const float InMax)
 	{
 		return InMin + (InMax - InMin) * FRand(SeqIndex);
 	}
@@ -78,13 +78,13 @@ namespace UE::RandomSequence
 	 * Helper function to return random float in specified range.
 	 * @return A random number >= Min and <= Max
 	 */
-	FORCEINLINE float RandRange(const int32 SeqIndex, const float InMin, const float InMax)
+	inline float RandRange(const int32 SeqIndex, const float InMin, const float InMax)
 	{
 		return FRandRange(SeqIndex, InMin, InMax);
 	}
 
 	UE_DEPRECATED(5.4, "Please use the correctly spelled version FibonacciHash()")
-	FORCEINLINE uint32 FibocciHash(const int32 SeqIndex)
+	inline uint32 FibocciHash(const int32 SeqIndex)
 	{
 		return FibonacciHash(SeqIndex);
 	}
diff --git a/Engine/Plugins/Runtime/MassGameplay/Source/MassEQS/Public/Generators/MassEnvQueryGenerator.h b/Engine/Plugins/Runtime/MassGameplay/Source/MassEQS/Public/Generators/MassEnvQueryGenerator.h
--- a/Engine/Plugins/Runtime/MassGameplay/Source/MassEQS/Public/Generators/MassEnvQueryGenerator.h
+++ b/Engine/Plugins/Runtime/MassGameplay/Source/MassEQS/Public/Generators/MassEnvQueryGenerator.h
@@ -18,7 +18,7 @@ public:
 	*/
 	virtual void GenerateItems(FEnvQueryInstance& QueryInstance) const final;
 
-	virtual FORCEINLINE bool IsCurrentlyRunningAsync() const final { return MassEQSRequestHandler.IsPendingResults(); }
+	virtual inline bool IsCurrentlyRunningAsync() const final { return MassEQSRequestHandler.IsPendingResults(); }
 
 	// To be implemented by child class:
 	// Begin IMassEQSRequestInterface
diff --git a/Engine/Plugins/Runtime/MassGameplay/Source/MassEQS/Public/MassEQSBlueprintLibrary.h b/Engine/Plugins/Runtime/MassGameplay/Source/MassEQS/Public/MassEQSBlueprintLibrary.h
--- a/Engine/Plugins/Runtime/MassGameplay/Source/MassEQS/Public/MassEQSBlueprintLibrary.h
+++ b/Engine/Plugins/Runtime/MassGameplay/Source/MassEQS/Public/MassEQSBlueprintLibrary.h
@@ -24,9 +24,9 @@ public:
 	{
 	}
 
-	FORCEINLINE FVector GetCachedEntityPosition() const { return EntityInfo.CachedTransform.GetLocation(); }
-	FORCEINLINE FMassEntityHandle GetEntityHandle() const { return EntityInfo.EntityHandle; }
-	FORCEINLINE void SetEntityHandle(FMassEntityHandle Handle) { EntityInfo.EntityHandle = Handle; }
+	inline FVector GetCachedEntityPosition() const { return EntityInfo.CachedTransform.GetLocation(); }
+	inline FMassEntityHandle GetEntityHandle() const { return EntityInfo.EntityHandle; }
+	inline void SetEntityHandle(FMassEntityHandle Handle) { EntityInfo.EntityHandle = Handle; }
 
 	const FMassEnvQueryEntityInfo& GetEntityInfo() const { return EntityInfo; };
 	inline bool operator==(const FMassEnvQueryEntityInfoBlueprintWrapper& Other) const { return GetEntityInfo() == Other.GetEntityInfo(); }
@@ -57,17 +57,17 @@ public:
 	//----------------------------------------------------------------------//
 
 	UFUNCTION(BlueprintCallable, BlueprintPure, Category = "MassEnvQuery|Utils")
-	static FORCEINLINE FString EntityToString(const FMassEnvQueryEntityInfoBlueprintWrapper& EntityInfo){ return EntityInfo.GetEntityHandle().DebugGetDescription(); }
+	static inline FString EntityToString(const FMassEnvQueryEntityInfoBlueprintWrapper& EntityInfo){ return EntityInfo.GetEntityHandle().DebugGetDescription(); }
 
 	UFUNCTION(BlueprintCallable, BlueprintPure, Category = "MassEnvQuery|Utils")
-	static FORCEINLINE FVector GetCachedEntityPosition(const FMassEnvQueryEntityInfoBlueprintWrapper& EntityInfo) { return EntityInfo.GetCachedEntityPosition(); }
+	static inline FVector GetCachedEntityPosition(const FMassEnvQueryEntityInfoBlueprintWrapper& EntityInfo) { return EntityInfo.GetCachedEntityPosition(); }
 
 	UFUNCTION(BlueprintCallable, BlueprintPure, Category = "MassEnvQuery|Utils")
 	static FVector GetCurrentEntityPosition(const AActor* Owner, const FMassEnvQueryEntityInfoBlueprintWrapper& EntityInfo);
 
 	/** Custom comparison function, as the Blueprint Equals did not seem to work. */
 	UFUNCTION(BlueprintCallable, BlueprintPure, Category = "MassEnvQuery|Utils")
-	static FORCEINLINE bool EntityComparison(const FMassEnvQueryEntityInfoBlueprintWrapper& A, const FMassEnvQueryEntityInfoBlueprintWrapper& B) { return A == B; }
+	static inline bool EntityComparison(const FMassEnvQueryEntityInfoBlueprintWrapper& A, const FMassEnvQueryEntityInfoBlueprintWrapper& B) { return A == B; }
 
 	/** Custom array-contains function, as the Blueprint version did not seem to work. */
 	UFUNCTION(BlueprintCallable, BlueprintPure, Category = "MassEnvQuery|Utils")
diff --git a/Engine/Plugins/Runtime/MassGameplay/Source/MassEQS/Public/MassEQSTypes.h b/Engine/Plugins/Runtime/MassGameplay/Source/MassEQS/Public/MassEQSTypes.h
--- a/Engine/Plugins/Runtime/MassGameplay/Source/MassEQS/Public/MassEQSTypes.h
+++ b/Engine/Plugins/Runtime/MassGameplay/Source/MassEQS/Public/MassEQSTypes.h
@@ -101,7 +101,7 @@ public:
 		/** Reset RequestID to 0 and RequestStartTime to -1.f */
 		void Reset();
 
-		FORCEINLINE bool IsPendingResults() const { return RequestHandle.IsValid(); }
+		inline bool IsPendingResults() const { return RequestHandle.IsValid(); }
 
 		// Handle for request in MassEQSSubsystem.
 		FMassEQSRequestHandle RequestHandle = FMassEQSRequestHandle::Invalid;
diff --git a/Engine/Plugins/Runtime/MassGameplay/Source/MassEQS/Public/MassEQSUtils.h b/Engine/Plugins/Runtime/MassGameplay/Source/MassEQS/Public/MassEQSUtils.h
--- a/Engine/Plugins/Runtime/MassGameplay/Source/MassEQS/Public/MassEQSUtils.h
+++ b/Engine/Plugins/Runtime/MassGameplay/Source/MassEQS/Public/MassEQSUtils.h
@@ -33,7 +33,7 @@ struct FMassEQSUtils
 	 * If InPtr is not null, then this Cast should never fail.
 	 */
 	template<typename RequestDataType>
-	static FORCEINLINE RequestDataType* TryAndEnsureCast(TUniquePtr<FMassEQSRequestData>& InPtr)
+	static inline RequestDataType* TryAndEnsureCast(TUniquePtr<FMassEQSRequestData>& InPtr)
 	{
 		if (!InPtr)
 		{
diff --git a/Engine/Plugins/Runtime/MassGameplay/Source/MassGameplayDebug/Public/MassDebugDrawHelpers.h b/Engine/Plugins/Runtime/MassGameplay/Source/MassGameplayDebug/Public/MassDebugDrawHelpers.h
--- a/Engine/Plugins/Runtime/MassGameplay/Source/MassGameplayDebug/Public/MassDebugDrawHelpers.h
+++ b/Engine/Plugins/Runtime/MassGameplay/Source/MassGameplayDebug/Public/MassDebugDrawHelpers.h
@@ -26,23 +26,23 @@ namespace UE::Mass::Debug
 		{	
 		}
 
-		FORCEINLINE void DrawSolidBox(const FVector& Center, const FVector& Extent, FColor const& Color) const
+		inline void DrawSolidBox(const FVector& Center, const FVector& Extent, FColor const& Color) const
 		{
 			const FBox Box = FBox::BuildAABB(Center, Extent);
 			LineBatcherInstance->DrawSolidBox(Box, FTransform::Identity, Color, /*DepthPriority=*/0, LifeTime);
 		}
 
-		FORCEINLINE void DrawWireBox(const FVector& Center, const FVector& Extent, FColor const& Color) const
+		inline void DrawWireBox(const FVector& Center, const FVector& Extent, FColor const& Color) const
 		{
 			LineBatcherInstance->DrawBox(Center, Extent, Color, LifeTime, /*DepthPriority=*/0, /*Thickness=*/0.f);
 		}
 
-		FORCEINLINE void DrawSphere(const FVector& Center, const float Radius, FLinearColor const& Color) const
+		inline void DrawSphere(const FVector& Center, const float Radius, FLinearColor const& Color) const
 		{
 			LineBatcherInstance->DrawSphere(Center, Radius, /*Segments=*/8, Color, LifeTime, /*DepthPriority=*/0, /*Thickness=*/0.f);
 		}
 
-		FORCEINLINE void DrawArrow(const FTransform& Transform, const float Length, FColor const& Color) const
+		inline void DrawArrow(const FTransform& Transform, const float Length, FColor const& Color) const
 		{
 			LineBatcherInstance->DrawDirectionalArrow(Transform.ToMatrixNoScale(), Color, Length, Length/5, /*DepthPriority=*/0);
 		}
diff --git a/Engine/Plugins/Runtime/MassGameplay/Source/MassLOD/Public/MassLODCalculator.h b/Engine/Plugins/Runtime/MassGameplay/Source/MassLOD/Public/MassLODCalculator.h
--- a/Engine/Plugins/Runtime/MassGameplay/Source/MassLOD/Public/MassLODCalculator.h
+++ b/Engine/Plugins/Runtime/MassGameplay/Source/MassLOD/Public/MassLODCalculator.h
@@ -54,7 +54,7 @@ struct TMassLODCalculator : public FMassLODBaseLogic
 			  typename TLODFragment,
 			  bool bCalculateLocalViewers = FLODLogic::bLocalViewersOnly,
 			  bool bCalculateVisibility = FLODLogic::bDoVisibilityLogic>
-	FORCEINLINE void CalculateLOD(FMassExecutionContext& Context,
+	inline void CalculateLOD(FMassExecutionContext& Context,
 					  			  TConstArrayView<TViewerInfoFragment> ViewersInfoList,
 					  			  TArrayView<TLODFragment> LODList)
 	{
@@ -110,7 +110,7 @@ struct TMassLODCalculator : public FMassLODBaseLogic
 			  typename TLODFragment,
 			  bool bCalculateLocalViewers = FLODLogic::bLocalViewersOnly,
 			  bool bCalculateVisibility = FLODLogic::bDoVisibilityLogic>
-	FORCEINLINE void AdjustLODFromCount(FMassExecutionContext& Context,
+	inline void AdjustLODFromCount(FMassExecutionContext& Context,
 										TConstArrayView<TViewerInfoFragment> ViewersInfoList,
 										TArrayView<TLODFragment> LODList)
 	{
diff --git a/Engine/Plugins/Runtime/MassGameplay/Source/MassLOD/Public/MassLODCollector.h b/Engine/Plugins/Runtime/MassGameplay/Source/MassLOD/Public/MassLODCollector.h
--- a/Engine/Plugins/Runtime/MassGameplay/Source/MassLOD/Public/MassLODCollector.h
+++ b/Engine/Plugins/Runtime/MassGameplay/Source/MassLOD/Public/MassLODCollector.h
@@ -36,7 +36,7 @@ struct TMassLODCollector : public FMassLODBaseLogic
 			  typename TViewerInfoFragment,
 			  bool bCollectLocalViewers = FLODLogic::bLocalViewersOnly,
 			  bool bCollectDistanceToFrustum = FLODLogic::bDoVisibilityLogic>
-	FORCEINLINE void CollectLODInfo(FMassExecutionContext& Context, 
+	inline void CollectLODInfo(FMassExecutionContext& Context, 
 									TConstArrayView<TTransformFragment> TransformList, 
 									TArrayView<TViewerInfoFragment> ViewersInfoList)
 	{
diff --git a/Engine/Plugins/Runtime/MassGameplay/Source/MassLOD/Public/MassLODLogic.h b/Engine/Plugins/Runtime/MassGameplay/Source/MassLOD/Public/MassLODLogic.h
--- a/Engine/Plugins/Runtime/MassGameplay/Source/MassLOD/Public/MassLODLogic.h
+++ b/Engine/Plugins/Runtime/MassGameplay/Source/MassLOD/Public/MassLODLogic.h
@@ -10,18 +10,18 @@
 #define UE_API MASSLOD_API
 
 #define DECLARE_CONDITIONAL_MEMBER_ACCESSORS( Condition, MemberType, MemberName ) \
-template <bool Condition, typename TemplateClass> static FORCEINLINE typename TEnableIf< Condition, MemberType>::Type Get##MemberName(TemplateClass& Obj, MemberType DefaultValue) { return Obj.MemberName; } \
-template <bool Condition, typename TemplateClass> static FORCEINLINE typename TEnableIf<!Condition, MemberType>::Type Get##MemberName(TemplateClass& Obj, MemberType DefaultValue) { return DefaultValue; } \
-template <bool Condition, typename TemplateClass> static FORCEINLINE void Set##MemberName(TemplateClass& Obj, typename TEnableIf< Condition, MemberType>::Type Value) { Obj.MemberName = Value; } \
-template <bool Condition, typename TemplateClass> static FORCEINLINE void Set##MemberName(TemplateClass& Obj, typename TEnableIf<!Condition, MemberType>::Type Value) {}
+template <bool Condition, typename TemplateClass> static inline typename TEnableIf< Condition, MemberType>::Type Get##MemberName(TemplateClass& Obj, MemberType DefaultValue) { return Obj.MemberName; } \
+template <bool Condition, typename TemplateClass> static inline typename TEnableIf<!Condition, MemberType>::Type Get##MemberName(TemplateClass& Obj, MemberType DefaultValue) { return DefaultValue; } \
+template <bool Condition, typename TemplateClass> static inline void Set##MemberName(TemplateClass& Obj, typename TEnableIf< Condition, MemberType>::Type Value) { Obj.MemberName = Value; } \
+template <bool Condition, typename TemplateClass> static inline void Set##MemberName(TemplateClass& Obj, typename TEnableIf<!Condition, MemberType>::Type Value) {}
 
 #define DECLARE_CONDITIONAL_MEMBER_ARRAY_ACCESSORS( Condition, MemberType, MemberName ) \
-template <bool Condition, typename TemplateClass> static FORCEINLINE void Set##MemberName##Num(TemplateClass& Obj, typename TEnableIf< Condition, int32>::Type Num) { Obj.MemberName.SetNum(Num); } \
-template <bool Condition, typename TemplateClass> static FORCEINLINE void Set##MemberName##Num(TemplateClass& Obj, typename TEnableIf<!Condition, int32>::Type Num) {} \
-template <bool Condition, typename TemplateClass> static FORCEINLINE typename TEnableIf< Condition, MemberType>::Type Get##MemberName(TemplateClass& Obj, int32 Index, MemberType DefaultValue) { return Obj.MemberName[Index]; } \
-template <bool Condition, typename TemplateClass> static FORCEINLINE typename TEnableIf<!Condition, MemberType>::Type Get##MemberName(TemplateClass& Obj, int32 Index, MemberType DefaultValue) { return DefaultValue; } \
-template <bool Condition, typename TemplateClass> static FORCEINLINE void Set##MemberName(TemplateClass& Obj, int32 Index, typename TEnableIf< Condition, MemberType>::Type Value) { Obj.MemberName[Index] = Value; } \
-template <bool Condition, typename TemplateClass> static FORCEINLINE void Set##MemberName(TemplateClass& Obj, int32 Index, typename TEnableIf<!Condition, MemberType>::Type Value) {}
+template <bool Condition, typename TemplateClass> static inline void Set##MemberName##Num(TemplateClass& Obj, typename TEnableIf< Condition, int32>::Type Num) { Obj.MemberName.SetNum(Num); } \
+template <bool Condition, typename TemplateClass> static inline void Set##MemberName##Num(TemplateClass& Obj, typename TEnableIf<!Condition, int32>::Type Num) {} \
+template <bool Condition, typename TemplateClass> static inline typename TEnableIf< Condition, MemberType>::Type Get##MemberName(TemplateClass& Obj, int32 Index, MemberType DefaultValue) { return Obj.MemberName[Index]; } \
+template <bool Condition, typename TemplateClass> static inline typename TEnableIf<!Condition, MemberType>::Type Get##MemberName(TemplateClass& Obj, int32 Index, MemberType DefaultValue) { return DefaultValue; } \
+template <bool Condition, typename TemplateClass> static inline void Set##MemberName(TemplateClass& Obj, int32 Index, typename TEnableIf< Condition, MemberType>::Type Value) { Obj.MemberName[Index] = Value; } \
+template <bool Condition, typename TemplateClass> static inline void Set##MemberName(TemplateClass& Obj, int32 Index, typename TEnableIf<!Condition, MemberType>::Type Value) {}
 
 /**
  * Traits for LOD logic calculation behaviors
diff --git a/Engine/Plugins/Runtime/MassGameplay/Source/MassLOD/Public/MassLODTickRateController.h b/Engine/Plugins/Runtime/MassGameplay/Source/MassLOD/Public/MassLODTickRateController.h
--- a/Engine/Plugins/Runtime/MassGameplay/Source/MassLOD/Public/MassLODTickRateController.h
+++ b/Engine/Plugins/Runtime/MassGameplay/Source/MassLOD/Public/MassLODTickRateController.h
@@ -52,7 +52,7 @@ public:
 	template <typename TLODFragment, typename TVariableTickRateFragment>
 	bool UpdateTickRateFromLOD(FMassExecutionContext& Context, TConstArrayView<TLODFragment> LODList, TArrayView<TVariableTickRateFragment> TickRateList, const double Time);
 
-	friend FORCEINLINE uint32 GetTypeHash(const TMassLODTickRateController<TVariableTickChunkFragment, FLODLogic>& SharedFragmentInstance)
+	friend inline uint32 GetTypeHash(const TMassLODTickRateController<TVariableTickChunkFragment, FLODLogic>& SharedFragmentInstance)
 	{
 		return HashCombineFast(GetTypeHash(SharedFragmentInstance.TickRates), uint32(SharedFragmentInstance.bShouldSpreadFirstUpdate));
 	}
diff --git a/Engine/Plugins/Runtime/MassGameplay/Source/MassRepresentation/Public/MassRepresentationTypes.h b/Engine/Plugins/Runtime/MassGameplay/Source/MassRepresentation/Public/MassRepresentationTypes.h
--- a/Engine/Plugins/Runtime/MassGameplay/Source/MassRepresentation/Public/MassRepresentationTypes.h
+++ b/Engine/Plugins/Runtime/MassGameplay/Source/MassRepresentation/Public/MassRepresentationTypes.h
@@ -106,7 +106,7 @@ struct FMassStaticMeshInstanceVisualizationMeshDesc
 			LocalTransform.Equals(Other.LocalTransform);
 	}
 	
-	friend FORCEINLINE uint32 GetTypeHash(const FMassStaticMeshInstanceVisualizationMeshDesc& MeshDesc)
+	friend inline uint32 GetTypeHash(const FMassStaticMeshInstanceVisualizationMeshDesc& MeshDesc)
 	{
 		uint32 Hash = 0x0;
 		Hash = PointerHash(MeshDesc.Mesh, Hash);
@@ -198,7 +198,7 @@ struct alignas(2) FStaticMeshInstanceVisualizationDescHandle
 		}
 	}
 
-	FORCEINLINE int32 ToIndex() const
+	inline int32 ToIndex() const
 	{
 		return IsValid() ? Index : INDEX_NONE;
 	}
@@ -643,7 +643,7 @@ public:
 	bool ShouldUseTransformOffset() const { return Desc.bUseTransformOffset; }
 	FTransform GetTransformOffset() const { return Desc.TransformOffset; }
 
-	FORCEINLINE FMassLODSignificanceRange* GetLODSignificanceRange(float LODSignificance)
+	inline FMassLODSignificanceRange* GetLODSignificanceRange(float LODSignificance)
 	{
 		for (FMassLODSignificanceRange& Range : LODSignificanceRanges)
 		{
@@ -672,7 +672,7 @@ public:
 		}
 	}
 
-	FORCEINLINE void RemoveInstance(const FMassEntityHandle EntityHandle, const float LODSignificance)
+	inline void RemoveInstance(const FMassEntityHandle EntityHandle, const float LODSignificance)
 	{
 		if (FMassLODSignificanceRange* Range = GetLODSignificanceRange(LODSignificance))
 		{
@@ -700,7 +700,7 @@ public:
 		}
 	}
 
-	FORCEINLINE void AddBatchedCustomDataFloats(const TArray<float>& CustomFloats, const float LODSignificance, const float PrevLODSignificance = -1.0f)
+	inline void AddBatchedCustomDataFloats(const TArray<float>& CustomFloats, const float LODSignificance, const float PrevLODSignificance = -1.0f)
 	{
 		if (FMassLODSignificanceRange* Range = GetLODSignificanceRange(LODSignificance))
 		{
@@ -809,7 +809,7 @@ struct FMassInstancedStaticMeshInfoArrayViewAccessDetector
 		}
 	}
 
-	FORCEINLINE FMassInstancedStaticMeshInfo& operator[](int32 Index) const
+	inline FMassInstancedStaticMeshInfo& operator[](int32 Index) const
 	{
 		return InstancedStaticMeshInfos[Index];
 	}
diff --git a/Engine/Source/Runtime/MassEntity/Public/MassArchetypeTypes.h b/Engine/Source/Runtime/MassEntity/Public/MassArchetypeTypes.h
--- a/Engine/Source/Runtime/MassEntity/Public/MassArchetypeTypes.h
+++ b/Engine/Source/Runtime/MassEntity/Public/MassArchetypeTypes.h
@@ -330,7 +330,7 @@ struct FMassQueryRequirementIndicesMapping
 	FMassFragmentIndicesMapping ChunkFragments;
 	FMassFragmentIndicesMapping ConstSharedFragments;
 	FMassFragmentIndicesMapping SharedFragments;
-	FORCEINLINE bool IsEmpty() const
+	inline bool IsEmpty() const
 	{
 		return EntityFragments.Num() == 0 || ChunkFragments.Num() == 0;
 	}
diff --git a/Engine/Source/Runtime/MassEntity/Public/MassBitSetRegistry.h b/Engine/Source/Runtime/MassEntity/Public/MassBitSetRegistry.h
--- a/Engine/Source/Runtime/MassEntity/Public/MassBitSetRegistry.h
+++ b/Engine/Source/Runtime/MassEntity/Public/MassBitSetRegistry.h
@@ -89,7 +89,7 @@ namespace UE::Mass
 			 * @param BitSet the bitset instance to modify
 			 * @return A new FBitSetBuilder instance
 			 */
-			[[nodiscard]] FORCEINLINE FBitSetBuilder MakeBuilder(FBitSet& BitSet) const
+			[[nodiscard]] inline FBitSetBuilder MakeBuilder(FBitSet& BitSet) const
 			{
 				return FBitSetBuilder(StructTracker, BitSet);
 			}
@@ -98,7 +98,7 @@ namespace UE::Mass
 			 * Creates a factory for building new bitsets, essentially a FBitSetBuilder-BitSet combo.
 			 * @return A new FBitSetFactory instance.
 			 */
-			[[nodiscard]] FORCEINLINE FBitSetFactory MakeBuilder() const
+			[[nodiscard]] inline FBitSetFactory MakeBuilder() const
 			{
 				return FBitSetFactory(StructTracker);
 			}
@@ -108,7 +108,7 @@ namespace UE::Mass
 			 * @param Type the UScriptStruct representing the type to register.
 			 * @return The index assigned to the registered type.
 			 */
-			FORCEINLINE int32 RegisterType(const UScriptStruct* Type)
+			inline int32 RegisterType(const UScriptStruct* Type)
 			{
 				return StructTracker.Register(*Type);
 			}
@@ -119,7 +119,7 @@ namespace UE::Mass
 			 * @return The index assigned to the registered type.
 			 */
 			template<typename TType>
-			FORCEINLINE int32 RegisterType()
+			inline int32 RegisterType()
 			{
 				return RegisterType(TType::StaticStrict());
 			}
diff --git a/Engine/Source/Runtime/MassEntity/Public/MassDebugLogging.h b/Engine/Source/Runtime/MassEntity/Public/MassDebugLogging.h
--- a/Engine/Source/Runtime/MassEntity/Public/MassDebugLogging.h
+++ b/Engine/Source/Runtime/MassEntity/Public/MassDebugLogging.h
@@ -37,7 +37,7 @@ namespace UE::Mass::Debug
 		{	
 		}
 
-		FORCEINLINE bool ShouldLogEntity(int32 EntityIndex, FColor* OutEntityColor = nullptr) const
+		inline bool ShouldLogEntity(int32 EntityIndex, FColor* OutEntityColor = nullptr) const
 		{
 #if ENABLE_VISUAL_LOG
 			if (bLogEverythingWhenRecording
@@ -53,7 +53,7 @@ namespace UE::Mass::Debug
 				&& IsDebuggingEntity(EntityListView[EntityIndex], OutEntityColor);
 		}
 
-		FORCEINLINE const UObject* GetLogOwner(int32 EntityIndex, const UObject* FallbackOwner = nullptr) const
+		inline const UObject* GetLogOwner(int32 EntityIndex, const UObject* FallbackOwner = nullptr) const
 		{
 			return DebugFragmentsView.IsValidIndex(EntityIndex)
 				? DebugFragmentsView[EntityIndex].LogOwner.Get()
@@ -80,7 +80,7 @@ namespace UE::Mass::Debug
 			return false;
 		}
 
-		FORCEINLINE UObject* GetLogOwner(int32 EntityIndex, UObject* FallbackOwner = nullptr) const
+		inline UObject* GetLogOwner(int32 EntityIndex, UObject* FallbackOwner = nullptr) const
 		{
 			return FallbackOwner;
 		}
diff --git a/Engine/Source/Runtime/MassEntity/Public/MassEntityHandle.h b/Engine/Source/Runtime/MassEntity/Public/MassEntityHandle.h
--- a/Engine/Source/Runtime/MassEntity/Public/MassEntityHandle.h
+++ b/Engine/Source/Runtime/MassEntity/Public/MassEntityHandle.h
@@ -45,7 +45,7 @@ struct alignas(8) FMassEntityHandle
 		return Index != 0 && SerialNumber != 0;
 	}
 
-	FORCEINLINE bool IsValid() const
+	inline bool IsValid() const
 	{
 		return IsSet();
 	}
diff --git a/Engine/Source/Runtime/MassEntity/Public/MassEntityManager.h b/Engine/Source/Runtime/MassEntity/Public/MassEntityManager.h
--- a/Engine/Source/Runtime/MassEntity/Public/MassEntityManager.h
+++ b/Engine/Source/Runtime/MassEntity/Public/MassEntityManager.h
@@ -290,7 +290,7 @@ public:
 	 * @return a creation context that will notify all the interested observers about newly created fragments once the context is released */
 	UE_API TSharedRef<FEntityCreationContext> BatchCreateReservedEntities(const FMassArchetypeHandle& ArchetypeHandle,
 		const FMassArchetypeSharedFragmentValues& SharedFragmentValues, TConstArrayView<FMassEntityHandle> ReservedEntities);
-	FORCEINLINE TSharedRef<FEntityCreationContext> BatchCreateReservedEntities(const FMassArchetypeHandle& ArchetypeHandle,
+	inline TSharedRef<FEntityCreationContext> BatchCreateReservedEntities(const FMassArchetypeHandle& ArchetypeHandle,
 		TConstArrayView<FMassEntityHandle> OutEntities)
 	{
 		return BatchCreateReservedEntities(ArchetypeHandle, FMassArchetypeSharedFragmentValues(), OutEntities);
@@ -303,7 +303,7 @@ public:
 	 * @param InOutEntities the newly created entities are appended to given array, i.e. the pre-existing content of OutEntities won't be affected by the call
 	 * @return a creation context that will notify all the interested observers about newly created fragments once the context is released */
 	UE_API TSharedRef<FEntityCreationContext> BatchCreateEntities(const FMassArchetypeHandle& ArchetypeHandle, const FMassArchetypeSharedFragmentValues& SharedFragmentValues, const int32 Count, TArray<FMassEntityHandle>& InOutEntities);
-	FORCEINLINE TSharedRef<FEntityCreationContext> BatchCreateEntities(const FMassArchetypeHandle& ArchetypeHandle, const int32 Count, TArray<FMassEntityHandle>& InOutEntities)
+	inline TSharedRef<FEntityCreationContext> BatchCreateEntities(const FMassArchetypeHandle& ArchetypeHandle, const int32 Count, TArray<FMassEntityHandle>& InOutEntities)
 	{
 		return BatchCreateEntities(ArchetypeHandle, FMassArchetypeSharedFragmentValues(), Count, InOutEntities);
 	}
diff --git a/Engine/Source/Runtime/MassEntity/Public/MassEntityQuery.h b/Engine/Source/Runtime/MassEntity/Public/MassEntityQuery.h
--- a/Engine/Source/Runtime/MassEntity/Public/MassEntityQuery.h
+++ b/Engine/Source/Runtime/MassEntity/Public/MassEntityQuery.h
@@ -118,7 +118,7 @@ public:
 		DirtyCachedData();
 	}
 
-	FORCEINLINE void DirtyCachedData()
+	inline void DirtyCachedData()
 	{
 		EntitySubsystemHash = 0;
 		LastUpdatedArchetypeDataVersion = 0;
diff --git a/Engine/Source/Runtime/MassEntity/Public/MassEntityTypes.h b/Engine/Source/Runtime/MassEntity/Public/MassEntityTypes.h
--- a/Engine/Source/Runtime/MassEntity/Public/MassEntityTypes.h
+++ b/Engine/Source/Runtime/MassEntity/Public/MassEntityTypes.h
@@ -243,28 +243,28 @@ struct FMassArchetypeSharedFragmentValues
 	FMassArchetypeSharedFragmentValues& operator=(const FMassArchetypeSharedFragmentValues& OtherFragmentValues) = default;
 	FMassArchetypeSharedFragmentValues& operator=(FMassArchetypeSharedFragmentValues&& OtherFragmentValues) = default;
 
-	FORCEINLINE bool HasExactFragmentTypesMatch(const FMassSharedFragmentBitSet& InSharedFragmentBitSet, const FMassConstSharedFragmentBitSet& InConstSharedFragmentBitSet) const
+	inline bool HasExactFragmentTypesMatch(const FMassSharedFragmentBitSet& InSharedFragmentBitSet, const FMassConstSharedFragmentBitSet& InConstSharedFragmentBitSet) const
 	{
 		return HasExactSharedFragmentTypesMatch(InSharedFragmentBitSet)
 			&& HasExactConstSharedFragmentTypesMatch(InConstSharedFragmentBitSet);
 	}
 
-	FORCEINLINE bool HasExactSharedFragmentTypesMatch(const FMassSharedFragmentBitSet& InSharedFragmentBitSet) const
+	inline bool HasExactSharedFragmentTypesMatch(const FMassSharedFragmentBitSet& InSharedFragmentBitSet) const
 	{
 		return SharedFragmentBitSet.IsEquivalent(InSharedFragmentBitSet);
 	}
 
-	FORCEINLINE bool HasAllRequiredSharedFragmentTypes(const FMassSharedFragmentBitSet& InSharedFragmentBitSet) const
+	inline bool HasAllRequiredSharedFragmentTypes(const FMassSharedFragmentBitSet& InSharedFragmentBitSet) const
 	{
 		return SharedFragmentBitSet.HasAll(InSharedFragmentBitSet);
 	}
 
-	FORCEINLINE bool HasExactConstSharedFragmentTypesMatch(const FMassConstSharedFragmentBitSet& InConstSharedFragmentBitSet) const
+	inline bool HasExactConstSharedFragmentTypesMatch(const FMassConstSharedFragmentBitSet& InConstSharedFragmentBitSet) const
 	{
 		return ConstSharedFragmentBitSet.IsEquivalent(InConstSharedFragmentBitSet);
 	}
 
-	FORCEINLINE bool HasAllRequiredConstSharedFragmentTypes(const FMassConstSharedFragmentBitSet& InConstSharedFragmentBitSet) const
+	inline bool HasAllRequiredConstSharedFragmentTypes(const FMassConstSharedFragmentBitSet& InConstSharedFragmentBitSet) const
 	{
 		return ConstSharedFragmentBitSet.HasAll(InConstSharedFragmentBitSet);
 	}
@@ -278,7 +278,7 @@ struct FMassArchetypeSharedFragmentValues
 			&& HasExactConstSharedFragmentTypesMatch(InDescriptor.ConstSharedFragments);
 	}
 
-	FORCEINLINE bool IsEquivalent(const FMassArchetypeSharedFragmentValues& OtherSharedFragmentValues) const
+	inline bool IsEquivalent(const FMassArchetypeSharedFragmentValues& OtherSharedFragmentValues) const
 	{
 		return GetTypeHash(*this) == GetTypeHash(OtherSharedFragmentValues);
 	}
@@ -291,7 +291,7 @@ struct FMassArchetypeSharedFragmentValues
 	 */
 	MASSENTITY_API bool HasSameValues(const FMassArchetypeSharedFragmentValues& Other) const;
 
-	FORCEINLINE bool ContainsType(const UScriptStruct* FragmentType) const
+	inline bool ContainsType(const UScriptStruct* FragmentType) const
 	{
 		if (UE::Mass::IsA<FMassSharedFragment>(FragmentType))
 		{
@@ -307,7 +307,7 @@ struct FMassArchetypeSharedFragmentValues
 	}
 
 	template<typename T>
-	FORCEINLINE bool ContainsType() const
+	inline bool ContainsType() const
 	{
 		if constexpr (UE::Mass::CConstSharedFragment<T>)
 		{
@@ -406,17 +406,17 @@ struct FMassArchetypeSharedFragmentValues
 		return Remove(InDescriptor.SharedFragments) + Remove(InDescriptor.ConstSharedFragments);
 	}
 
-	FORCEINLINE const TArray<FConstSharedStruct>& GetConstSharedFragments() const
+	inline const TArray<FConstSharedStruct>& GetConstSharedFragments() const
 	{
 		return ConstSharedFragments;
 	}
 
-	FORCEINLINE TArray<FSharedStruct>& GetMutableSharedFragments()
+	inline TArray<FSharedStruct>& GetMutableSharedFragments()
 	{
 		return SharedFragments;
 	}
 	
-	FORCEINLINE const TArray<FSharedStruct>& GetSharedFragments() const
+	inline const TArray<FSharedStruct>& GetSharedFragments() const
 	{
 		return SharedFragments;
 	}
@@ -449,14 +449,14 @@ struct FMassArchetypeSharedFragmentValues
 		return ConstSharedFragmentBitSet;
 	}
 
-	FORCEINLINE void DirtyHashCache()
+	inline void DirtyHashCache()
 	{
 		HashCache = UINT32_MAX;
 		// we consider a single shared fragment as being "sorted"
 		bSorted = (SharedFragments.Num() + ConstSharedFragments.Num() <= 1) ;
 	}
 
-	FORCEINLINE void CacheHash() const
+	inline void CacheHash() const
 	{
 		if (HashCache == UINT32_MAX)
 		{
@@ -464,7 +464,7 @@ struct FMassArchetypeSharedFragmentValues
 		}
 	}
 
-	friend FORCEINLINE uint32 GetTypeHash(const FMassArchetypeSharedFragmentValues& SharedFragmentValues)
+	friend inline uint32 GetTypeHash(const FMassArchetypeSharedFragmentValues& SharedFragmentValues)
 	{
 		SharedFragmentValues.CacheHash();
 		return SharedFragmentValues.HashCache;
@@ -507,7 +507,7 @@ public:
 	// DEPRECATED
 	//-----------------------------------------------------------------------------
 	UE_DEPRECATED(5.5, "HasExactFragmentTypesMatch is deprecated. Use HasExactSharedFragmentTypesMatch or the two-parameter version of HasExactFragmentTypesMatch.")
-	FORCEINLINE bool HasExactFragmentTypesMatch(const FMassSharedFragmentBitSet& InSharedFragmentBitSet) const
+	inline bool HasExactFragmentTypesMatch(const FMassSharedFragmentBitSet& InSharedFragmentBitSet) const
 	{
 		return HasExactSharedFragmentTypesMatch(InSharedFragmentBitSet);
 	}
@@ -578,7 +578,7 @@ struct FMassGenericPayloadView
 		Content = TArrayView<FStructArrayView>();
 	}
 
-	FORCEINLINE void Swap(const int32 A, const int32 B)
+	inline void Swap(const int32 A, const int32 B)
 	{
 		for (FStructArrayView& View : Content)
 		{
diff --git a/Engine/Source/Runtime/MassEntity/Public/MassExecutionContext.h b/Engine/Source/Runtime/MassEntity/Public/MassExecutionContext.h
--- a/Engine/Source/Runtime/MassEntity/Public/MassExecutionContext.h
+++ b/Engine/Source/Runtime/MassEntity/Public/MassExecutionContext.h
@@ -172,32 +172,32 @@ public:
 	 */
 	struct FEntityIterator
 	{
-		FORCEINLINE int32 operator*() const
+		inline int32 operator*() const
 		{
 			return EntityIndex;
 		}
 
-		FORCEINLINE bool operator!=(const int& Other) const
+		inline bool operator!=(const int& Other) const
 		{
 			return EntityIndex != Other;
 		}
 
-		FORCEINLINE operator int32() const
+		inline operator int32() const
 		{
 			return EntityIndex;
 		}
 
-		FORCEINLINE operator bool() const
+		inline operator bool() const
 		{
 			return SerialNumber && EntityIndex < NumEntities;
 		}
 
-		FORCEINLINE bool operator<(const int32 Other) const
+		inline bool operator<(const int32 Other) const
 		{
 			return SerialNumber && EntityIndex != INDEX_NONE && EntityIndex < Other;
 		}
 
-		FORCEINLINE FEntityIterator& operator++()
+		inline FEntityIterator& operator++()
 		{
 			++EntityIndex;
 #if WITH_MASSENTITY_DEBUG
@@ -210,7 +210,7 @@ public:
 			return *this;
 		}
 
-		FORCEINLINE void operator++(int)
+		inline void operator++(int)
 		{
 			++(*this);
 		}
diff --git a/Engine/Source/Runtime/MassEntity/Public/MassObserverNotificationTypes.h b/Engine/Source/Runtime/MassEntity/Public/MassObserverNotificationTypes.h
--- a/Engine/Source/Runtime/MassEntity/Public/MassObserverNotificationTypes.h
+++ b/Engine/Source/Runtime/MassEntity/Public/MassObserverNotificationTypes.h
@@ -92,7 +92,7 @@ namespace UE::Mass
 				StoredCollection->AddHandle(EntityHandle);
 			}
 
-			FORCEINLINE void AppendEntities(const FMassEntityHandle EntityHandle)
+			inline void AppendEntities(const FMassEntityHandle EntityHandle)
 			{
 				AddHandle(EntityHandle);
 			}
diff --git a/Engine/Source/Runtime/MassEntity/Public/MassQueryExecutor.h b/Engine/Source/Runtime/MassEntity/Public/MassQueryExecutor.h
--- a/Engine/Source/Runtime/MassEntity/Public/MassQueryExecutor.h
+++ b/Engine/Source/Runtime/MassEntity/Public/MassQueryExecutor.h
@@ -53,13 +53,13 @@ struct FQueryExecutor
 
 protected:
 
-	FORCEINLINE UObject* GetLogOwner()
+	inline UObject* GetLogOwner()
 	{
 		return LogOwner.Get();
 	}
 
 	template<typename TAccessors, typename TFunc>
-	FORCEINLINE void ForEachEntityChunk(FMassExecutionContext& ExecutionContext, TAccessors& Accessors, const TFunc&& ExecuteFunction)
+	inline void ForEachEntityChunk(FMassExecutionContext& ExecutionContext, TAccessors& Accessors, const TFunc&& ExecuteFunction)
 	{
 		BoundQuery->ForEachEntityChunk(ExecutionContext, [&Accessors, &ExecuteFunction](FMassExecutionContext& Context)
 		{
@@ -72,7 +72,7 @@ protected:
 	}
 
 	template<typename TAccessors, typename TFunc>
-	FORCEINLINE void ParallelForEachEntityChunk(FMassExecutionContext& ExecutionContext, const TAccessors& Accessors, const TFunc&& ExecuteFunction)
+	inline void ParallelForEachEntityChunk(FMassExecutionContext& ExecutionContext, const TAccessors& Accessors, const TFunc&& ExecuteFunction)
 	{
 		BoundQuery->ParallelForEachEntityChunk(ExecutionContext, [&Accessors, &ExecuteFunction](FMassExecutionContext& Context)
 		{
@@ -88,7 +88,7 @@ protected:
 	}
 
 	template<typename TAccessors, typename TFunc>
-	FORCEINLINE void ForEachEntity(FMassExecutionContext& ExecutionContext, TAccessors& Accessors, const TFunc&& ExecuteFunction)
+	inline void ForEachEntity(FMassExecutionContext& ExecutionContext, TAccessors& Accessors, const TFunc&& ExecuteFunction)
 	{
 		BoundQuery->ForEachEntityChunk(ExecutionContext, [&Accessors, &ExecuteFunction](FMassExecutionContext& Context)
 		{
@@ -105,7 +105,7 @@ protected:
 	}
 
 	template<typename TAccessors, typename TFunc>
-	FORCEINLINE void ParallelForEachEntity(FMassExecutionContext& ExecutionContext, TAccessors& Accessors, const TFunc&& ExecuteFunction)
+	inline void ParallelForEachEntity(FMassExecutionContext& ExecutionContext, TAccessors& Accessors, const TFunc&& ExecuteFunction)
 	{
 		BoundQuery->ParallelForEachEntityChunk(ExecutionContext, [&Accessors, &ExecuteFunction](FMassExecutionContext& Context)
 		{
@@ -200,7 +200,7 @@ struct FQueryDefinition : public FQueryDefinitionBase
 		});
 	}
 
-	FORCEINLINE void SetupForChunk(FMassExecutionContext& Context)
+	inline void SetupForChunk(FMassExecutionContext& Context)
 	{
 		AccessorTuple.ApplyBefore([&Context](auto&... Accessor)
 		{
@@ -209,7 +209,7 @@ struct FQueryDefinition : public FQueryDefinitionBase
 	}
 
 	template <typename TFragment>
-	FORCEINLINE constexpr auto& Get()
+	inline constexpr auto& Get()
 	{
 		constexpr std::size_t Index = GetAccessorIndex<TFragment, Accessors...>();
 		return AccessorTuple.template Get<Index>();
@@ -226,26 +226,26 @@ struct FMutableFragmentAccess
 
 	FMutableFragmentAccess() = default;
 
-	FORCEINLINE TArrayView<TFragment>& Get()
+	inline TArrayView<TFragment>& Get()
 	{
 		return View;
 	}
 
-	FORCEINLINE TFragment& operator[](int32 Index)
+	inline TFragment& operator[](int32 Index)
 	{
 		return View[Index];
 	}
 
-	FORCEINLINE void ConfigureQuery(FMassEntityQuery& EntityQuery, FMassSubsystemRequirements& ProcessorRequirements)
+	inline void ConfigureQuery(FMassEntityQuery& EntityQuery, FMassSubsystemRequirements& ProcessorRequirements)
 	{
 		EntityQuery.AddRequirement<TFragment>(EMassFragmentAccess::ReadWrite);
 	}
 
-	FORCEINLINE void SetupForExecute(FMassExecutionContext& Context)
+	inline void SetupForExecute(FMassExecutionContext& Context)
 	{
 	}
 
-	FORCEINLINE void SetupForChunk(FMassExecutionContext& Context)
+	inline void SetupForChunk(FMassExecutionContext& Context)
 	{
 		View = Context.GetMutableFragmentView<TFragment>();
 	}
@@ -263,26 +263,26 @@ struct FConstFragmentAccess
 
 	FConstFragmentAccess() = default;
 
-	FORCEINLINE const TConstArrayView<TFragment>& Get() const
+	inline const TConstArrayView<TFragment>& Get() const
 	{
 		return View;
 	}
 
-	FORCEINLINE const TFragment& operator[](int32 Index) const
+	inline const TFragment& operator[](int32 Index) const
 	{
 		return View[Index];
 	}
 
-	FORCEINLINE void ConfigureQuery(FMassEntityQuery& EntityQuery, FMassSubsystemRequirements& ProcessorRequirements) const
+	inline void ConfigureQuery(FMassEntityQuery& EntityQuery, FMassSubsystemRequirements& ProcessorRequirements) const
 	{
 		EntityQuery.AddRequirement<TFragment>(EMassFragmentAccess::ReadOnly);
 	}
 
-	FORCEINLINE void SetupForExecute(FMassExecutionContext& Context)
+	inline void SetupForExecute(FMassExecutionContext& Context)
 	{
 	}
 
-	FORCEINLINE void SetupForChunk(FMassExecutionContext& Context)
+	inline void SetupForChunk(FMassExecutionContext& Context)
 	{
 		View = Context.GetFragmentView<TFragment>();
 	}
@@ -300,36 +300,36 @@ struct FMutableOptionalFragmentAccess
 
 	FMutableOptionalFragmentAccess() = default;
 
-	FORCEINLINE TArrayView<TFragment>& Get() const
+	inline TArrayView<TFragment>& Get() const
 	{
 		return View;
 	}
 
-	FORCEINLINE TFragment& operator[](int32 Index) const
+	inline TFragment& operator[](int32 Index) const
 	{
 		return View[Index];
 	}
 
-	FORCEINLINE int32 Num() const
+	inline int32 Num() const
 	{
 		return View.Num();
 	}
 
-	FORCEINLINE operator bool() const
+	inline operator bool() const
 	{
 		return View.Num() > 0;
 	}
 
-	FORCEINLINE void ConfigureQuery(FMassEntityQuery& EntityQuery, FMassSubsystemRequirements& ProcessorRequirements) const
+	inline void ConfigureQuery(FMassEntityQuery& EntityQuery, FMassSubsystemRequirements& ProcessorRequirements) const
 	{
 		EntityQuery.AddRequirement<TFragment>(EMassFragmentAccess::ReadWrite, EMassFragmentPresence::Optional);
 	}
 
-	FORCEINLINE void SetupForExecute(FMassExecutionContext& Context)
+	inline void SetupForExecute(FMassExecutionContext& Context)
 	{
 	}
 
-	FORCEINLINE void SetupForChunk(FMassExecutionContext& Context)
+	inline void SetupForChunk(FMassExecutionContext& Context)
 	{
 		View = Context.GetMutableFragmentView<TFragment>();
 	}
@@ -347,36 +347,36 @@ struct FConstOptionalFragmentAccess
 
 	FConstOptionalFragmentAccess() = default;
 
-	FORCEINLINE const TConstArrayView<TFragment>& Get() const
+	inline const TConstArrayView<TFragment>& Get() const
 	{
 		return View;
 	}
 
-	FORCEINLINE const TFragment& operator[](int32 Index) const
+	inline const TFragment& operator[](int32 Index) const
 	{
 		return View[Index];
 	}
 
-	FORCEINLINE int32 Num() const
+	inline int32 Num() const
 	{
 		return View.Num();
 	}
 
-	FORCEINLINE operator bool() const
+	inline operator bool() const
 	{
 		return View.Num() > 0;
 	}
 
-	FORCEINLINE void ConfigureQuery(FMassEntityQuery& EntityQuery, FMassSubsystemRequirements& ProcessorRequirements) const
+	inline void ConfigureQuery(FMassEntityQuery& EntityQuery, FMassSubsystemRequirements& ProcessorRequirements) const
 	{
 		EntityQuery.AddRequirement<TFragment>(EMassFragmentAccess::ReadOnly);
 	}
 
-	FORCEINLINE void SetupForExecute(FMassExecutionContext& Context)
+	inline void SetupForExecute(FMassExecutionContext& Context)
 	{
 	}
 
-	FORCEINLINE void SetupForChunk(FMassExecutionContext& Context)
+	inline void SetupForChunk(FMassExecutionContext& Context)
 	{
 		View = Context.GetFragmentView<TFragment>();
 	}
@@ -394,16 +394,16 @@ struct FMassTagRequired
 
 	FMassTagRequired() = default;
 
-	FORCEINLINE void ConfigureQuery(FMassEntityQuery& EntityQuery, FMassSubsystemRequirements& ProcessorRequirements) const
+	inline void ConfigureQuery(FMassEntityQuery& EntityQuery, FMassSubsystemRequirements& ProcessorRequirements) const
 	{
 		EntityQuery.AddTagRequirement<TTag>(EMassFragmentPresence::All);
 	}
 
-	FORCEINLINE void SetupForExecute(FMassExecutionContext& Context)
+	inline void SetupForExecute(FMassExecutionContext& Context)
 	{
 	}
 
-	FORCEINLINE void SetupForChunk(FMassExecutionContext& Context)
+	inline void SetupForChunk(FMassExecutionContext& Context)
 	{
 	}
 };
@@ -413,16 +413,16 @@ struct FMassTagBlocked
 {
 	using FragmentType = TFragment;
 
-	FORCEINLINE void ConfigureQuery(FMassEntityQuery& EntityQuery, FMassSubsystemRequirements& ProcessorRequirements) const
+	inline void ConfigureQuery(FMassEntityQuery& EntityQuery, FMassSubsystemRequirements& ProcessorRequirements) const
 	{
 		EntityQuery.AddTagRequirement<TFragment>(EMassFragmentPresence::None);
 	}
 
-	FORCEINLINE void SetupForExecute(FMassExecutionContext& Context)
+	inline void SetupForExecute(FMassExecutionContext& Context)
 	{
 	}
 
-	FORCEINLINE void SetupForChunk(FMassExecutionContext& Context)
+	inline void SetupForChunk(FMassExecutionContext& Context)
 	{
 	}
 };
@@ -437,33 +437,33 @@ struct FMutableSharedFragmentAccess
 
 	FMutableSharedFragmentAccess() = default;
 
-	FORCEINLINE TFragment& Get() const
+	inline TFragment& Get() const
 	{
 		check(Fragment);
 		return *Fragment;
 	}
 
-	FORCEINLINE TFragment& operator*() const
+	inline TFragment& operator*() const
 	{
 		check(Fragment);
 		return *Fragment;
 	}
 
-	FORCEINLINE TFragment* operator->() const
+	inline TFragment* operator->() const
 	{
 		return Fragment;
 	}
 
-	FORCEINLINE void ConfigureQuery(FMassEntityQuery& EntityQuery, FMassSubsystemRequirements& ProcessorRequirements) const
+	inline void ConfigureQuery(FMassEntityQuery& EntityQuery, FMassSubsystemRequirements& ProcessorRequirements) const
 	{
 		EntityQuery.AddSharedRequirement<TFragment>(EMassFragmentAccess::ReadWrite);
 	}
 
-	FORCEINLINE void SetupForExecute(FMassExecutionContext& Context)
+	inline void SetupForExecute(FMassExecutionContext& Context)
 	{
 	}
 
-	FORCEINLINE void SetupForChunk(FMassExecutionContext& Context)
+	inline void SetupForChunk(FMassExecutionContext& Context)
 	{
 		Fragment = &(Context.GetSharedFragment<TFragment>());
 	}
@@ -481,33 +481,33 @@ struct FConstSharedFragmentAccess
 
 	FConstSharedFragmentAccess() = default;
 
-	FORCEINLINE const TFragment& Get() const
+	inline const TFragment& Get() const
 	{
 		check(Fragment);
 		return *Fragment;
 	}
 
-	FORCEINLINE const TFragment& operator*() const
+	inline const TFragment& operator*() const
 	{
 		check(Fragment);
 		return *Fragment;
 	}
 
-	FORCEINLINE const TFragment* operator->() const
+	inline const TFragment* operator->() const
 	{
 		return Fragment;
 	}
 
-	FORCEINLINE void ConfigureQuery(FMassEntityQuery& EntityQuery, FMassSubsystemRequirements& ProcessorRequirements) const
+	inline void ConfigureQuery(FMassEntityQuery& EntityQuery, FMassSubsystemRequirements& ProcessorRequirements) const
 	{
 		EntityQuery.AddConstSharedRequirement<TFragment>();
 	}
 
-	FORCEINLINE void SetupForExecute(FMassExecutionContext& Context)
+	inline void SetupForExecute(FMassExecutionContext& Context)
 	{
 	}
 
-	FORCEINLINE void SetupForChunk(FMassExecutionContext& Context)
+	inline void SetupForChunk(FMassExecutionContext& Context)
 	{
 		Fragment = &Context.GetConstSharedFragment<TFragment>();
 	}
@@ -525,33 +525,33 @@ struct FMutableChunkFragmentAccess
 
 	FMutableChunkFragmentAccess() = default;
 
-	FORCEINLINE TFragment& Get() const
+	inline TFragment& Get() const
 	{
 		check(Fragment);
 		return *Fragment;
 	}
 
-	FORCEINLINE TFragment& operator*() const
+	inline TFragment& operator*() const
 	{
 		check(Fragment);
 		return *Fragment;
 	}
 
-	FORCEINLINE TFragment* operator->() const
+	inline TFragment* operator->() const
 	{
 		return Fragment;
 	}
 
-	FORCEINLINE void ConfigureQuery(FMassEntityQuery& EntityQuery, FMassSubsystemRequirements& ProcessorRequirements) const
+	inline void ConfigureQuery(FMassEntityQuery& EntityQuery, FMassSubsystemRequirements& ProcessorRequirements) const
 	{
 		EntityQuery.AddChunkRequirement<TFragment>(EMassFragmentAccess::ReadWrite, EMassFragmentPresence::All);
 	}
 
-	FORCEINLINE void SetupForExecute(FMassExecutionContext& Context)
+	inline void SetupForExecute(FMassExecutionContext& Context)
 	{
 	}
 
-	FORCEINLINE void SetupForChunk(FMassExecutionContext& Context)
+	inline void SetupForChunk(FMassExecutionContext& Context)
 	{
 		Fragment = Context.GetMutableChunkFragmentPtr<TFragment>();
 	}
@@ -569,33 +569,33 @@ struct FConstChunkFragmentAccess
 
 	FConstChunkFragmentAccess() = default;
 
-	FORCEINLINE const TFragment& Get() const
+	inline const TFragment& Get() const
 	{
 		check(Fragment);
 		return *Fragment;
 	}
 
-	FORCEINLINE const TFragment& operator*() const
+	inline const TFragment& operator*() const
 	{
 		check(Fragment);
 		return *Fragment;
 	}
 
-	FORCEINLINE const TFragment* operator->() const
+	inline const TFragment* operator->() const
 	{
 		return Fragment;
 	}
 
-	FORCEINLINE void ConfigureQuery(FMassEntityQuery& EntityQuery, FMassSubsystemRequirements& ProcessorRequirements) const
+	inline void ConfigureQuery(FMassEntityQuery& EntityQuery, FMassSubsystemRequirements& ProcessorRequirements) const
 	{
 		EntityQuery.AddChunkRequirement<TFragment>(EMassFragmentAccess::ReadOnly, EMassFragmentPresence::All);
 	}
 
-	FORCEINLINE void SetupForExecute(FMassExecutionContext& Context)
+	inline void SetupForExecute(FMassExecutionContext& Context)
 	{
 	}
 
-	FORCEINLINE void SetupForChunk(FMassExecutionContext& Context)
+	inline void SetupForChunk(FMassExecutionContext& Context)
 	{
 		Fragment = Context.GetChunkFragmentPtr<TFragment>();
 	}
@@ -613,35 +613,35 @@ struct FMutableOptionalChunkFragmentAccess
 
 	FMutableOptionalChunkFragmentAccess() = default;
 
-	FORCEINLINE TFragment* Get() const
+	inline TFragment* Get() const
 	{
 		check(Fragment);
 		return Fragment;
 	}
 
-	FORCEINLINE TFragment& operator*() const
+	inline TFragment& operator*() const
 	{
 		check(Fragment);
 		return *Fragment;
 	}
 
-	FORCEINLINE TFragment* operator->() const
+	inline TFragment* operator->() const
 	{
 		return Fragment;
 	}
 
-	FORCEINLINE operator bool() { return Fragment != nullptr; }
+	inline operator bool() { return Fragment != nullptr; }
 
-	FORCEINLINE void ConfigureQuery(FMassEntityQuery& EntityQuery, FMassSubsystemRequirements& ProcessorRequirements) const
+	inline void ConfigureQuery(FMassEntityQuery& EntityQuery, FMassSubsystemRequirements& ProcessorRequirements) const
 	{
 		EntityQuery.AddChunkRequirement<TFragment>(EMassFragmentAccess::ReadWrite, EMassFragmentPresence::Optional);
 	}
 
-	FORCEINLINE void SetupForExecute(FMassExecutionContext& Context)
+	inline void SetupForExecute(FMassExecutionContext& Context)
 	{
 	}
 
-	FORCEINLINE void SetupForChunk(FMassExecutionContext& Context)
+	inline void SetupForChunk(FMassExecutionContext& Context)
 	{
 		Fragment = Context.GetMutableChunkFragmentPtr<TFragment>();
 	}
@@ -659,38 +659,38 @@ struct FConstOptionalChunkFragmentAccess
 
 	FConstOptionalChunkFragmentAccess() = default;
 
-	FORCEINLINE const TFragment* Get()
+	inline const TFragment* Get()
 	{
 		check(Fragment);
 		return Fragment;
 	}
 
-	FORCEINLINE const TFragment& operator*() const
+	inline const TFragment& operator*() const
 	{
 		check(Fragment);
 		return *Fragment;
 	}
 
-	FORCEINLINE const TFragment* operator->() const
+	inline const TFragment* operator->() const
 	{
 		return Fragment;
 	}
 
-	FORCEINLINE operator bool() const
+	inline operator bool() const
 	{
 		return Fragment != nullptr;
 	}
 
-	FORCEINLINE void ConfigureQuery(FMassEntityQuery& EntityQuery, FMassSubsystemRequirements& ProcessorRequirements) const
+	inline void ConfigureQuery(FMassEntityQuery& EntityQuery, FMassSubsystemRequirements& ProcessorRequirements) const
 	{
 		EntityQuery.AddChunkRequirement<TFragment>(EMassFragmentAccess::ReadOnly, EMassFragmentPresence::Optional);
 	}
 
-	FORCEINLINE void SetupForExecute(FMassExecutionContext& Context)
+	inline void SetupForExecute(FMassExecutionContext& Context)
 	{
 	}
 
-	FORCEINLINE void SetupForChunk(FMassExecutionContext& Context)
+	inline void SetupForChunk(FMassExecutionContext& Context)
 	{
 		Fragment = Context.GetChunkFragmentPtr<TFragment>();
 	}
@@ -708,38 +708,38 @@ struct FMutableSubsystemAccess
 
 	FMutableSubsystemAccess() = default;
 
-	FORCEINLINE TSubsystem* Get() const
+	inline TSubsystem* Get() const
 	{
 		return Subsystem;
 	}
 
-	FORCEINLINE TSubsystem& operator*() const
+	inline TSubsystem& operator*() const
 	{
 		check(Subsystem);
 		return *Subsystem;
 	}
 
-	FORCEINLINE TSubsystem* operator->() const
+	inline TSubsystem* operator->() const
 	{
 		return Subsystem;
 	}
 
-	FORCEINLINE operator bool() const
+	inline operator bool() const
 	{
 		return Subsystem != nullptr;
 	}
 
-	FORCEINLINE void ConfigureQuery(FMassEntityQuery& EntityQuery, FMassSubsystemRequirements& ProcessorRequirements) const
+	inline void ConfigureQuery(FMassEntityQuery& EntityQuery, FMassSubsystemRequirements& ProcessorRequirements) const
 	{
 		ProcessorRequirements.AddSubsystemRequirement<TSubsystem>(EMassFragmentAccess::ReadWrite);
 	}
 
-	FORCEINLINE void SetupForExecute(FMassExecutionContext& Context)
+	inline void SetupForExecute(FMassExecutionContext& Context)
 	{
 		Subsystem = Context.GetMutableSubsystem<TSubsystem>();
 	}
 
-	FORCEINLINE void SetupForChunk(FMassExecutionContext& Context)
+	inline void SetupForChunk(FMassExecutionContext& Context)
 	{
 	}
 
@@ -756,38 +756,38 @@ struct FConstSubsystemAccess
 
 	FConstSubsystemAccess() = default;
 
-	FORCEINLINE const TSubsystem* Get() const
+	inline const TSubsystem* Get() const
 	{
 		return Subsystem;
 	}
 
-	FORCEINLINE const TSubsystem& operator*() const
+	inline const TSubsystem& operator*() const
 	{
 		check(Subsystem);
 		return *Subsystem;
 	}
 
-	FORCEINLINE const TSubsystem* operator->() const
+	inline const TSubsystem* operator->() const
 	{
 		return Subsystem;
 	}
 
-	FORCEINLINE operator bool() const
+	inline operator bool() const
 	{
 		return Subsystem != nullptr;
 	}
 
-	FORCEINLINE void ConfigureQuery(FMassEntityQuery& EntityQuery, FMassSubsystemRequirements& ProcessorRequirements) const
+	inline void ConfigureQuery(FMassEntityQuery& EntityQuery, FMassSubsystemRequirements& ProcessorRequirements) const
 	{
 		ProcessorRequirements.AddSubsystemRequirement<TSubsystem>(EMassFragmentAccess::ReadOnly);
 	}
 
-	FORCEINLINE void SetupForExecute(FMassExecutionContext& Context)
+	inline void SetupForExecute(FMassExecutionContext& Context)
 	{
 		Subsystem = Context.GetSubsystem(TSubsystem::StaticClass());
 	}
 
-	FORCEINLINE void SetupForChunk(FMassExecutionContext& Context)
+	inline void SetupForChunk(FMassExecutionContext& Context)
 	{
 	}
 
diff --git a/Engine/Source/Runtime/MassEntity/Public/MassRequirements.h b/Engine/Source/Runtime/MassEntity/Public/MassRequirements.h
--- a/Engine/Source/Runtime/MassEntity/Public/MassRequirements.h
+++ b/Engine/Source/Runtime/MassEntity/Public/MassRequirements.h
@@ -659,7 +659,7 @@ inline uint32 GetTypeHash(const FMassSubsystemRequirements& Instance)
 }
 
 template<>
-FORCEINLINE FMassFragmentRequirements& FMassFragmentRequirements::AddTagRequirements<EMassFragmentPresence::All>(const FMassTagBitSet& TagBitSet)
+inline FMassFragmentRequirements& FMassFragmentRequirements::AddTagRequirements<EMassFragmentPresence::All>(const FMassTagBitSet& TagBitSet)
 {
 	RequiredAllTags += TagBitSet;
 	// force re-caching the next time this query is used or the following CacheArchetypes call.
@@ -668,7 +668,7 @@ FORCEINLINE FMassFragmentRequirements& FMassFragmentRequirements::AddTagRequirem
 }
 
 template<>
-FORCEINLINE FMassFragmentRequirements& FMassFragmentRequirements::AddTagRequirements<EMassFragmentPresence::Any>(const FMassTagBitSet& TagBitSet)
+inline FMassFragmentRequirements& FMassFragmentRequirements::AddTagRequirements<EMassFragmentPresence::Any>(const FMassTagBitSet& TagBitSet)
 {
 	RequiredAnyTags += TagBitSet;
 	// force re-caching the next time this query is used or the following CacheArchetypes call.
@@ -677,7 +677,7 @@ FORCEINLINE FMassFragmentRequirements& FMassFragmentRequirements::AddTagRequirem
 }
 
 template<>
-FORCEINLINE FMassFragmentRequirements& FMassFragmentRequirements::AddTagRequirements<EMassFragmentPresence::None>(const FMassTagBitSet& TagBitSet)
+inline FMassFragmentRequirements& FMassFragmentRequirements::AddTagRequirements<EMassFragmentPresence::None>(const FMassTagBitSet& TagBitSet)
 {
 	RequiredNoneTags += TagBitSet;
 	// force re-caching the next time this query is used or the following CacheArchetypes call.
@@ -686,7 +686,7 @@ FORCEINLINE FMassFragmentRequirements& FMassFragmentRequirements::AddTagRequirem
 }
 
 template<>
-FORCEINLINE FMassFragmentRequirements& FMassFragmentRequirements::AddTagRequirements<EMassFragmentPresence::Optional>(const FMassTagBitSet& TagBitSet)
+inline FMassFragmentRequirements& FMassFragmentRequirements::AddTagRequirements<EMassFragmentPresence::Optional>(const FMassTagBitSet& TagBitSet)
 {
 	RequiredOptionalTags += TagBitSet;
 	// force re-caching the next time this query is used or the following CacheArchetypes call.
@@ -824,7 +824,7 @@ inline bool FMassFragmentRequirements::DoesRequireGameThreadExecution() const
 	return bRequiresGameThreadExecution;
 }
 
-FORCEINLINE void FMassFragmentRequirements::IncrementChangeCounter()
+inline void FMassFragmentRequirements::IncrementChangeCounter()
 { 
 	++IncrementalChangesCount; 
 	bPropertiesCached = false;
diff --git a/Engine/Source/Runtime/MassEntity/Public/MassTypeManager.h b/Engine/Source/Runtime/MassEntity/Public/MassTypeManager.h
--- a/Engine/Source/Runtime/MassEntity/Public/MassTypeManager.h
+++ b/Engine/Source/Runtime/MassEntity/Public/MassTypeManager.h
@@ -37,7 +37,7 @@ namespace UE::Mass
 			return Cast<UScriptStruct>(TypeKey.ResolveObjectPtr());
 		}
 
-		friend FORCEINLINE uint32 GetTypeHash(const FTypeHandle& InHandle)
+		friend inline uint32 GetTypeHash(const FTypeHandle& InHandle)
 		{
 			return GetTypeHash(InHandle.TypeKey);
 		}
diff --git a/Engine/Source/Runtime/MassEntity/Public/TypeBitSetBuilder.h b/Engine/Source/Runtime/MassEntity/Public/TypeBitSetBuilder.h
--- a/Engine/Source/Runtime/MassEntity/Public/TypeBitSetBuilder.h
+++ b/Engine/Source/Runtime/MassEntity/Public/TypeBitSetBuilder.h
@@ -52,7 +52,7 @@ private:
 		 * Equality operator to compare bitsets.
 		 * Uses CompareSetBits to compare the set bits, ignoring missing bits.
 		 */
-		FORCEINLINE bool operator==(const FBitSet& Other) const
+		inline bool operator==(const FBitSet& Other) const
 		{
 			return TContainer::CompareSetBits(Other, /*bMissingBitValue=*/false);
 		}
@@ -219,7 +219,7 @@ public:
 	 * @return The index of the added struct type.
 	 */
 	template<typename T>
-	FORCEINLINE int32 Add()
+	inline int32 Add()
 	{
 		static_assert(FTraits::template IsValidType<T>, "Given struct is not a valid type for this TypeBitSetBuilder.");
 		const int32 StructTypeIndex = GetTypeIndex<T>();
@@ -233,7 +233,7 @@ public:
 	 * @return The index of the removed struct type.
 	 */
 	template<typename T>
-	FORCEINLINE int32 Remove()
+	inline int32 Remove()
 	{
 		static_assert(FTraits::template IsValidType<T>, "Given struct is not a valid type for this TypeBitSetBuilder.");
 		const int32 StructTypeIndex = GetTypeIndex<T>();
@@ -245,7 +245,7 @@ public:
 	 * Removes all bits set in another builder's bitset from this builder's bitset.
 	 * @param Other the other builder whose bits to remove.
 	 */
-	FORCEINLINE void Remove(const TTypeBitSetBuilder& Other)
+	inline void Remove(const TTypeBitSetBuilder& Other)
 	{
 		StructTypesBitArray -= Other.StructTypesBitArray;
 	}
@@ -256,7 +256,7 @@ public:
 	 * @return True if the type is contained; false otherwise.
 	 */
 	template<typename T>
-	FORCEINLINE bool Contains() const
+	inline bool Contains() const
 	{
 		static_assert(FTraits::template IsValidType<T>, "Given struct is not a valid type for this TypeBitSetBuilder.");
 		const int32 StructTypeIndex = GetTypeIndex<T>();
@@ -268,7 +268,7 @@ public:
 	 * @param Other the other builder.
 	 * @return A new builder with the result.
 	 */
-	FORCEINLINE TTypeBitSetBuilder operator-(const TTypeBitSetBuilder& Other) const
+	inline TTypeBitSetBuilder operator-(const TTypeBitSetBuilder& Other) const
 	{
 		TTypeBitSetBuilder Result = *this;
 		Result -= Other;
@@ -280,7 +280,7 @@ public:
 	 * @param Other the other builder.
 	 * @return A new builder with the result.
 	 */
-	FORCEINLINE TTypeBitSetBuilder operator+(const TTypeBitSetBuilder& Other) const
+	inline TTypeBitSetBuilder operator+(const TTypeBitSetBuilder& Other) const
 	{
 		TTypeBitSetBuilder Result(StructTracker);
 		Result.StructTypesBitArray = TBitArray<>::BitwiseOR(StructTypesBitArray, Other.StructTypesBitArray, EBitwiseOperatorFlags::MaxSize);
@@ -292,7 +292,7 @@ public:
 	 * @param Other the other builder.
 	 * @return A new builder with the result.
 	 */
-	FORCEINLINE TTypeBitSetBuilder operator&(const TTypeBitSetBuilder& Other) const
+	inline TTypeBitSetBuilder operator&(const TTypeBitSetBuilder& Other) const
 	{
 		return TTypeBitSetBuilder(StructTracker, TBitArray<>::BitwiseAND(StructTypesBitArray, Other.StructTypesBitArray, EBitwiseOperatorFlags::MinSize));
 	}
@@ -302,7 +302,7 @@ public:
 	 * @param Other the other builder.
 	 * @return A new builder with the result.
 	 */
-	FORCEINLINE TTypeBitSetBuilder operator|(const TTypeBitSetBuilder& Other) const
+	inline TTypeBitSetBuilder operator|(const TTypeBitSetBuilder& Other) const
 	{
 		return TTypeBitSetBuilder(StructTracker, TBitArray<>::BitwiseOR(StructTypesBitArray, Other.StructTypesBitArray, EBitwiseOperatorFlags::MaxSize));
 	}
@@ -312,7 +312,7 @@ public:
 	 * @param Other the other builder.
 	 * @return A new builder representing the overlap.
 	 */
-	FORCEINLINE TTypeBitSetBuilder GetOverlap(const TTypeBitSetBuilder& Other) const
+	inline TTypeBitSetBuilder GetOverlap(const TTypeBitSetBuilder& Other) const
 	{
 		return *this & Other;
 	}
@@ -322,7 +322,7 @@ public:
 	 * @param Other the other builder.
 	 * @return True if equivalent; false otherwise.
 	 */
-	FORCEINLINE bool IsEquivalent(const TTypeBitSetBuilder& Other) const
+	inline bool IsEquivalent(const TTypeBitSetBuilder& Other) const
 	{
 		return StructTypesBitArray.CompareSetBits(Other.StructTypesBitArray, /*bMissingBitValue=*/false);
 	}
@@ -332,7 +332,7 @@ public:
 	 * @param Other the other builder.
 	 * @return True if all bits are set; false otherwise.
 	 */
-	FORCEINLINE bool HasAll(const TTypeBitSetBuilder& Other) const
+	inline bool HasAll(const TTypeBitSetBuilder& Other) const
 	{
 		return StructTypesBitArray.HasAll(Other.StructTypesBitArray);
 	}
@@ -342,7 +342,7 @@ public:
 	 * @param Other the other builder.
 	 * @return True if any bits are set; false otherwise.
 	 */
-	FORCEINLINE bool HasAny(const TTypeBitSetBuilder& Other) const
+	inline bool HasAny(const TTypeBitSetBuilder& Other) const
 	{
 		return StructTypesBitArray.HasAny(Other.StructTypesBitArray);
 	}
@@ -352,7 +352,7 @@ public:
 	 * @param Other the other builder.
 	 * @return True if no bits are set; false otherwise.
 	 */
-	FORCEINLINE bool HasNone(const TTypeBitSetBuilder& Other) const
+	inline bool HasNone(const TTypeBitSetBuilder& Other) const
 	{
 		return !StructTypesBitArray.HasAny(Other.StructTypesBitArray);
 	}
@@ -371,7 +371,7 @@ public:
 	 * @param BitIndex the index of the bit to check.
 	 * @return True if the bit is set; false otherwise.
 	 */
-	FORCEINLINE bool IsBitSet(const int32 BitIndex) const
+	inline bool IsBitSet(const int32 BitIndex) const
 	{
 		return StructTypesBitArray.Contains(BitIndex);
 	}
@@ -409,7 +409,7 @@ public:
 	 * @param Other the other builder to compare.
 	 * @return True if equal; false otherwise.
 	 */
-	FORCEINLINE bool operator==(const TTypeBitSetBuilder& Other) const
+	inline bool operator==(const TTypeBitSetBuilder& Other) const
 	{
 		return StructTypesBitArray.CompareSetBits(Other.StructTypesBitArray, /*bMissingBitValue=*/false);
 	}
@@ -419,7 +419,7 @@ public:
 	 * @param Other the other builder to compare.
 	 * @return True if not equal; false otherwise.
 	 */
-	FORCEINLINE bool operator!=(const TTypeBitSetBuilder& Other) const
+	inline bool operator!=(const TTypeBitSetBuilder& Other) const
 	{
 		return !(*this == Other);
 	}
