[Mass] Entity Builder API extension, supporting use of arbitrary "element" type and have the builder figure out what kind of an element that is and store it appropriately.
Additional changes:
* fixed FEntityBuilder::AddInternal implementation to not add duplicated instances of a given element type if one is already present.
* added creation context to the deferred path, to ensure the fragment values are already set when the observers kick in. This was already done in the synchronous path, so this change addresses the inconsistency

diff --git a/Engine/Source/Developer/MassEntityTestSuite/Private/MassEntityBuilderTest.cpp b/Engine/Source/Developer/MassEntityTestSuite/Private/MassEntityBuilderTest.cpp
--- a/Engine/Source/Developer/MassEntityTestSuite/Private/MassEntityBuilderTest.cpp
+++ b/Engine/Source/Developer/MassEntityTestSuite/Private/MassEntityBuilderTest.cpp
@@ -611,15 +611,24 @@ namespace UE::Mass::Test::EntityBuilder
 
 				FInstancedStruct ElementInstance;
 
+				// adding "wrong value" first to verify that the last "add" will hold
+				ElementInstance.InitializeAs<FTestFragment_Int>();
+				Builder.Add(ElementInstance);
 				ElementInstance.InitializeAs<FTestFragment_Int>(TestIntValue);
 				Builder.Add(ElementInstance);
 				
+				ElementInstance.InitializeAs<FTestFragment_Float>();
+				Builder.Add(MoveTemp(ElementInstance));
 				ElementInstance.InitializeAs<FTestFragment_Float>(TestFloatValue);
 				Builder.Add(MoveTemp(ElementInstance));
 
+				ElementInstance.InitializeAs<FTestSharedFragment_Float>();
+				Builder.Add(MoveTemp(ElementInstance));
 				ElementInstance.InitializeAs<FTestSharedFragment_Float>(TestSharedFloatValue);
 				Builder.Add(MoveTemp(ElementInstance));
 
+				ElementInstance.InitializeAs<FTestConstSharedFragment_Int>();
+				Builder.Add(ElementInstance);
 				ElementInstance.InitializeAs<FTestConstSharedFragment_Int>(TestSharedIntValue);
 				Builder.Add(ElementInstance);
 
diff --git a/Engine/Source/Runtime/CoreUObject/Public/StructUtils/StructTypeBitSet.h b/Engine/Source/Runtime/CoreUObject/Public/StructUtils/StructTypeBitSet.h
--- a/Engine/Source/Runtime/CoreUObject/Public/StructUtils/StructTypeBitSet.h
+++ b/Engine/Source/Runtime/CoreUObject/Public/StructUtils/StructTypeBitSet.h
@@ -523,6 +523,12 @@ struct TTypeBitSetBase
 		StructTypesBitArray.AddAtIndex(StructTypeIndex);
 	}
 
+	/** Sets the bit at index StructTypeIndex */
+	void AddAtIndex(const int32 StructTypeIndex)
+	{
+		StructTypesBitArray.AddAtIndex(StructTypeIndex);
+	}
+
 	/**
 	 * Removes a struct type from the bitset.
 	 * @param InStructType - The struct type to remove.
@@ -543,6 +549,12 @@ struct TTypeBitSetBase
 		StructTypesBitArray.RemoveAtIndex(StructTypeIndex);
 	}
 
+	/** Clears the bit at index StructTypeIndex */
+	void RemoveAtIndex(const int32 StructTypeIndex)
+	{
+		StructTypesBitArray.RemoveAtIndex(StructTypeIndex);
+	}
+
 	/**
 	 * Resets all bits in the bitset
 	 */
diff --git a/Engine/Source/Runtime/MassEntity/Private/MassEntityBuilder.cpp b/Engine/Source/Runtime/MassEntity/Private/MassEntityBuilder.cpp
--- a/Engine/Source/Runtime/MassEntity/Private/MassEntityBuilder.cpp
+++ b/Engine/Source/Runtime/MassEntity/Private/MassEntityBuilder.cpp
@@ -282,6 +282,10 @@ FMassEntityHandle FEntityBuilder::Commit()
 		EntityManager->Defer().PushCommand<FMassDeferredCreateCommand>(
 					[ReservedEntityHandle = EntityHandle, SharedFragmentValues = CachedSharedFragmentValues, ArchetypeHandle = CachedArchetypeHandle, FragmentsCopy = Fragments](FMassEntityManager& Manager)
 					{
+						// creating creation context to block observers from triggering until the
+						// values are set with SetEntityFragmentValues. The lock gets auto-released at the end of the scope
+						// or persists, if anyone else has a shared ptr to it
+						TSharedRef<FMassEntityManager::FEntityCreationContext> CreationContext = Manager.GetOrMakeCreationContext();
 						Manager.BuildEntity(ReservedEntityHandle, ArchetypeHandle, SharedFragmentValues);
 						Manager.SetEntityFragmentValues(ReservedEntityHandle, FragmentsCopy);
 					});
@@ -473,24 +477,52 @@ void FEntityBuilder::CacheSharedFragmentValue()
 }
 
 template<typename T>
-FORCEINLINE FEntityBuilder& FEntityBuilder::AddInternal(T&& ElementInstance)
+bool FEntityBuilder::HandleTypeInstance(const UScriptStruct* Type, TArray<FInstancedStruct>*& OutTargetContainer, bool& bOutAlreadyInComposition)
 {
-	if (const UScriptStruct* Type = ElementInstance.GetScriptStruct())
+	if (UE::Mass::IsA<T>(Type))
 	{
-		if (Type->IsChildOf(FMassFragment::StaticStruct()))
+		auto& BitSet = Composition.GetContainer<T>();
+		const int32 TypeIndex = BitSet.GetTypeIndex(*Type);
+		OutTargetContainer = &GetInstancedStructContainerInternal<T>();
+		if (BitSet.IsBitSet(TypeIndex) == false)
 		{
-			Composition.Fragments.Add(*Type);
-			Fragments.Add(Forward<T>(ElementInstance));
+			BitSet.AddAtIndex(TypeIndex);
 		}
-		else if (Type->IsChildOf(FMassSharedFragment::StaticStruct()))
+		else
 		{
-			Composition.SharedFragments.Add(*Type);
-			SharedFragments.Add(Forward<T>(ElementInstance));
+			bOutAlreadyInComposition = true;
 		}
-		else if (Type->IsChildOf(FMassConstSharedFragment::StaticStruct()))
+		return true;
+	}
+	return false;
+}
+
+template<typename TInstancedStruct>
+FEntityBuilder& FEntityBuilder::AddInternal(TInstancedStruct&& ElementInstance)
+{
+	if (const UScriptStruct* Type = ElementInstance.GetScriptStruct())
+	{
+		TArray<FInstancedStruct>* TargetContainer = nullptr;
+
+		bool bAlreadyInComposition = false;
+		HandleTypeInstance<FMassFragment>(Type, TargetContainer, bAlreadyInComposition)
+		|| HandleTypeInstance<FMassSharedFragment>(Type, TargetContainer, bAlreadyInComposition)
+		|| HandleTypeInstance<FMassConstSharedFragment>(Type, TargetContainer, bAlreadyInComposition);
+
+		if (TargetContainer)
 		{
-			Composition.ConstSharedFragments.Add(*Type);
-			ConstSharedFragments.Add(Forward<T>(ElementInstance));
+			auto* FoundElement = bAlreadyInComposition
+				? TargetContainer->FindByPredicate(FStructInstanceFindingPredicate{Type})
+				: nullptr;
+
+			if (FoundElement)
+			{
+				*FoundElement = Forward<TInstancedStruct>(ElementInstance);
+			}
+			else
+			{
+				TargetContainer->Add(Forward<TInstancedStruct>(ElementInstance));
+			}
 		}
 	}
 
@@ -507,4 +539,34 @@ FEntityBuilder& FEntityBuilder::Add(FInstancedStruct&& ElementInstance)
 	return AddInternal(ElementInstance);
 }
 
+FEntityBuilder& FEntityBuilder::Add(TNotNull<const UScriptStruct*> ElementType)
+{
+	if (UE::Mass::IsA<FMassTag>(ElementType))
+	{
+		Composition.Tags.Add(*ElementType);
+	}
+	else if (UE::Mass::IsA<FMassFragment>(ElementType))
+	{
+		Composition.Fragments.Add(*ElementType);
+	}
+	else if (UE::Mass::IsA<FMassSharedFragment>(ElementType))
+	{
+		Composition.SharedFragments.Add(*ElementType);
+	}
+	else if (UE::Mass::IsA<FMassConstSharedFragment>(ElementType))
+	{
+		Composition.ConstSharedFragments.Add(*ElementType);
+	}
+	else if (ensureMsgf(UE::Mass::IsA<FMassChunkFragment>(ElementType)
+		, TEXT("Unhandled element type %s"), *ElementType->GetName()))
+	{
+		Composition.ChunkFragments.Add(*ElementType);
+	}
+
+	State = EState::ReadyToCommit;
+	CachedArchetypeHandle = FMassArchetypeHandle();
+	
+	return *this;
+}
+
 } // namespace UE::Mass
\ No newline at end of file
diff --git a/Engine/Source/Runtime/MassEntity/Public/MassEntityBuilder.h b/Engine/Source/Runtime/MassEntity/Public/MassEntityBuilder.h
--- a/Engine/Source/Runtime/MassEntity/Public/MassEntityBuilder.h
+++ b/Engine/Source/Runtime/MassEntity/Public/MassEntityBuilder.h
@@ -177,6 +177,9 @@ struct FEntityBuilder
 	UE_API FEntityBuilder& Add(const FInstancedStruct& ElementInstance);
 	UE_API FEntityBuilder& Add(FInstancedStruct&& ElementInstance);
 
+	/** Adds the ElementType to the target archetype's composition */
+	UE_API FEntityBuilder& Add(TNotNull<const UScriptStruct*> ElementType);
+
 	/**
 	 * Finds and retrieves a pointer to a fragment of type T if it exists.
 	 * @param T - The type of fragment to find.
@@ -216,9 +219,22 @@ protected:
 	UE_API void CacheArchetypeHandle();
 	UE_API void InvalidateCachedData();
 
+	struct FStructInstanceFindingPredicate
+	{
+		template<typename TInstancedStruct>
+		bool operator()(const TInstancedStruct& Instance) const
+		{
+			return Instance.GetScriptStruct() == SearchedType;
+		}
+		const UScriptStruct* SearchedType = nullptr;
+	};
+
 private:
+	template<typename TInstancedStruct>
+	FEntityBuilder& AddInternal(TInstancedStruct&& ElementInstance);
+
 	template<typename T>
-	FEntityBuilder& AddInternal(T&& ElementInstance);
+	bool HandleTypeInstance(const UScriptStruct* Type, TArray<FInstancedStruct>*& OutTargetContainer, bool& bOutAlreadyInComposition);
 
 	TSharedRef<FMassEntityManager> EntityManager;
 	mutable FMassEntityHandle EntityHandle;
@@ -338,10 +354,8 @@ T& FEntityBuilder::Add_GetRef(TArgs&&... InArgs) requires (CElement<T> && !(CTag
 		return GetInstancedStructContainerInternal<T>().Add_GetRef(FInstancedStruct::Make<T>(InArgs...)).template GetMutable<T>();
 	}
 
-	return GetInstancedStructContainerInternal<T>().FindByPredicate([](const FInstancedStruct& ExistingElement)
-		{
-			return ExistingElement.GetScriptStruct() == T::StaticStruct();
-		})->template GetMutable<T>();
+	return GetInstancedStructContainerInternal<T>().FindByPredicate(FStructInstanceFindingPredicate{T::StaticStruct()})
+		->template GetMutable<T>();
 }
 
 template<typename T, typename... TArgs>
@@ -355,10 +369,7 @@ T& FEntityBuilder::GetOrCreate(TArgs&&... InArgs) requires (CElement<T> && !(CTa
 	if (GetBitSetBuilder<FElementType>().template Contains<T>())
 	{
 		// replace
-		ElementFound = GetInstancedStructContainerInternal<T>().FindByPredicate([](const FInstancedStruct& ExistingElement)
-		{
-			return ExistingElement.GetScriptStruct() == T::StaticStruct();
-		});
+		ElementFound = GetInstancedStructContainerInternal<T>().FindByPredicate(FStructInstanceFindingPredicate{T::StaticStruct()});
 
 		checkf(ElementFound, TEXT("We expect the element to be found since we already tested the Composition"));
 		ElementFound->GetMutable<T>() = T(InArgs...);
@@ -381,10 +392,7 @@ T* FEntityBuilder::Find() requires (CElement<T> && !(CTag<T> || CChunkFragment<T
 
 	if (GetBitSetBuilder<FElementType>().template Contains<T>())
 	{
-		FInstancedStruct* ElementFound = GetInstancedStructContainerInternal<T>().FindByPredicate([](const FInstancedStruct& Element)
-		{
-			return Element.GetScriptStruct() == T::StaticStruct();
-		});
+		FInstancedStruct* ElementFound = GetInstancedStructContainerInternal<T>().FindByPredicate(FStructInstanceFindingPredicate{T::StaticStruct()});
 
 		checkf(ElementFound, TEXT("We expect the element to be found since we already tested the Composition"));
 		return ElementFound->GetMutablePtr<T>();
