[Mass] Added more comments to FEntityBuilder::AddInternal to explain some non-obvious bits

#rnx

diff --git a/Engine/Source/Runtime/MassEntity/Private/MassEntityBuilder.cpp b/Engine/Source/Runtime/MassEntity/Private/MassEntityBuilder.cpp
--- a/Engine/Source/Runtime/MassEntity/Private/MassEntityBuilder.cpp
+++ b/Engine/Source/Runtime/MassEntity/Private/MassEntityBuilder.cpp
@@ -505,6 +505,8 @@ FEntityBuilder& FEntityBuilder::AddInternal(TInstancedStruct&& ElementInstance)
 		TArray<FInstancedStruct>* TargetContainer = nullptr;
 
 		bool bAlreadyInComposition = false;
+		// the following chain will stop at the first successful HandleTypeInstance call.
+		// This is by design, since we don't support there being any overlap between the types.
 		HandleTypeInstance<FMassFragment>(Type, TargetContainer, bAlreadyInComposition)
 		|| HandleTypeInstance<FMassSharedFragment>(Type, TargetContainer, bAlreadyInComposition)
 		|| HandleTypeInstance<FMassConstSharedFragment>(Type, TargetContainer, bAlreadyInComposition);
@@ -515,6 +517,10 @@ FEntityBuilder& FEntityBuilder::AddInternal(TInstancedStruct&& ElementInstance)
 				? TargetContainer->FindByPredicate(FStructInstanceFindingPredicate{Type})
 				: nullptr;
 
+			// note that it's perfectly fine for FoundElement being null while bAlreadyInComposition == true.
+			// This will happen if the element type has been added just as a type first (which only annotates the bitset)
+			// and then an instance of the same type gets added - the bitset already contains the bit, but
+			// the container doesn't hold an instance yet.
 			if (FoundElement)
 			{
 				*FoundElement = Forward<TInstancedStruct>(ElementInstance);
