[Mass] fixed misnamed TRACE_CPUPROFILER_EVENT_SCOPE in MassArchetypeGroupCommands.h

#rnx

diff --git a/Engine/Source/Runtime/MassEntity/Public/MassArchetypeGroupCommands.h b/Engine/Source/Runtime/MassEntity/Public/MassArchetypeGroupCommands.h
--- a/Engine/Source/Runtime/MassEntity/Public/MassArchetypeGroupCommands.h
+++ b/Engine/Source/Runtime/MassEntity/Public/MassArchetypeGroupCommands.h
@@ -35,7 +35,7 @@ namespace UE::Mass::Command
 	protected:
 		virtual void Execute(FMassEntityManager& EntityManager) const override
 		{
-			TRACE_CPUPROFILER_EVENT_SCOPE(MassCommandAddElement_Execute);
+			TRACE_CPUPROFILER_EVENT_SCOPE(MassCommandGroupEntities_Execute);
 			for (auto Group : Groups)
 			{
 				TArray<FMassArchetypeEntityCollection> CollectionsAtLevel;
