[Mass] Deprecated FMassBatchedCommand::Execute as a const function and introduced non-const replacement named `Run`.

Unfortunately a new name had to be introduced since the constness of the function is the only thing that's changing.
Added a temporary construct to MassCommandBuffer to statically detect commands that still implement the deprecated virtual Execute in an attempt to make it easier for people updating their code to notice the change - the construct will check at compile time all the pushed commands of known types and statically assert if the command type still overrides the deprecated `Execute` function.

diff --git a/Engine/Plugins/Runtime/InstancedActors/Source/InstancedActors/Public/InstancedActorsCommands.h b/Engine/Plugins/Runtime/InstancedActors/Source/InstancedActors/Public/InstancedActorsCommands.h
--- a/Engine/Plugins/Runtime/InstancedActors/Source/InstancedActors/Public/InstancedActorsCommands.h
+++ b/Engine/Plugins/Runtime/InstancedActors/Source/InstancedActors/Public/InstancedActorsCommands.h
@@ -59,12 +59,12 @@ protected:
 		return Super::GetAllocatedSize() + ManagersToResave.GetAllocatedSize();
 	}
 
-	virtual void Execute(FMassEntityManager& EntityManager) const override
+	virtual void Run(FMassEntityManager& EntityManager) override
 	{
 		TRACE_CPUPROFILER_EVENT_SCOPE(MassCommandAddFragmentInstancesAndResaveIAPersistence_Execute);
 
 		// Add / update fragments
-		Super::Execute(EntityManager);
+		Super::Run(EntityManager);
 
 		// Resave Instanced Actor persistence for now-updated fragments
 		for (const TObjectPtr<AInstancedActorsManager>& ManagerToResave : ManagersToResave)
diff --git a/Engine/Source/Runtime/MassEntity/Private/MassCommandBuffer.cpp b/Engine/Source/Runtime/MassEntity/Private/MassCommandBuffer.cpp
--- a/Engine/Source/Runtime/MassEntity/Private/MassCommandBuffer.cpp
+++ b/Engine/Source/Runtime/MassEntity/Private/MassCommandBuffer.cpp
@@ -207,7 +207,7 @@ bool FMassCommandBuffer::Flush(FMassEntityManager& EntityManager)
 		FCsvProfiler::RecordCustomStat(**Name, CSV_CATEGORY_INDEX(MassEntitiesCounters), Command->GetNumOperationsStat(), ECsvCustomStatOp::Accumulate);
 #endif // CSV_PROFILER_STATS
 
-		Command->Execute(EntityManager);
+		Command->Run(EntityManager);
 		Command->Reset();
 	}
 
diff --git a/Engine/Source/Runtime/MassEntity/Public/MassArchetypeGroupCommands.h b/Engine/Source/Runtime/MassEntity/Public/MassArchetypeGroupCommands.h
--- a/Engine/Source/Runtime/MassEntity/Public/MassArchetypeGroupCommands.h
+++ b/Engine/Source/Runtime/MassEntity/Public/MassArchetypeGroupCommands.h
@@ -33,7 +33,7 @@ namespace UE::Mass::Command
 		}
 
 	protected:
-		virtual void Execute(FMassEntityManager& EntityManager) const override
+		virtual void Run(FMassEntityManager& EntityManager) override
 		{
 			TRACE_CPUPROFILER_EVENT_SCOPE(MassCommandGroupEntities_Execute);
 			for (auto Group : Groups)
diff --git a/Engine/Source/Runtime/MassEntity/Public/MassCommandBuffer.h b/Engine/Source/Runtime/MassEntity/Public/MassCommandBuffer.h
--- a/Engine/Source/Runtime/MassEntity/Public/MassCommandBuffer.h
+++ b/Engine/Source/Runtime/MassEntity/Public/MassCommandBuffer.h
@@ -10,7 +10,8 @@
 #include "Misc/TransactionallySafeCriticalSection.h"
 #include "MassEntityUtils.h"
 #include "MassCommands.h"
-
+// required for UE::Mass::Deprecation
+#include <type_traits>
 
 struct FMassEntityManager;
 
@@ -21,9 +22,33 @@ struct FMassEntityManager;
 checkf(IsFlushing() == false, TEXT("Trying to push commands is not supported while the given buffer is being flushed")); \
 checkf(OwnerThreadId == FPlatformTLS::GetCurrentThreadId(), TEXT("Commands can be pushed only in the same thread where the command buffer was created."))
 
+// The following is a temporary construct to help users update their code after FMassBatchedCommand::Execute deprecation
+// @todo remove by 5.9
+namespace UE::Mass::Deprecation
+{
+PRAGMA_DISABLE_DEPRECATION_WARNINGS
+	template <typename TCommand>
+	constexpr bool IsExecuteOverridden()
+	{
+	    using TResolved = decltype(&TCommand::Execute);
+	    using TBase = decltype(&FMassBatchedCommand::Execute);
+	    return !std::is_same_v<TResolved, TBase>;
+	}
+
+	template <typename TCommand>
+	requires std::is_base_of_v<FMassBatchedCommand, TCommand>
+	struct TOverridesExecute : std::bool_constant<IsExecuteOverridden<TCommand>()>
+	{
+	};
+PRAGMA_ENABLE_DEPRECATION_WARNINGS
+}
+
+#define ASSERT_EXECUTE_DEPRECATION(CommandType) \
+		static_assert(!UE::Mass::Deprecation::TOverridesExecute<CommandType>::value, "Mass Commands: CONST Execute function is deprecated in 5.7 and will be removed by 5.9. Use Run instead.");
+
+
 struct FMassCommandBuffer
 {
-public:
 	MASSENTITY_API FMassCommandBuffer();
 	MASSENTITY_API ~FMassCommandBuffer();
 
@@ -235,6 +260,8 @@ private:
 	template<typename T>
 	T& CreateOrAddCommand()
 	{
+		ASSERT_EXECUTE_DEPRECATION(T);
+
 		static_assert(!UE::Mass::Command::TCommandTraits<T>::RequiresUniqueHandling, "This command type needs to be added via PushUniqueCommand");
 		const int32 Index = FMassBatchedCommand::GetCommandIndex<T>();
 
diff --git a/Engine/Source/Runtime/MassEntity/Public/MassCommands.h b/Engine/Source/Runtime/MassEntity/Public/MassCommands.h
--- a/Engine/Source/Runtime/MassEntity/Public/MassCommands.h
+++ b/Engine/Source/Runtime/MassEntity/Public/MassCommands.h
@@ -99,9 +99,24 @@ struct FMassBatchedCommand
 		, DebugName(DebugName)
 	{}
 #endif // CSV_PROFILER_STATS || WITH_MASSENTITY_DEBUG
-	virtual ~FMassBatchedCommand() { Reset(); }
+	virtual ~FMassBatchedCommand()
+	{
+		Reset();
+	}
+
+	UE_DEPRECATED(5.7, "Mass Commands: CONST Execute function is deprecated in 5.7 and will be removed by 5.9. Use Run instead.")
+	virtual void Execute(FMassEntityManager& EntityManager) const
+	{
+		ensureMsgf(false, TEXT("FMassBatchedCommand::Execute is DEPRECATED, override Run function instead."));
+	}
+
+	virtual void Run(FMassEntityManager& EntityManager)
+	{
+PRAGMA_DISABLE_DEPRECATION_WARNINGS
+		Execute(EntityManager);
+PRAGMA_ENABLE_DEPRECATION_WARNINGS
+	}
 
-	virtual void Execute(FMassEntityManager& EntityManager) const = 0;
 	virtual void Reset()
 	{
 		bHasWork = false;
@@ -202,7 +217,7 @@ struct FMassCommandDestroyEntities : public FMassBatchedEntityCommand
 	}
 
 protected:
-	virtual void Execute(FMassEntityManager& EntityManager) const override
+	virtual void Run(FMassEntityManager& EntityManager) override
 	{
 		TRACE_CPUPROFILER_EVENT_SCOPE(MassCommandDestroyEntities_Execute);
 
@@ -225,7 +240,7 @@ struct FMassCommandAddFragmentsInternal : public FMassBatchedEntityCommand
 	{}
 
 protected:
-	virtual void Execute(FMassEntityManager& EntityManager) const override
+	virtual void Run(FMassEntityManager& EntityManager) override
 	{
 		TRACE_CPUPROFILER_EVENT_SCOPE(MassCommandAddFragments_Execute);
 		TArray<FMassArchetypeEntityCollection> EntityCollections;
@@ -248,7 +263,7 @@ struct FMassCommandRemoveFragmentsInternal : public FMassBatchedEntityCommand
 	{}
 
 protected:
-	virtual void Execute(FMassEntityManager& EntityManager) const override
+	virtual void Run(FMassEntityManager& EntityManager) override
 	{
 		TRACE_CPUPROFILER_EVENT_SCOPE(MassCommandRemoveFragments_Execute);
 		TArray<FMassArchetypeEntityCollection> EntityCollections;
@@ -278,7 +293,7 @@ struct FMassCommandChangeTags : public FMassBatchedEntityCommand
 	{}
 
 protected:
-	virtual void Execute(FMassEntityManager& EntityManager) const override
+	virtual void Run(FMassEntityManager& EntityManager) override
 	{
 		TRACE_CPUPROFILER_EVENT_SCOPE(MassCommandChangeTags_Execute);
 		TArray<FMassArchetypeEntityCollection> EntityCollections;
@@ -381,7 +396,7 @@ protected:
 		return Super::GetAllocatedSize() + Fragments.GetAllocatedSize() + FragmentsAffected.GetAllocatedSize();
 	}
 
-	virtual void Execute(FMassEntityManager& EntityManager) const override
+	virtual void Run(FMassEntityManager& EntityManager) override
 	{
 		TRACE_CPUPROFILER_EVENT_SCOPE(MassCommandAddFragmentInstances_Execute);
 
@@ -414,7 +429,7 @@ struct FMassCommandAddElement : public FMassBatchedEntityCommand
 	{}
 
 protected:
-	virtual void Execute(FMassEntityManager& EntityManager) const override
+	virtual void Run(FMassEntityManager& EntityManager) override
 	{
 		TRACE_CPUPROFILER_EVENT_SCOPE(MassCommandAddElement_Execute);
 		EntityManager.AddElementToEntities(TargetEntities, ElementType);
@@ -444,7 +459,7 @@ struct FMassCommandBuildEntity : public FMassCommandAddFragmentInstances<TOthers
 
 protected:
 
-	virtual void Execute(FMassEntityManager& EntityManager) const override
+	virtual void Run(FMassEntityManager& EntityManager) override
 	{
 		TRACE_CPUPROFILER_EVENT_SCOPE(MassCommandBuildEntity_Execute);
 
@@ -507,7 +522,7 @@ protected:
 		return TotalSize;
 	}
 
-	virtual void Execute(FMassEntityManager& EntityManager) const 
+	virtual void Run(FMassEntityManager& EntityManager) override
 	{
 		TRACE_CPUPROFILER_EVENT_SCOPE(MassCommandBuildEntityWithSharedFragments_Execute);
 
@@ -604,7 +619,7 @@ protected:
 		return DeferredFunctions.GetAllocatedSize();
 	}
 
-	virtual void Execute(FMassEntityManager& EntityManager) const override
+	virtual void Run(FMassEntityManager& EntityManager) override
 	{
 		TRACE_CPUPROFILER_EVENT_SCOPE(MassDeferredCommand_Execute);
 
