[Mass] Extracted WITH_MASSENTITY_DEBUG and WITH_STRUCTUTILS_DEBUG into separate files to limit the amount of stuff included when just wanting to have one of the two available in one's headers.

#rnx

diff --git a/Engine/Source/Runtime/CoreUObject/Public/StructUtils/StructUtilsMacros.h b/Engine/Source/Runtime/CoreUObject/Public/StructUtils/StructUtilsMacros.h
new file mode 100644
--- /dev/null
+++ b/Engine/Source/Runtime/CoreUObject/Public/StructUtils/StructUtilsMacros.h
@@ -0,0 +1,9 @@
+// Copyright Epic Games, Inc. All Rights Reserved.
+
+#pragma once
+
+#include "Misc/Build.h"
+
+#ifndef WITH_STRUCTUTILS_DEBUG
+#define WITH_STRUCTUTILS_DEBUG (!(UE_BUILD_SHIPPING || UE_BUILD_SHIPPING_WITH_EDITOR || UE_BUILD_TEST) && 1)
+#endif // WITH_STRUCTUTILS_DEBUG
diff --git a/Engine/Source/Runtime/CoreUObject/Public/StructUtils/StructUtilsTypes.h b/Engine/Source/Runtime/CoreUObject/Public/StructUtils/StructUtilsTypes.h
--- a/Engine/Source/Runtime/CoreUObject/Public/StructUtils/StructUtilsTypes.h
+++ b/Engine/Source/Runtime/CoreUObject/Public/StructUtils/StructUtilsTypes.h
@@ -4,13 +4,10 @@
 
 #include <type_traits>
 #include "UObject/Class.h"
+#include "StructUtils/StructUtilsMacros.h"
 
 #define UE_API COREUOBJECT_API
 
-#ifndef WITH_STRUCTUTILS_DEBUG
-#define WITH_STRUCTUTILS_DEBUG (!(UE_BUILD_SHIPPING || UE_BUILD_SHIPPING_WITH_EDITOR || UE_BUILD_TEST) && 1)
-#endif // WITH_STRUCTUTILS_DEBUG
-
 struct FSharedStruct;
 struct FConstSharedStruct;
 struct FStructView;
diff --git a/Engine/Source/Runtime/MassEntity/Public/MassEntityMacros.h b/Engine/Source/Runtime/MassEntity/Public/MassEntityMacros.h
new file mode 100644
--- /dev/null
+++ b/Engine/Source/Runtime/MassEntity/Public/MassEntityMacros.h
@@ -0,0 +1,9 @@
+// Copyright Epic Games, Inc. All Rights Reserved.
+
+#pragma once
+
+#include "StructUtils/StructUtilsMacros.h"
+
+#ifndef WITH_MASSENTITY_DEBUG
+#define WITH_MASSENTITY_DEBUG (!(UE_BUILD_SHIPPING || UE_BUILD_SHIPPING_WITH_EDITOR || UE_BUILD_TEST) && WITH_STRUCTUTILS_DEBUG && 1)
+#endif
diff --git a/Engine/Source/Runtime/MassEntity/Public/MassProcessingTypes.h b/Engine/Source/Runtime/MassEntity/Public/MassProcessingTypes.h
--- a/Engine/Source/Runtime/MassEntity/Public/MassProcessingTypes.h
+++ b/Engine/Source/Runtime/MassEntity/Public/MassProcessingTypes.h
@@ -5,6 +5,7 @@
 #include "StructUtils/StructUtilsTypes.h"
 #include "StructUtils/InstancedStruct.h"
 #include "Templates/SubclassOf.h"
+#include "MassEntityMacros.h"
 #include "MassProcessingTypes.generated.h"
 
 #define UE_API MASSENTITY_API
@@ -13,10 +14,6 @@
 #define MASS_DO_PARALLEL !UE_SERVER
 #endif // MASS_DO_PARALLEL
 
-#ifndef WITH_MASSENTITY_DEBUG
-#define WITH_MASSENTITY_DEBUG (!(UE_BUILD_SHIPPING || UE_BUILD_SHIPPING_WITH_EDITOR || UE_BUILD_TEST) && WITH_STRUCTUTILS_DEBUG && 1)
-#endif
-
 struct FMassEntityManager;
 class UMassProcessor;
 class UMassCompositeProcessor;
