[Mass] Fixed an issue in FMassEntityQuery::ParallelForEachEntityChunk steming from the code not taking into consideration that ParallelForEach's workers can be moved between threads. If that happened and the query was run with `bAllowParallelCommands == true`, then the command buffer would complain about being used outside of the thread of its creation.
The fix boils down to updating the thread ID of the command buffer whenever it's picked up.

diff --git a/Engine/Source/Runtime/MassEntity/Private/MassEntityQuery.cpp b/Engine/Source/Runtime/MassEntity/Private/MassEntityQuery.cpp
--- a/Engine/Source/Runtime/MassEntity/Private/MassEntityQuery.cpp
+++ b/Engine/Source/Runtime/MassEntity/Private/MassEntityQuery.cpp
@@ -447,6 +447,11 @@ void FMassEntityQuery::ParallelForEachEntityChunk(FMassExecutionContext& Executi
 						// lazily creating the command buffer to ensure we create it in the same thread it's going to be used in
 						CommandBuffer = MakeShared<FMassCommandBuffer>();
 					}
+					else
+					{
+						// force-updating the thread ID because ParallelFor can move workers between threads.
+						CommandBuffer->ForceUpdateCurrentThreadID();
+					}
 					return CommandBuffer;
 				}
 			private:
diff --git a/Engine/Source/Runtime/MassEntity/Public/MassCommandBuffer.h b/Engine/Source/Runtime/MassEntity/Public/MassCommandBuffer.h
--- a/Engine/Source/Runtime/MassEntity/Public/MassCommandBuffer.h
+++ b/Engine/Source/Runtime/MassEntity/Public/MassCommandBuffer.h
@@ -252,11 +252,18 @@ struct FMassCommandBuffer
 		return OwnerThreadId == FPlatformTLS::GetCurrentThreadId();
 	}
 
+	/**
+	 * Updates the OwnerThreadId which indicates that the given command buffer instance is being
+	 * used in a different thread now. Use this with extreme caution, it's not a tool to be used
+	 * every time we get "Commands can be pushed only in the same thread where the command buffer was created."
+	 * error. It's meant to be used when there's a possibility the code owning the buffer has been
+	 * moved to another thread (like in ParallelFor).
+	 */
+	void ForceUpdateCurrentThreadID();
+
 private:
 	friend FMassEntityManager;
 
-	void ForceUpdateCurrentThreadID();
-
 	template<typename T>
 	T& CreateOrAddCommand()
 	{
