[Mass] Minor massaging of MassArchetypeData, achieved %2 mem footprint reduction

diff --git a/Engine/Source/Runtime/MassEntity/Internal/MassArchetypeData.h b/Engine/Source/Runtime/MassEntity/Internal/MassArchetypeData.h
--- a/Engine/Source/Runtime/MassEntity/Internal/MassArchetypeData.h
+++ b/Engine/Source/Runtime/MassEntity/Internal/MassArchetypeData.h
@@ -19,7 +19,7 @@ struct FMassFragmentRequirements;
 
 namespace UE::Mass
 {
-	int32 SanitizeChunkMemorySize(const int32 InChunkMemorySize, const bool bLogMismatch = true);
+	uint32 SanitizeChunkMemorySize(const uint32 InChunkMemorySize, const bool bLogMismatch = true);
 }
 
 // This is one chunk within an archetype
@@ -191,7 +191,7 @@ private:
 	UE::Mass::FArchetypeGroups Groups;
 
 	int32 NumEntitiesPerChunk;
-	SIZE_T TotalBytesPerEntity = 0;
+	uint32 TotalBytesPerEntity = 0;
 	int32 EntityListOffsetWithinChunk;
 
 	/**
@@ -206,6 +206,9 @@ private:
 	 */
 	uint32 EntityOrderVersion = 0;
 
+	/** Defaults to UMassEntitySettings.ChunkMemorySize. In near future will support being set via constructor. */
+	const uint32 ChunkMemorySize = 0;
+
 #if WITH_MASSENTITY_DEBUG
 	/** Arrays of names the archetype is referred as. */
 	TArray<FName> DebugNames;
@@ -217,9 +220,6 @@ private:
 	 */
 	FColor DebugColor;
 #endif // WITH_MASSENTITY_DEBUG
-
-	// Defaults to UMassEntitySettings.ChunkMemorySize. In near future will support being set via constructor.
-	const SIZE_T ChunkMemorySize = 0;
 	
 	friend FMassEntityQuery;
 	friend FMassArchetypeEntityCollection;
diff --git a/Engine/Source/Runtime/MassEntity/Private/MassArchetypeData.cpp b/Engine/Source/Runtime/MassEntity/Private/MassArchetypeData.cpp
--- a/Engine/Source/Runtime/MassEntity/Private/MassArchetypeData.cpp
+++ b/Engine/Source/Runtime/MassEntity/Private/MassArchetypeData.cpp
@@ -16,15 +16,15 @@ namespace UE::Mass
 	namespace Private
 	{
 		constexpr int32 UninitializedInt32 = -1;
-		constexpr int32 MinChunkMemorySize = 1024;
-		constexpr int32 MaxChunkMemorySize = 512 * 1024;
+		constexpr uint32 MinChunkMemorySize = 1024;
+		constexpr uint32 MaxChunkMemorySize = 512 * 1024;
 	}
 
-	int32 SanitizeChunkMemorySize(const int32 InChunkMemorySize, const bool bLogMismatch)
+	uint32 SanitizeChunkMemorySize(const uint32 InChunkMemorySize, const bool bLogMismatch)
 	{
-		const int32 SanitizedSize = FMath::Clamp(InChunkMemorySize, Private::MinChunkMemorySize, Private::MaxChunkMemorySize);
+		const uint32 SanitizedSize = FMath::Clamp(InChunkMemorySize, Private::MinChunkMemorySize, Private::MaxChunkMemorySize);
 		UE_CLOG(bLogMismatch && SanitizedSize != InChunkMemorySize, LogMass, Warning
-			, TEXT("ChunkMemorySize sanitization resulted in changing value. Old: %d, modified: %d")
+			, TEXT("ChunkMemorySize sanitization resulted in changing value. Old: %u, modified: %u")
 			, InChunkMemorySize, SanitizedSize);
 		return SanitizedSize;
 	}
@@ -160,9 +160,10 @@ void FMassArchetypeData::ConfigureFragments(const FMassEntityManager& EntityMana
 		FragmentIndexMap.Add(FragmentType, FragmentIndex);
 	}
 
-	TotalBytesPerEntity = FragmentSizeTallyBytes;
+	checkf(FragmentSizeTallyBytes < TNumericLimits<uint32>::Max(), TEXT("Single entity's size exceeds 2^32. This is not supported."));
+	TotalBytesPerEntity = static_cast<uint32>(FragmentSizeTallyBytes);
 	const SIZE_T ChunkAvailableSize = GetChunkAllocSize() - AlignmentPadding;
-	check(TotalBytesPerEntity <= ChunkAvailableSize);
+	checkf(TotalBytesPerEntity <= ChunkAvailableSize, TEXT("Single entity's size is larger than max chunk size - no entities will be created."));
 
 	NumEntitiesPerChunk = static_cast<int32>(ChunkAvailableSize / TotalBytesPerEntity);
 
@@ -1031,7 +1032,7 @@ void FMassArchetypeData::DebugPrintArchetype(FOutputDevice& Ar)
 		Ar.Logf(ELogVerbosity::Log, TEXT("\tEntity Occupancy: %.1f%%"),
 			CurrentEntityCapacity > 0 ? ((static_cast<float>(EntityMap.Num()) * 100.0f) / static_cast<float>(CurrentEntityCapacity)) : 0.f);
 	}
-	Ar.Logf(ELogVerbosity::Log, TEXT("\tBytes / Entity  : %zu"), TotalBytesPerEntity);
+	Ar.Logf(ELogVerbosity::Log, TEXT("\tBytes / Entity  : %u"), TotalBytesPerEntity);
 	Ar.Logf(ELogVerbosity::Log, TEXT("\tEntities / Chunk: %d"), NumEntitiesPerChunk);
 
 	Ar.Logf(ELogVerbosity::Log, TEXT("\tOffset 0x%04X: Entity[] (%llu bytes each)"), EntityListOffsetWithinChunk, sizeof(FMassEntityHandle));
diff --git a/Engine/Source/Runtime/MassEntity/Public/MassEntitySettings.h b/Engine/Source/Runtime/MassEntity/Public/MassEntitySettings.h
--- a/Engine/Source/Runtime/MassEntity/Public/MassEntitySettings.h
+++ b/Engine/Source/Runtime/MassEntity/Public/MassEntitySettings.h
@@ -57,7 +57,7 @@ protected:
 
 public:
 	UPROPERTY(EditDefaultsOnly, Category = Mass, config, AdvancedDisplay)
-	int32 ChunkMemorySize = 128 * 1024;
+	uint32 ChunkMemorySize = 128 * 1024;
 
 	/**
 	 * The name of the file to dump the processor dependency graph. T
