[Mass] Limited access to FMassArchetypeComposition's bitsets in preparation for near-future changes

diff --git a/Engine/Plugins/AI/MassAI/Source/MassAIDebug/Private/GameplayDebuggerCategory_Mass.cpp b/Engine/Plugins/AI/MassAI/Source/MassAIDebug/Private/GameplayDebuggerCategory_Mass.cpp
--- a/Engine/Plugins/AI/MassAI/Source/MassAIDebug/Private/GameplayDebuggerCategory_Mass.cpp
+++ b/Engine/Plugins/AI/MassAI/Source/MassAIDebug/Private/GameplayDebuggerCategory_Mass.cpp
@@ -472,20 +472,24 @@ void FGameplayDebuggerCategory_Mass::CollectData(APlayerController* OwnerPC, AAc
 					};
 
 					TArray<FName> ItemNames;
-					Composition.Tags.DebugGetIndividualNames(ItemNames);
+					Composition.GetTags().DebugGetIndividualNames(ItemNames);
 					AddTextLine(FString::Printf(TEXT("{Green}Tags:{White}%s"), *DescriptionBuilder(ItemNames)));
 					
 					ItemNames.Reset();
-					Composition.Fragments.DebugGetIndividualNames(ItemNames);
+					Composition.GetFragments().DebugGetIndividualNames(ItemNames);
 					AddTextLine(FString::Printf(TEXT("{Green}Fragments:{White}%s"), *DescriptionBuilder(ItemNames)));
 					
 					ItemNames.Reset();
-					Composition.ChunkFragments.DebugGetIndividualNames(ItemNames);
+					Composition.GetChunkFragments().DebugGetIndividualNames(ItemNames);
 					AddTextLine(FString::Printf(TEXT("{Green}Chunk Fragments:{White}%s"), *DescriptionBuilder(ItemNames)));
 
 					ItemNames.Reset();
-					Composition.SharedFragments.DebugGetIndividualNames(ItemNames);
+					Composition.GetSharedFragments().DebugGetIndividualNames(ItemNames);
 					AddTextLine(FString::Printf(TEXT("{Green}Shared Fragments:{White}%s"), *DescriptionBuilder(ItemNames)));
+
+					ItemNames.Reset();
+					Composition.GetConstSharedFragments().DebugGetIndividualNames(ItemNames);
+					AddTextLine(FString::Printf(TEXT("{Green}Const Shared Fragments:{White}%s"), *DescriptionBuilder(ItemNames)));
 				}
 
 				if (TransformFragment)
diff --git a/Engine/Plugins/Experimental/EditorDataStorage/Source/TedsCore/Private/Processors/TypedElementProcessorAdaptors.cpp b/Engine/Plugins/Experimental/EditorDataStorage/Source/TedsCore/Private/Processors/TypedElementProcessorAdaptors.cpp
--- a/Engine/Plugins/Experimental/EditorDataStorage/Source/TedsCore/Private/Processors/TypedElementProcessorAdaptors.cpp
+++ b/Engine/Plugins/Experimental/EditorDataStorage/Source/TedsCore/Private/Processors/TypedElementProcessorAdaptors.cpp
@@ -103,11 +103,11 @@ namespace UE::Editor::DataStorage
 
 				if (UE::Mass::IsA<FMassTag>(ColumnType))
 				{
-					return Composition.Tags.Contains(*ColumnType);
+					return Composition.GetTags().Contains(*ColumnType);
 				}
 				if (UE::Mass::IsA<FMassFragment>(ColumnType))
 				{
-					return Composition.Fragments.Contains(*ColumnType);
+					return Composition.GetFragments().Contains(*ColumnType);
 				}
 				const bool bIsTagOrFragment = false;
 				checkf(bIsTagOrFragment, TEXT("Attempting to check for a column type that is not a column or tag."));
@@ -216,7 +216,7 @@ namespace UE::Editor::DataStorage
 					{
 						if (this->Context.DoesArchetypeHaveTag(*ColumnType))
 						{
-							Descriptor.Tags.Add(*ColumnType);
+							Descriptor.GetTags().Add(*ColumnType);
 						}
 					}
 					else
@@ -225,7 +225,7 @@ namespace UE::Editor::DataStorage
 							TEXT("Given struct type is not a valid fragment or tag type."));
 						if (this->Context.DoesArchetypeHaveFragment(*ColumnType))
 						{
-							Descriptor.Fragments.Add(*ColumnType);
+							Descriptor.GetFragments().Add(*ColumnType);
 						}
 					}
 				}
@@ -239,13 +239,13 @@ namespace UE::Editor::DataStorage
 				{
 					if (UE::Mass::IsA<FMassTag>(ColumnType))
 					{
-						Descriptor.Tags.Add(*ColumnType);
+						Descriptor.GetTags().Add(*ColumnType);
 					}
 					else
 					{
 						checkf(UE::Mass::IsA<FMassFragment>(ColumnType),
 							TEXT("Given struct type is not a valid fragment or tag type."));
-						Descriptor.Fragments.Add(*ColumnType);
+						Descriptor.GetFragments().Add(*ColumnType);
 
 					}
 				}
@@ -287,7 +287,7 @@ namespace UE::Editor::DataStorage
 							UniqueAdd;
 						for (const FAddedColumns* const AddColumns : AddedColumnsQueue)
 						{
-							UniqueAdd.Add(TPair<FMassTagBitSet, FMassFragmentBitSet>(AddColumns->AddDescriptor.Tags, AddColumns->AddDescriptor.Fragments));
+							UniqueAdd.Add(TPair<FMassTagBitSet, FMassFragmentBitSet>(AddColumns->AddDescriptor.GetTags(), AddColumns->AddDescriptor.GetFragments()));
 						}
 						const int32 AddedColumnsRowGuess = UniqueAdd.IsEmpty() ? 0 : AddedColumnsQueue.Num() / UniqueAdd.Num();
 						UniqueAdd.Empty();
@@ -346,8 +346,8 @@ namespace UE::Editor::DataStorage
 						for (FAddedColumns* const AddColumns : AddedColumnsQueue)
 						{
 							FRowOperations& Operations = RowOperations.FindOrAdd(AddColumns->Entity);
-							Operations.ColumnsChange.AddedTags = AddColumns->AddDescriptor.Tags;
-							Operations.ColumnsChange.AddedFragments = AddColumns->AddDescriptor.Fragments;
+							Operations.ColumnsChange.AddedTags = AddColumns->AddDescriptor.GetTags();
+							Operations.ColumnsChange.AddedFragments = AddColumns->AddDescriptor.GetFragments();
 						}
 						AddedColumnsQueue.Empty();
 
@@ -356,8 +356,8 @@ namespace UE::Editor::DataStorage
 							for (const FMassEntityHandle& Row : AddColumns->Entities)
 							{
 								FRowOperations& Operations = RowOperations.FindOrAdd(Row);
-								Operations.ColumnsChange.AddedTags += AddColumns->AddDescriptor.Tags;
-								Operations.ColumnsChange.AddedFragments += AddColumns->AddDescriptor.Fragments;
+								Operations.ColumnsChange.AddedTags += AddColumns->AddDescriptor.GetTags();
+								Operations.ColumnsChange.AddedFragments += AddColumns->AddDescriptor.GetFragments();
 							}
 						}
 						AddedColumnsBatchQueue.Empty();
@@ -400,8 +400,8 @@ namespace UE::Editor::DataStorage
 						for (FRemovedColumns* const RemovedColumns : RemovedColumnsQueue)
 						{
 							FRowOperations& Operations = RowOperations.FindOrAdd(RemovedColumns->Entity);
-							Operations.ColumnsChange.RemovedTags += RemovedColumns->RemoveDescriptor.Tags;
-							Operations.ColumnsChange.RemovedFragments += RemovedColumns->RemoveDescriptor.Fragments;
+							Operations.ColumnsChange.RemovedTags += RemovedColumns->RemoveDescriptor.GetTags();
+							Operations.ColumnsChange.RemovedFragments += RemovedColumns->RemoveDescriptor.GetFragments();
 						}
 						RemovedColumnsQueue.Empty();
 
@@ -410,8 +410,8 @@ namespace UE::Editor::DataStorage
 							for (const FMassEntityHandle& Row : RemovedColumns->Entities)
 							{
 								FRowOperations& Operations = RowOperations.FindOrAdd(Row);
-								Operations.ColumnsChange.RemovedTags += RemovedColumns->RemoveDescriptor.Tags;
-								Operations.ColumnsChange.RemovedFragments += RemovedColumns->RemoveDescriptor.Fragments;
+								Operations.ColumnsChange.RemovedTags += RemovedColumns->RemoveDescriptor.GetTags();
+								Operations.ColumnsChange.RemovedFragments += RemovedColumns->RemoveDescriptor.GetFragments();
 							}
 						}
 						RemovedColumnsBatchQueue.Empty();
@@ -869,7 +869,7 @@ namespace UE::Editor::DataStorage
 
 
 							// This could be improved by adding an operation that would both combine the Fragments and Tags change in one bath operation.
-							if (!RemovedColumns->RemoveDescriptor.Fragments.IsEmpty())
+							if (!RemovedColumns->RemoveDescriptor.GetFragments().IsEmpty())
 							{
 								for (FMassEntityHandle Entity : Entities)
 								{
@@ -889,9 +889,9 @@ namespace UE::Editor::DataStorage
 									EntityCollections.Emplace(It.Key(), It.Value(), FMassArchetypeEntityCollection::EDuplicatesHandling::FoldDuplicates);
 								}
 
-								System.BatchChangeFragmentCompositionForEntities(EntityCollections, FMassFragmentBitSet(), RemovedColumns->RemoveDescriptor.Fragments);
+								System.BatchChangeFragmentCompositionForEntities(EntityCollections, FMassFragmentBitSet(), RemovedColumns->RemoveDescriptor.GetFragments());
 							}
-							if (!RemovedColumns->RemoveDescriptor.Tags.IsEmpty())
+							if (!RemovedColumns->RemoveDescriptor.GetTags().IsEmpty())
 							{
 								LookupTable.Reset();
 								EntityCollections.Reset();
@@ -914,7 +914,7 @@ namespace UE::Editor::DataStorage
 									EntityCollections.Emplace(It.Key(), It.Value(), FMassArchetypeEntityCollection::EDuplicatesHandling::FoldDuplicates);
 								}
 
-								System.BatchChangeTagsForEntities(EntityCollections, FMassTagBitSet(), RemovedColumns->RemoveDescriptor.Tags);
+								System.BatchChangeTagsForEntities(EntityCollections, FMassTagBitSet(), RemovedColumns->RemoveDescriptor.GetTags());
 							}
 						});
 					}
diff --git a/Engine/Plugins/Experimental/EditorDataStorage/Source/TedsCore/Private/Queries/TypedElementExtendedQueryStore.cpp b/Engine/Plugins/Experimental/EditorDataStorage/Source/TedsCore/Private/Queries/TypedElementExtendedQueryStore.cpp
--- a/Engine/Plugins/Experimental/EditorDataStorage/Source/TedsCore/Private/Queries/TypedElementExtendedQueryStore.cpp
+++ b/Engine/Plugins/Experimental/EditorDataStorage/Source/TedsCore/Private/Queries/TypedElementExtendedQueryStore.cpp
@@ -54,10 +54,10 @@ namespace UE::Editor::DataStorage
 		StringBuilder.Append( TEXT("Composition Descriptor:\nColumns: {"));
 		{
 			int32 FragmentsProcessed = 0;
-			const int32 FragmentCount = CompositionDescriptor.Fragments.CountStoredTypes();
-			for(auto IndexIterator = CompositionDescriptor.Fragments.GetIndexIterator(); IndexIterator; ++IndexIterator)
+			const int32 FragmentCount = CompositionDescriptor.GetFragments().CountStoredTypes();
+			for(auto IndexIterator = CompositionDescriptor.GetFragments().GetIndexIterator(); IndexIterator; ++IndexIterator)
 			{
-				StringBuilder.Append(CompositionDescriptor.Fragments.GetTypeAtIndex(*IndexIterator)->GetName());
+				StringBuilder.Append(CompositionDescriptor.GetFragments().GetTypeAtIndex(*IndexIterator)->GetName());
 				if (FragmentsProcessed != FragmentCount - 1)
 				{
 					StringBuilder.Append(TEXT(", "));
@@ -68,11 +68,11 @@ namespace UE::Editor::DataStorage
 		}
 		{
 			int32 TagsProcessed = 0;
-			const int32 TagsCount = CompositionDescriptor.Tags.CountStoredTypes();
+			const int32 TagsCount = CompositionDescriptor.GetTags().CountStoredTypes();
 			StringBuilder.Append(TEXT("Tags: {"));
-			for(auto IndexIterator = CompositionDescriptor.Tags.GetIndexIterator(); IndexIterator; ++IndexIterator)
+			for(auto IndexIterator = CompositionDescriptor.GetTags().GetIndexIterator(); IndexIterator; ++IndexIterator)
 			{
-				StringBuilder.Append(CompositionDescriptor.Tags.GetTypeAtIndex(*IndexIterator)->GetName());
+				StringBuilder.Append(CompositionDescriptor.GetTags().GetTypeAtIndex(*IndexIterator)->GetName());
 				if (TagsProcessed != TagsCount - 1)
 				{
 					StringBuilder.Append(TEXT(", "));
@@ -96,10 +96,10 @@ namespace UE::Editor::DataStorage
 		if (Query->Description.Conditions)
 		{
 			TArray<TWeakObjectPtr<const UScriptStruct>, TInlineAllocator<64>> ColumnTypes;
-			auto FragmentIterator = CompositionDescriptor.Fragments.GetIndexIterator(true);
+			auto FragmentIterator = CompositionDescriptor.GetFragments().GetIndexIterator(true);
 			while (FragmentIterator)
 			{
-				const UScriptStruct* FragmentStruct = CompositionDescriptor.Fragments.GetTypeAtIndex(*FragmentIterator);
+				const UScriptStruct* FragmentStruct = CompositionDescriptor.GetFragments().GetTypeAtIndex(*FragmentIterator);
 				ColumnTypes.Add(FragmentStruct);
 				++FragmentIterator;
 			}
@@ -118,10 +118,10 @@ namespace UE::Editor::DataStorage
 			}
 
 			// Then ensure it also matches the query conditions
-			auto TagIterator = CompositionDescriptor.Tags.GetIndexIterator(true);			
+			auto TagIterator = CompositionDescriptor.GetTags().GetIndexIterator(true);			
 			while (TagIterator)
 			{
-				const UScriptStruct* TagStruct = CompositionDescriptor.Tags.GetTypeAtIndex(*TagIterator);
+				const UScriptStruct* TagStruct = CompositionDescriptor.GetTags().GetTypeAtIndex(*TagIterator);
 				ColumnTypes.Add(TagStruct);
 				++TagIterator;
 			}
diff --git a/Engine/Plugins/Experimental/EditorDataStorage/Source/TedsCore/Private/TypedElementDatabase.cpp b/Engine/Plugins/Experimental/EditorDataStorage/Source/TedsCore/Private/TypedElementDatabase.cpp
--- a/Engine/Plugins/Experimental/EditorDataStorage/Source/TedsCore/Private/TypedElementDatabase.cpp
+++ b/Engine/Plugins/Experimental/EditorDataStorage/Source/TedsCore/Private/TypedElementDatabase.cpp
@@ -888,11 +888,11 @@ bool UEditorDataStorage::HasColumns(UE::Editor::DataStorage::RowHandle Row, TCon
 			{
 				if (UE::Mass::IsA<FMassFragment>(*ColumnType))
 				{
-					bHasAllColumns = Composition.Fragments.Contains(**ColumnType);
+					bHasAllColumns = Composition.GetFragments().Contains(**ColumnType);
 				}
 				else if (UE::Mass::IsA<FMassTag>(*ColumnType))
 				{
-					bHasAllColumns = Composition.Tags.Contains(**ColumnType);
+					bHasAllColumns = Composition.GetTags().Contains(**ColumnType);
 				}
 				else
 				{
@@ -934,12 +934,12 @@ bool UEditorDataStorage::HasColumns(UE::Editor::DataStorage::RowHandle Row, TCon
 				{
 					if (UE::Mass::IsA<FMassFragment>(ColumnType->Get()))
 					{
-						bHasAllColumns = Composition.Fragments.Contains(**ColumnType);
+						bHasAllColumns = Composition.GetFragments().Contains(**ColumnType);
 						continue;
 					}
 					else if (UE::Mass::IsA<FMassTag>(ColumnType->Get()))
 					{
-						bHasAllColumns = Composition.Tags.Contains(**ColumnType);
+						bHasAllColumns = Composition.GetTags().Contains(**ColumnType);
 						continue;
 					}
 				}
@@ -979,8 +979,8 @@ void UEditorDataStorage::ListColumns(UE::Editor::DataStorage::RowHandle Row, UE:
 					}
 					return true;
 				};
-			Composition.Fragments.ExportTypes(CallbackWrapper);
-			Composition.Tags.ExportTypes(CallbackWrapper);
+			Composition.GetFragments().ExportTypes(CallbackWrapper);
+			Composition.GetTags().ExportTypes(CallbackWrapper);
 		}
 		else
 		{
@@ -999,7 +999,7 @@ void UEditorDataStorage::ListColumns(UE::Editor::DataStorage::RowHandle Row, UE:
 			FMassArchetypeHandle Archetype = ActiveEditorEntityManager->GetArchetypeForEntity(Entity);
 			const FMassArchetypeCompositionDescriptor& Composition = ActiveEditorEntityManager->GetArchetypeComposition(Archetype);
 
-			Composition.Fragments.ExportTypes(
+			Composition.GetFragments().ExportTypes(
 				[this, &Callback, Entity](const UScriptStruct* ColumnType)
 				{
 					if (ColumnType)
@@ -1008,7 +1008,7 @@ void UEditorDataStorage::ListColumns(UE::Editor::DataStorage::RowHandle Row, UE:
 					}
 					return true;
 				});
-			Composition.Tags.ExportTypes(
+			Composition.GetTags().ExportTypes(
 				[&Callback](const UScriptStruct* ColumnType)
 				{
 					if (ColumnType)
diff --git a/Engine/Plugins/Runtime/MassGameplay/Source/MassGameplayEditor/Private/MassTraitRepository.cpp b/Engine/Plugins/Runtime/MassGameplay/Source/MassGameplayEditor/Private/MassTraitRepository.cpp
--- a/Engine/Plugins/Runtime/MassGameplay/Source/MassGameplayEditor/Private/MassTraitRepository.cpp
+++ b/Engine/Plugins/Runtime/MassGameplay/Source/MassGameplayEditor/Private/MassTraitRepository.cpp
@@ -365,33 +365,33 @@ void UMassTraitRepository::OnNewTraitType(UMassEntityTraitBase& Trait)
 
 	const FMassArchetypeCompositionDescriptor& Composition = InvestigationContext.EntityTemplate.GetCompositionDescriptor();
 
-	for (FMassFragmentBitSet::FIndexIterator It = Composition.Fragments.GetIndexIterator(); It; ++It)
+	for (FMassFragmentBitSet::FIndexIterator It = Composition.GetFragments().GetIndexIterator(); It; ++It)
 	{
-		const FName StructName = Composition.Fragments.DebugGetStructTypeName(*It);
+		const FName StructName = Composition.GetFragments().DebugGetStructTypeName(*It);
 		TraitData.ElementNames.Add(StructName);
 	}
 
-	for (FMassTagBitSet::FIndexIterator It = Composition.Tags.GetIndexIterator(); It; ++It)
+	for (FMassTagBitSet::FIndexIterator It = Composition.GetTags().GetIndexIterator(); It; ++It)
 	{
-		const FName StructName = Composition.Tags.DebugGetStructTypeName(*It);
+		const FName StructName = Composition.GetTags().DebugGetStructTypeName(*It);
 		TraitData.ElementNames.Add(StructName);
 	}
 
-	for (FMassChunkFragmentBitSet::FIndexIterator It = Composition.ChunkFragments.GetIndexIterator(); It; ++It)
+	for (FMassChunkFragmentBitSet::FIndexIterator It = Composition.GetChunkFragments().GetIndexIterator(); It; ++It)
 	{
-		const FName StructName = Composition.ChunkFragments.DebugGetStructTypeName(*It);
+		const FName StructName = Composition.GetChunkFragments().DebugGetStructTypeName(*It);
 		TraitData.ElementNames.Add(StructName);
 	}
 
-	for (FMassSharedFragmentBitSet::FIndexIterator It = Composition.SharedFragments.GetIndexIterator(); It; ++It)
+	for (FMassSharedFragmentBitSet::FIndexIterator It = Composition.GetSharedFragments().GetIndexIterator(); It; ++It)
 	{
-		const FName StructName = Composition.SharedFragments.DebugGetStructTypeName(*It);
+		const FName StructName = Composition.GetSharedFragments().DebugGetStructTypeName(*It);
 		TraitData.ElementNames.Add(StructName);
 	}
 
-	for (FMassConstSharedFragmentBitSet::FIndexIterator It = Composition.ConstSharedFragments.GetIndexIterator(); It; ++It)
+	for (FMassConstSharedFragmentBitSet::FIndexIterator It = Composition.GetConstSharedFragments().GetIndexIterator(); It; ++It)
 	{
-		const FName StructName = Composition.ConstSharedFragments.DebugGetStructTypeName(*It);
+		const FName StructName = Composition.GetConstSharedFragments().DebugGetStructTypeName(*It);
 		TraitData.ElementNames.Add(StructName);
 	}
 
diff --git a/Engine/Plugins/Runtime/MassGameplay/Source/MassGameplayTestSuite/Private/MassEntityTemplateTest.cpp b/Engine/Plugins/Runtime/MassGameplay/Source/MassGameplayTestSuite/Private/MassEntityTemplateTest.cpp
--- a/Engine/Plugins/Runtime/MassGameplay/Source/MassGameplayTestSuite/Private/MassEntityTemplateTest.cpp
+++ b/Engine/Plugins/Runtime/MassGameplay/Source/MassGameplayTestSuite/Private/MassEntityTemplateTest.cpp
@@ -74,11 +74,11 @@ struct FEntityTemplate_Composition : FEntityTemplateBase
 		TemplateData.Sort();
 
 		FMassArchetypeCompositionDescriptor ExpectedComposition;
-		ExpectedComposition.Fragments.Add<FTestFragment_Int>();
-		ExpectedComposition.Fragments.Add<FTestFragment_Float>();
-		ExpectedComposition.ChunkFragments.Add<FTestChunkFragment_Int>();
-		ExpectedComposition.SharedFragments.Add<FTestSharedFragment_Int>();
-		ExpectedComposition.SharedFragments.Add<FTestSharedFragment_Float>();
+		ExpectedComposition.GetFragments().Add<FTestFragment_Int>();
+		ExpectedComposition.GetFragments().Add<FTestFragment_Float>();
+		ExpectedComposition.GetChunkFragments().Add<FTestChunkFragment_Int>();
+		ExpectedComposition.GetSharedFragments().Add<FTestSharedFragment_Int>();
+		ExpectedComposition.GetSharedFragments().Add<FTestSharedFragment_Float>();
 
 		AITEST_TRUE("The composition should end up being the same regardless of whether the data is added via a template data or a composition descriptor"
 			, ExpectedComposition.IsEquivalent(TemplateData.GetCompositionDescriptor()));
diff --git a/Engine/Plugins/Runtime/MassGameplay/Source/MassSpawner/Public/MassEntityTemplate.h b/Engine/Plugins/Runtime/MassGameplay/Source/MassSpawner/Public/MassEntityTemplate.h
--- a/Engine/Plugins/Runtime/MassGameplay/Source/MassSpawner/Public/MassEntityTemplate.h
+++ b/Engine/Plugins/Runtime/MassGameplay/Source/MassSpawner/Public/MassEntityTemplate.h
@@ -130,13 +130,13 @@ struct FMassEntityTemplateData
 	void AddFragment()
 	{
 		static_assert(UE::Mass::CFragment<T>, "Given struct doesn't represent a valid fragment type. Make sure to inherit from FMassFragment or one of its child-types.");
-		Composition.Fragments.Add<T>();
+		Composition.GetFragments().Add<T>();
 	}
 
 	void AddFragment(const UScriptStruct& FragmentType)
 	{
 		checkf(UE::Mass::IsA<FMassFragment>(&FragmentType), TEXT("Given struct doesn't represent a valid fragment type. Make sure to inherit from FMassFragment or one of its child-types."));
-		Composition.Fragments.Add(FragmentType);
+		Composition.GetFragments().Add(FragmentType);
 	}
 
 	// @todo this function is doing nothing if a given fragment's initial value has already been created. This seems inconsistent with the other AddFragment functions (especially AddFragment_GetRef).
@@ -144,9 +144,9 @@ struct FMassEntityTemplateData
 	{
 		const UScriptStruct* FragmentType = Fragment.GetScriptStruct();
 		checkf(UE::Mass::IsA<FMassFragment>(FragmentType), TEXT("Given struct doesn't represent a valid fragment type. Make sure to inherit from FMassFragment or one of its child-types."));
-		if (!Composition.Fragments.Contains(*FragmentType))
+		if (!Composition.GetFragments().Contains(*FragmentType))
 		{
-			Composition.Fragments.Add(*FragmentType);
+			Composition.GetFragments().Add(*FragmentType);
 			InitialFragmentValues.Emplace(Fragment);
 		}
 		else if (!InitialFragmentValues.ContainsByPredicate(FStructTypeEqualOperator(FragmentType)))
@@ -159,9 +159,9 @@ struct FMassEntityTemplateData
 	T& AddFragment_GetRef()
 	{
 		static_assert(UE::Mass::CFragment<T>, "Given struct doesn't represent a valid fragment type. Make sure to inherit from FMassFragment or one of its child-types.");
-		if (!Composition.Fragments.Contains<T>())
+		if (!Composition.GetFragments().Contains<T>())
 		{
-			Composition.Fragments.Add<T>();
+			Composition.GetFragments().Add<T>();
 		}
 		else if (FInstancedStruct* Fragment = InitialFragmentValues.FindByPredicate(FStructTypeEqualOperator(T::StaticStruct())))
 		{
@@ -184,26 +184,26 @@ struct FMassEntityTemplateData
 	void AddTag()
 	{
 		static_assert(UE::Mass::CTag<T>, "Given struct doesn't represent a valid mass tag type. Make sure to inherit from FMassTag or one of its child-types.");
-		Composition.Tags.Add<T>();
+		Composition.GetTags().Add<T>();
 	}
 	
 	void AddTag(const UScriptStruct& TagType)
 	{
 		checkf(UE::Mass::IsA<FMassTag>(&TagType), TEXT("Given struct doesn't represent a valid mass tag type. Make sure to inherit from FMassTag or one of its child-types."));
-		Composition.Tags.Add(TagType);
+		Composition.GetTags().Add(TagType);
 	}
 
 	template<typename T>
 	void RemoveTag()
 	{
 		static_assert(UE::Mass::CTag<T>, "Given struct doesn't represent a valid mass tag type. Make sure to inherit from FMassTag or one of its child-types.");
-		Composition.Tags.Remove<T>();
+		Composition.GetTags().Remove<T>();
 	}
 
 	void RemoveTag(const UScriptStruct& TagType)
 	{
 		checkf(UE::Mass::IsA<FMassTag>(&TagType), TEXT("Given struct doesn't represent a valid mass tag type. Make sure to inherit from FMassTag or one of its child-types."));
-		Composition.Tags.Remove(TagType);
+		Composition.GetTags().Remove(TagType);
 	}
 
 	const FMassTagBitSet& GetTags() const;
@@ -213,7 +213,7 @@ struct FMassEntityTemplateData
 	void AddChunkFragment()
 	{
 		static_assert(UE::Mass::CChunkFragment<T>, "Given struct doesn't represent a valid chunk fragment type. Make sure to inherit from FMassChunkFragment or one of its child-types.");
-		Composition.ChunkFragments.Add<T>();
+		Composition.GetChunkFragments().Add<T>();
 	}
 
 	void AddConstSharedFragment(const FConstSharedStruct& SharedFragment)
@@ -221,9 +221,9 @@ struct FMassEntityTemplateData
 		const UScriptStruct* FragmentType = SharedFragment.GetScriptStruct();
 		if(ensureMsgf(UE::Mass::IsA<FMassConstSharedFragment>(FragmentType), TEXT("Given struct doesn't represent a valid const shared fragment type. Make sure to inherit from FMassConstSharedFragment or one of its child-types.")))
 		{
-			if (!Composition.ConstSharedFragments.Contains(*FragmentType))
+			if (!Composition.GetConstSharedFragments().Contains(*FragmentType))
 			{
-				Composition.ConstSharedFragments.Add(*FragmentType);
+				Composition.GetConstSharedFragments().Add(*FragmentType);
 				SharedFragmentValues.Add(SharedFragment);
 			}
 #if DO_ENSURE
@@ -242,9 +242,9 @@ struct FMassEntityTemplateData
 		const UScriptStruct* FragmentType = SharedFragment.GetScriptStruct();
 		if(ensureMsgf(UE::Mass::IsA<FMassSharedFragment>(FragmentType), TEXT("Given struct doesn't represent a valid shared fragment type. Make sure to inherit from FMassSharedFragment or one of its child-types.")))
 		{
-			if (!Composition.SharedFragments.Contains(*FragmentType))
+			if (!Composition.GetSharedFragments().Contains(*FragmentType))
 			{
-				Composition.SharedFragments.Add(*FragmentType);
+				Composition.GetSharedFragments().Add(*FragmentType);
 				SharedFragmentValues.Add(SharedFragment);
 			}
 	#if DO_ENSURE
@@ -261,46 +261,46 @@ struct FMassEntityTemplateData
 	template<typename T>
 	bool HasFragment() const
 	{
-		return Composition.Fragments.Contains<T>();
+		return Composition.GetFragments().Contains<T>();
 	}
 	
 	bool HasFragment(const UScriptStruct& ScriptStruct) const
 	{
-		return Composition.Fragments.Contains(ScriptStruct);
+		return Composition.GetFragments().Contains(ScriptStruct);
 	}
 
 	template<typename T>
 	bool HasTag() const
 	{
-		return Composition.Tags.Contains<T>();
+		return Composition.GetTags().Contains<T>();
 	}
 
 	template<typename T>
 	bool HasChunkFragment() const
 	{
-		return Composition.ChunkFragments.Contains<T>();
+		return Composition.GetChunkFragments().Contains<T>();
 	}
 
 	template<typename T>
 	bool HasSharedFragment() const
 	{
-		return Composition.SharedFragments.Contains<T>();
+		return Composition.GetSharedFragments().Contains<T>();
 	}
 
 	bool HasSharedFragment(const UScriptStruct& ScriptStruct) const
 	{
-		return Composition.SharedFragments.Contains(ScriptStruct);
+		return Composition.GetSharedFragments().Contains(ScriptStruct);
 	}
 
 	template<typename T>
 	bool HasConstSharedFragment() const
 	{
-		return Composition.ConstSharedFragments.Contains<T>();
+		return Composition.GetConstSharedFragments().Contains<T>();
 	}
 
 	bool HasConstSharedFragment(const UScriptStruct& ScriptStruct) const
 	{
-		return Composition.ConstSharedFragments.Contains(ScriptStruct);
+		return Composition.GetConstSharedFragments().Contains(ScriptStruct);
 	}
 
 	void Sort()
@@ -430,12 +430,12 @@ inline void FMassEntityTemplateData::SetTemplateName(const FString& Name)
 
 inline const FMassTagBitSet& FMassEntityTemplateData::GetTags() const
 {
-	return Composition.Tags;
+	return Composition.GetTags();
 }
 
 inline FMassTagBitSet& FMassEntityTemplateData::GetMutableTags()
 {
-	return Composition.Tags;
+	return Composition.GetTags();
 }
 
 inline FMassArchetypeCreationParams& FMassEntityTemplateData::GetArchetypeCreationParams()
diff --git a/Engine/Source/Developer/MassEntityTestSuite/Private/MassQueryTest.cpp b/Engine/Source/Developer/MassEntityTestSuite/Private/MassQueryTest.cpp
--- a/Engine/Source/Developer/MassEntityTestSuite/Private/MassQueryTest.cpp
+++ b/Engine/Source/Developer/MassEntityTestSuite/Private/MassQueryTest.cpp
@@ -608,13 +608,13 @@ struct FQueryTest_AllOptional : FEntityTestBase
 
 		{
 			FMassArchetypeCompositionDescriptor Descriptor(EntityManager->GetArchetypeComposition(IntsArchetype));
-			Descriptor.ChunkFragments.Add<FTestChunkFragment_Int>();
+			Descriptor.GetChunkFragments().Add<FTestChunkFragment_Int>();
 			EntityManager->CreateArchetype(Descriptor);
 			++ExpectedNumOfArchetypes;
 		}
 		{
 			FMassArchetypeCompositionDescriptor Descriptor(EntityManager->GetArchetypeComposition(IntsArchetype));
-			Descriptor.ChunkFragments.Add<FTestChunkFragment_Float>();
+			Descriptor.GetChunkFragments().Add<FTestChunkFragment_Float>();
 			EntityManager->CreateArchetype(Descriptor);
 		}
 		Query.CacheArchetypes();
@@ -677,13 +677,13 @@ struct FQueryTest_JustATag : FEntityTestBase
 
 		{
 			FMassArchetypeCompositionDescriptor Descriptor(EntityManager->GetArchetypeComposition(IntsArchetype));
-			Descriptor.Tags.Add<FTestTag_A>();
+			Descriptor.GetTags().Add<FTestTag_A>();
 			EntityManager->CreateArchetype(Descriptor);
 			++ExpectedNumOfArchetypes;
 		}
 		{
 			FMassArchetypeCompositionDescriptor Descriptor(EntityManager->GetArchetypeComposition(IntsArchetype));
-			Descriptor.Tags.Add<FTestTag_B>();
+			Descriptor.GetTags().Add<FTestTag_B>();
 			EntityManager->CreateArchetype(Descriptor);
 		}
 		Query.CacheArchetypes();
@@ -691,17 +691,17 @@ struct FQueryTest_JustATag : FEntityTestBase
 
 		{
 			FMassArchetypeCompositionDescriptor Descriptor(EntityManager->GetArchetypeComposition(IntsArchetype));
-			Descriptor.Tags.Add<FTestTag_A>();
-			Descriptor.Tags.Add<FTestTag_C>();
-			Descriptor.Tags.Add<FTestTag_D>();
+			Descriptor.GetTags().Add<FTestTag_A>();
+			Descriptor.GetTags().Add<FTestTag_C>();
+			Descriptor.GetTags().Add<FTestTag_D>();
 			EntityManager->CreateArchetype(Descriptor);
 			++ExpectedNumOfArchetypes;
 		}
 		{
 			FMassArchetypeCompositionDescriptor Descriptor(EntityManager->GetArchetypeComposition(IntsArchetype));
-			Descriptor.Tags.Add<FTestTag_B>();
-			Descriptor.Tags.Add<FTestTag_C>();
-			Descriptor.Tags.Add<FTestTag_D>();
+			Descriptor.GetTags().Add<FTestTag_B>();
+			Descriptor.GetTags().Add<FTestTag_C>();
+			Descriptor.GetTags().Add<FTestTag_D>();
 			EntityManager->CreateArchetype(Descriptor);
 		}
 		Query.CacheArchetypes();
@@ -730,13 +730,13 @@ struct FQueryTest_JustAChunkFragment : FEntityTestBase
 
 		{
 			FMassArchetypeCompositionDescriptor Descriptor(EntityManager->GetArchetypeComposition(IntsArchetype));
-			Descriptor.ChunkFragments.Add<FTestChunkFragment_Int>();
+			Descriptor.GetChunkFragments().Add<FTestChunkFragment_Int>();
 			TargetArchetype = EntityManager->CreateArchetype(Descriptor);
 			++ExpectedNumOfArchetypes;
 		}
 		{
 			FMassArchetypeCompositionDescriptor Descriptor(EntityManager->GetArchetypeComposition(IntsArchetype));
-			Descriptor.ChunkFragments.Add<FTestChunkFragment_Float>();
+			Descriptor.GetChunkFragments().Add<FTestChunkFragment_Float>();
 			EntityManager->CreateArchetype(Descriptor);
 		}
 		Query.CacheArchetypes();
diff --git a/Engine/Source/Editor/MassEntityDebugger/Private/MassDebuggerModel.cpp b/Engine/Source/Editor/MassEntityDebugger/Private/MassDebuggerModel.cpp
--- a/Engine/Source/Editor/MassEntityDebugger/Private/MassDebuggerModel.cpp
+++ b/Engine/Source/Editor/MassEntityDebugger/Private/MassDebuggerModel.cpp
@@ -32,10 +32,10 @@ namespace UE::Mass::Debugger::Private
 
 		check(TotalLength > 0);
 
-		return float(BitSetDistance(A.Composition.Fragments, B.Composition.Fragments)
-			+ BitSetDistance(A.Composition.Tags, B.Composition.Tags)
-			+ BitSetDistance(A.Composition.ChunkFragments, B.Composition.ChunkFragments)
-			+ BitSetDistance(A.Composition.SharedFragments, B.Composition.SharedFragments)) 
+		return float(BitSetDistance(A.Composition.GetFragments(), B.Composition.GetFragments())
+			+ BitSetDistance(A.Composition.GetTags(), B.Composition.GetTags())
+			+ BitSetDistance(A.Composition.GetChunkFragments(), B.Composition.GetChunkFragments())
+			+ BitSetDistance(A.Composition.GetSharedFragments(), B.Composition.GetSharedFragments())) 
 			/ TotalLength;
 	}
 
@@ -269,9 +269,9 @@ FMassDebuggerArchetypeData::FMassDebuggerArchetypeData(const FMassArchetypeHandl
 		PrimaryDebugName = HashAsString;
 
 		// Use first fragment as name
-		if (FMassFragmentBitSet::FIndexIterator It = Composition.Fragments.GetIndexIterator())
+		if (FMassFragmentBitSet::FIndexIterator It = Composition.GetFragments().GetIndexIterator())
 		{
-			const FName FirstStructName = Composition.Fragments.DebugGetStructTypeName(*It);
+			const FName FirstStructName = Composition.GetFragments().DebugGetStructTypeName(*It);
 			TStringBuilder<256> StringBuilder;
 			StringBuilder.Append(FirstStructName.ToString());
 			StringBuilder.Append(TEXT("..."));
diff --git a/Engine/Source/Editor/MassEntityDebugger/Private/SMassArchetype.cpp b/Engine/Source/Editor/MassEntityDebugger/Private/SMassArchetype.cpp
--- a/Engine/Source/Editor/MassEntityDebugger/Private/SMassArchetype.cpp
+++ b/Engine/Source/Editor/MassEntityDebugger/Private/SMassArchetype.cpp
@@ -105,19 +105,19 @@ void SMassArchetype::Construct(const FArguments& InArgs, TSharedPtr<FMassDebugge
 	{
 		const FMassArchetypeCompositionDescriptor& ParentComposition = BaseArchetypeDebugData->Composition;
 
-		UE::Mass::Debugger::UI::AddBitSetDiff(Box, ParentComposition.Fragments, Composition.Fragments, TEXT("Fragments"), Brush, Prune, InDebuggerModel);
-		UE::Mass::Debugger::UI::AddBitSetDiff(Box, ParentComposition.Tags, Composition.Tags, TEXT("Tags"), Brush, Prune, InDebuggerModel);
-		UE::Mass::Debugger::UI::AddBitSetDiff(Box, ParentComposition.ChunkFragments, Composition.ChunkFragments, TEXT("Chunk Fragments"), Brush, Prune, InDebuggerModel);
-		UE::Mass::Debugger::UI::AddBitSetDiff(Box, ParentComposition.SharedFragments, Composition.SharedFragments, TEXT("Shared Fragments"), Brush, Prune, InDebuggerModel);
-		UE::Mass::Debugger::UI::AddBitSetDiff(Box, ParentComposition.ConstSharedFragments, Composition.ConstSharedFragments, TEXT("Const Shared Fragments"), Brush, Prune, InDebuggerModel);
+		UE::Mass::Debugger::UI::AddBitSetDiff(Box, ParentComposition.GetFragments(), Composition.GetFragments(), TEXT("Fragments"), Brush, Prune, InDebuggerModel);
+		UE::Mass::Debugger::UI::AddBitSetDiff(Box, ParentComposition.GetTags(), Composition.GetTags(), TEXT("Tags"), Brush, Prune, InDebuggerModel);
+		UE::Mass::Debugger::UI::AddBitSetDiff(Box, ParentComposition.GetChunkFragments(), Composition.GetChunkFragments(), TEXT("Chunk Fragments"), Brush, Prune, InDebuggerModel);
+		UE::Mass::Debugger::UI::AddBitSetDiff(Box, ParentComposition.GetSharedFragments(), Composition.GetSharedFragments(), TEXT("Shared Fragments"), Brush, Prune, InDebuggerModel);
+		UE::Mass::Debugger::UI::AddBitSetDiff(Box, ParentComposition.GetConstSharedFragments(), Composition.GetConstSharedFragments(), TEXT("Const Shared Fragments"), Brush, Prune, InDebuggerModel);
 	}
 	else
 	{
-		UE::Mass::Debugger::UI::AddBitSet(Box, Composition.Fragments, TEXT("Fragments"), Brush, InDebuggerModel);
-		UE::Mass::Debugger::UI::AddBitSet(Box, Composition.Tags, TEXT("Tags"), Brush, InDebuggerModel);
-		UE::Mass::Debugger::UI::AddBitSet(Box, Composition.ChunkFragments, TEXT("Chunk Fragments"), Brush, InDebuggerModel);
-		UE::Mass::Debugger::UI::AddBitSet(Box, Composition.SharedFragments, TEXT("Shared Fragments"), Brush, InDebuggerModel);
-		UE::Mass::Debugger::UI::AddBitSet(Box, Composition.ConstSharedFragments, TEXT("Const Shared Fragments"), Brush, InDebuggerModel);
+		UE::Mass::Debugger::UI::AddBitSet(Box, Composition.GetFragments(), TEXT("Fragments"), Brush, InDebuggerModel);
+		UE::Mass::Debugger::UI::AddBitSet(Box, Composition.GetTags(), TEXT("Tags"), Brush, InDebuggerModel);
+		UE::Mass::Debugger::UI::AddBitSet(Box, Composition.GetChunkFragments(), TEXT("Chunk Fragments"), Brush, InDebuggerModel);
+		UE::Mass::Debugger::UI::AddBitSet(Box, Composition.GetSharedFragments(), TEXT("Shared Fragments"), Brush, InDebuggerModel);
+		UE::Mass::Debugger::UI::AddBitSet(Box, Composition.GetConstSharedFragments(), TEXT("Const Shared Fragments"), Brush, InDebuggerModel);
 	}
 
 	TSharedRef<SVerticalBox> MainBox = SNew(SVerticalBox);
diff --git a/Engine/Source/Runtime/MassEntity/Internal/MassArchetypeData.h b/Engine/Source/Runtime/MassEntity/Internal/MassArchetypeData.h
--- a/Engine/Source/Runtime/MassEntity/Internal/MassArchetypeData.h
+++ b/Engine/Source/Runtime/MassEntity/Internal/MassArchetypeData.h
@@ -229,11 +229,11 @@ public:
 	explicit FMassArchetypeData(const FMassArchetypeCreationParams& CreationParams = FMassArchetypeCreationParams());
 
 	TConstArrayView<FMassArchetypeFragmentConfig> GetFragmentConfigs() const { return FragmentConfigs; }
-	const FMassFragmentBitSet& GetFragmentBitSet() const { return CompositionDescriptor.Fragments; }
-	const FMassTagBitSet& GetTagBitSet() const { return CompositionDescriptor.Tags; }
-	const FMassChunkFragmentBitSet& GetChunkFragmentBitSet() const { return CompositionDescriptor.ChunkFragments; }
-	const FMassSharedFragmentBitSet& GetSharedFragmentBitSet() const { return CompositionDescriptor.SharedFragments; }
-	const FMassConstSharedFragmentBitSet& GetConstSharedFragmentBitSet() const { return CompositionDescriptor.ConstSharedFragments; }
+	const FMassFragmentBitSet& GetFragmentBitSet() const { return CompositionDescriptor.GetFragments(); }
+	const FMassTagBitSet& GetTagBitSet() const { return CompositionDescriptor.GetTags(); }
+	const FMassChunkFragmentBitSet& GetChunkFragmentBitSet() const { return CompositionDescriptor.GetChunkFragments(); }
+	const FMassSharedFragmentBitSet& GetSharedFragmentBitSet() const { return CompositionDescriptor.GetSharedFragments(); }
+	const FMassConstSharedFragmentBitSet& GetConstSharedFragmentBitSet() const { return CompositionDescriptor.GetConstSharedFragments(); }
 
 	const FMassArchetypeCompositionDescriptor& GetCompositionDescriptor() const { return CompositionDescriptor; }
 	FORCEINLINE const FMassArchetypeSharedFragmentValues& GetSharedFragmentValues(int32 EntityIndex) const
@@ -255,7 +255,7 @@ public:
 	/** Method to iterate on all the fragment types */
 	void ForEachFragmentType(TFunction< void(const UScriptStruct* /*FragmentType*/)> Function) const;
 	bool HasFragmentType(const UScriptStruct* FragmentType) const;
-	bool HasTagType(const UScriptStruct* FragmentType) const { check(FragmentType); return CompositionDescriptor.Tags.Contains(*FragmentType); }
+	bool HasTagType(const UScriptStruct* FragmentType) const { check(FragmentType); return CompositionDescriptor.GetTags().Contains(*FragmentType); }
 
 	bool IsEquivalent(const FMassArchetypeCompositionDescriptor& OtherCompositionDescriptor, const UE::Mass::FArchetypeGroups& OtherGroups) const;
 
diff --git a/Engine/Source/Runtime/MassEntity/Private/MassArchetypeData.cpp b/Engine/Source/Runtime/MassEntity/Private/MassArchetypeData.cpp
--- a/Engine/Source/Runtime/MassEntity/Private/MassArchetypeData.cpp
+++ b/Engine/Source/Runtime/MassEntity/Private/MassArchetypeData.cpp
@@ -54,7 +54,7 @@ void FMassArchetypeData::ForEachFragmentType(TFunction< void(const UScriptStruct
 
 bool FMassArchetypeData::HasFragmentType(const UScriptStruct* FragmentType) const
 {
-	return (FragmentType && CompositionDescriptor.Fragments.Contains(*FragmentType));
+	return (FragmentType && CompositionDescriptor.GetFragments().Contains(*FragmentType));
 }
 
 void FMassArchetypeData::Initialize(const FMassEntityManager& EntityManager, const FMassArchetypeCompositionDescriptor& InCompositionDescriptor, const uint32 ArchetypeDataVersion)
@@ -69,16 +69,16 @@ void FMassArchetypeData::Initialize(const FMassEntityManager& EntityManager, con
 	}
 
 	CreatedArchetypeDataVersion = ArchetypeDataVersion;
-	CompositionDescriptor.Fragments = InCompositionDescriptor.Fragments;
+	CompositionDescriptor.SetFragments(InCompositionDescriptor.GetFragments());
 	ConfigureFragments(EntityManager);
 
 	// Tags
-	CompositionDescriptor.Tags = InCompositionDescriptor.Tags;
+	CompositionDescriptor.SetTags(InCompositionDescriptor.GetTags());
 
 	// Chunk fragments
-	CompositionDescriptor.ChunkFragments = InCompositionDescriptor.ChunkFragments;
+	CompositionDescriptor.SetChunkFragments(InCompositionDescriptor.GetChunkFragments());
 	TArray<const UScriptStruct*, TInlineAllocator<16>> ChunkFragmentList;
-	CompositionDescriptor.ChunkFragments.ExportTypes(ChunkFragmentList);
+	CompositionDescriptor.GetChunkFragments().ExportTypes(ChunkFragmentList);
 	ChunkFragmentList.Sort(FScriptStructSortOperator());
 	for (const UScriptStruct* ChunkFragmentType : ChunkFragmentList)
 	{
@@ -87,8 +87,8 @@ void FMassArchetypeData::Initialize(const FMassEntityManager& EntityManager, con
 	}
 
 	// Share fragments
-	CompositionDescriptor.SharedFragments = InCompositionDescriptor.SharedFragments;
-	CompositionDescriptor.ConstSharedFragments = InCompositionDescriptor.ConstSharedFragments;
+	CompositionDescriptor.SetSharedFragments(InCompositionDescriptor.GetSharedFragments());
+	CompositionDescriptor.SetConstSharedFragments(InCompositionDescriptor.GetConstSharedFragments());
 
 	EntityListOffsetWithinChunk = 0;
 
@@ -106,7 +106,7 @@ void FMassArchetypeData::InitializeWithSimilar(const FMassEntityManager& EntityM
 
 	// note that we're calling this function rarely, so we can be a little bit inefficient here.
 	CompositionDescriptor = MoveTemp(NewComposition);
-	if (CompositionDescriptor.Fragments != BaseArchetype.GetCompositionDescriptor().Fragments)
+	if (CompositionDescriptor.GetFragments() != BaseArchetype.GetCompositionDescriptor().GetFragments())
 	{
 		ConfigureFragments(EntityManager);
 	}
@@ -131,7 +131,7 @@ void FMassArchetypeData::InitializeWithSimilar(const FMassEntityManager& EntityM
 void FMassArchetypeData::ConfigureFragments(const FMassEntityManager& EntityManager)
 {
 	TArray<const UScriptStruct*, TInlineAllocator<16>> SortedFragmentList;
-	CompositionDescriptor.Fragments.ExportTypes(SortedFragmentList);
+	CompositionDescriptor.GetFragments().ExportTypes(SortedFragmentList);
 
 	SortedFragmentList.Sort(FScriptStructSortOperator());
 
@@ -196,7 +196,7 @@ void FMassArchetypeData::AddEntity(FMassEntityHandle Entity, const FMassArchetyp
 int32 FMassArchetypeData::AddEntityInternal(FMassEntityHandle Entity, const FMassArchetypeSharedFragmentValues& SharedFragmentValues)
 {
 	checkf(SharedFragmentValues.IsSorted(), TEXT("Expecting shared fragment values to be previously sorted"));
-	checkf(SharedFragmentValues.HasExactFragmentTypesMatch(CompositionDescriptor.SharedFragments, CompositionDescriptor.ConstSharedFragments)
+	checkf(SharedFragmentValues.HasExactFragmentTypesMatch(CompositionDescriptor.GetSharedFragments(), CompositionDescriptor.GetConstSharedFragments())
 		, TEXT("Expecting values for every specified shared fragment in the archetype and only those"))
 
 	int32 IndexWithinChunk = 0;
@@ -370,7 +370,7 @@ void FMassArchetypeData::BatchDestroyEntityChunks(FMassArchetypeEntityCollection
 
 bool FMassArchetypeData::HasFragmentDataForEntity(const UScriptStruct* FragmentType, int32 EntityIndex) const
 {
-	return (FragmentType && CompositionDescriptor.Fragments.Contains(*FragmentType));
+	return (FragmentType && CompositionDescriptor.GetFragments().Contains(*FragmentType));
 }
 
 void* FMassArchetypeData::GetFragmentDataForEntityChecked(const UScriptStruct* FragmentType, int32 EntityIndex) const
@@ -954,11 +954,11 @@ FString FMassArchetypeData::DebugGetDescription() const
 		OutDescription += TEXT("\n");
 	}
 	OutDescription += TEXT("Chunk fragments: ");
-	CompositionDescriptor.ChunkFragments.DebugGetStringDesc(OutDescription);
+	CompositionDescriptor.GetChunkFragments().DebugGetStringDesc(OutDescription);
 	OutDescription += TEXT("\nTags: ");
-	CompositionDescriptor.Tags.DebugGetStringDesc(OutDescription);
+	CompositionDescriptor.GetTags().DebugGetStringDesc(OutDescription);
 	OutDescription += TEXT("\nFragments: ");
-	CompositionDescriptor.Fragments.DebugGetStringDesc(OutDescription);
+	CompositionDescriptor.GetFragments().DebugGetStringDesc(OutDescription);
 	
 	return static_cast<FString>(OutDescription);
 #else
@@ -992,7 +992,7 @@ void FMassArchetypeData::DebugPrintArchetype(FOutputDevice& Ar)
 	Ar.Logf(ELogVerbosity::Log, TEXT("Name: %s"), *GetCombinedDebugNamesAsString());
 
 	FStringOutputDevice TagsDescription;
-	CompositionDescriptor.Tags.DebugGetStringDesc(TagsDescription);
+	CompositionDescriptor.GetTags().DebugGetStringDesc(TagsDescription);
 	Ar.Logf(ELogVerbosity::Log, TEXT("Tags: %s"), *TagsDescription);
 	Ar.Logf(ELogVerbosity::Log, TEXT("Fragments: %s"), *DebugGetDescription());
 	Ar.Logf(ELogVerbosity::Log, TEXT("\tChunks: %d x %llu KB = %llu KB total"), Chunks.Num(), GetChunkAllocSize() / 1024, (GetChunkAllocSize()*Chunks.Num()) / 1024);
@@ -1115,8 +1115,8 @@ void FMassArchetypeData::BatchAddEntities(TConstArrayView<FMassEntityHandle> Ent
 {
 	SCOPE_CYCLE_COUNTER(STAT_Mass_ArchetypeBatchAdd);
 
-	testableCheckfReturn(SharedFragmentValues.HasExactSharedFragmentTypesMatch(GetCompositionDescriptor().SharedFragments), return, TEXT("%hs parameter SharedFragmentValues doesn't match archetype's composition"), __FUNCTION__);
-	testableCheckfReturn(SharedFragmentValues.HasExactConstSharedFragmentTypesMatch(GetCompositionDescriptor().ConstSharedFragments), return, TEXT("%hs parameter ConstSharedFragmentValues doesn't match archetype's composition"), __FUNCTION__);
+	testableCheckfReturn(SharedFragmentValues.HasExactSharedFragmentTypesMatch(GetCompositionDescriptor().GetSharedFragments()), return, TEXT("%hs parameter SharedFragmentValues doesn't match archetype's composition"), __FUNCTION__);
+	testableCheckfReturn(SharedFragmentValues.HasExactConstSharedFragmentTypesMatch(GetCompositionDescriptor().GetConstSharedFragments()), return, TEXT("%hs parameter ConstSharedFragmentValues doesn't match archetype's composition"), __FUNCTION__);
 
 	FMassArchetypeEntityCollection::FArchetypeEntityRange ResultSubchunk;
 	ResultSubchunk.ChunkIndex = 0;
@@ -1155,7 +1155,7 @@ void FMassArchetypeData::BatchMoveEntitiesToAnotherArchetype(const FMassArchetyp
 			FMassSharedFragmentBitSet NewSharedFragmentsBitset = GetSharedFragmentBitSet();
 			NewSharedFragmentsBitset -= *SharedFragmentToRemoveBitSet;
 			NewSharedFragmentsBitset += SharedFragmentValuesToAdd->GetSharedFragmentBitSet();
-			bIsValidArchetype = NewArchetype.GetCompositionDescriptor().SharedFragments == NewSharedFragmentsBitset;
+			bIsValidArchetype = NewArchetype.GetCompositionDescriptor().GetSharedFragments() == NewSharedFragmentsBitset;
 		}
 
 		if (bIsValidArchetype && ConstSharedFragmentToRemoveBitSet)
@@ -1163,7 +1163,7 @@ void FMassArchetypeData::BatchMoveEntitiesToAnotherArchetype(const FMassArchetyp
 			FMassConstSharedFragmentBitSet NewConstSharedFragmentsBitset = GetConstSharedFragmentBitSet();
 			NewConstSharedFragmentsBitset -= *ConstSharedFragmentToRemoveBitSet;
 			NewConstSharedFragmentsBitset += SharedFragmentValuesToAdd->GetConstSharedFragmentBitSet();
-			bIsValidArchetype = NewArchetype.GetCompositionDescriptor().ConstSharedFragments == NewConstSharedFragmentsBitset;
+			bIsValidArchetype = NewArchetype.GetCompositionDescriptor().GetConstSharedFragments() == NewConstSharedFragmentsBitset;
 		}
 
 		testableCheckfReturn(bIsValidArchetype, return, TEXT("%hs parameter SharedFragmentValues doesn't match archetype's composition"), __FUNCTION__);
@@ -1273,7 +1273,7 @@ void FMassArchetypeData::BatchMoveEntitiesToAnotherArchetype(const FMassArchetyp
 FMassArchetypeEntityCollection::FArchetypeEntityRange FMassArchetypeData::PrepareNextEntitiesSpanInternal(TConstArrayView<FMassEntityHandle> Entities, const FMassArchetypeSharedFragmentValues& SharedFragmentValues, int32 StartingChunk)
 {
 	checkf(SharedFragmentValues.IsSorted(), TEXT("Expecting shared fragment values to be previously sorted"));
-	checkf(SharedFragmentValues.HasExactFragmentTypesMatch(CompositionDescriptor.SharedFragments, CompositionDescriptor.ConstSharedFragments)
+	checkf(SharedFragmentValues.HasExactFragmentTypesMatch(CompositionDescriptor.GetSharedFragments(), CompositionDescriptor.GetConstSharedFragments())
 		, TEXT("Expecting values for every specified shared fragment in the archetype and only those"))
 
 	int32 StartIndexWithinChunk = INDEX_NONE;
diff --git a/Engine/Source/Runtime/MassEntity/Private/MassDebugger.cpp b/Engine/Source/Runtime/MassEntity/Private/MassDebugger.cpp
--- a/Engine/Source/Runtime/MassEntity/Private/MassDebugger.cpp
+++ b/Engine/Source/Runtime/MassEntity/Private/MassDebugger.cpp
@@ -675,92 +675,92 @@ FString FMassDebugger::GetArchetypeRequirementCompatibilityDescription(const FMa
 
 	if (Requirements.HasNegativeRequirements())
 	{
-		if (ArchetypeComposition.Fragments.HasNone(Requirements.RequiredNoneFragments) == false)
+		if (ArchetypeComposition.GetFragments().HasNone(Requirements.RequiredNoneFragments) == false)
 		{
 			// has some of the fragments required absent
 			OutDescription += TEXT("\nHas fragments required absent: ");
-			(Requirements.RequiredNoneFragments & ArchetypeComposition.Fragments).DebugGetStringDesc(OutDescription);
+			(Requirements.RequiredNoneFragments & ArchetypeComposition.GetFragments()).DebugGetStringDesc(OutDescription);
 		}
 
-		if (ArchetypeComposition.Tags.HasNone(Requirements.RequiredNoneTags) == false)
+		if (ArchetypeComposition.GetTags().HasNone(Requirements.RequiredNoneTags) == false)
 		{
 			// has some of the tags required absent
 			OutDescription += TEXT("\nHas tags required absent: ");
-			(Requirements.RequiredNoneTags & ArchetypeComposition.Tags).DebugGetStringDesc(OutDescription);
+			(Requirements.RequiredNoneTags & ArchetypeComposition.GetTags()).DebugGetStringDesc(OutDescription);
 		}
 
-		if (ArchetypeComposition.ChunkFragments.HasNone(Requirements.RequiredNoneChunkFragments) == false)
+		if (ArchetypeComposition.GetChunkFragments().HasNone(Requirements.RequiredNoneChunkFragments) == false)
 		{
 			// has some of the chunk fragments required absent
 			OutDescription += TEXT("\nHas chunk fragments required absent: ");
-			(Requirements.RequiredNoneChunkFragments & ArchetypeComposition.ChunkFragments).DebugGetStringDesc(OutDescription);
+			(Requirements.RequiredNoneChunkFragments & ArchetypeComposition.GetChunkFragments()).DebugGetStringDesc(OutDescription);
 		}
 
-		if (ArchetypeComposition.SharedFragments.HasNone(Requirements.RequiredNoneSharedFragments) == false)
+		if (ArchetypeComposition.GetSharedFragments().HasNone(Requirements.RequiredNoneSharedFragments) == false)
 		{
 			// has some of the chunk fragments required absent
 			OutDescription += TEXT("\nHas shared fragments required absent: ");
-			(Requirements.RequiredNoneSharedFragments & ArchetypeComposition.SharedFragments).DebugGetStringDesc(OutDescription);
+			(Requirements.RequiredNoneSharedFragments & ArchetypeComposition.GetSharedFragments()).DebugGetStringDesc(OutDescription);
 		}
 
-		if (ArchetypeComposition.ConstSharedFragments.HasNone(Requirements.RequiredNoneConstSharedFragments) == false)
+		if (ArchetypeComposition.GetConstSharedFragments().HasNone(Requirements.RequiredNoneConstSharedFragments) == false)
 		{
 			// has some of the chunk fragments required absent
 			OutDescription += TEXT("\nHas shared fragments required absent: ");
-			(Requirements.RequiredNoneConstSharedFragments & ArchetypeComposition.ConstSharedFragments).DebugGetStringDesc(OutDescription);
+			(Requirements.RequiredNoneConstSharedFragments & ArchetypeComposition.GetConstSharedFragments()).DebugGetStringDesc(OutDescription);
 		}
 	}
 
 	// if we have regular (i.e. non-optional) positive requirements then these are the determining factor, we don't check optionals
 	if (Requirements.HasPositiveRequirements())
 	{
-		if (ArchetypeComposition.Fragments.HasAll(Requirements.RequiredAllFragments) == false)
+		if (ArchetypeComposition.GetFragments().HasAll(Requirements.RequiredAllFragments) == false)
 		{
 			// missing one of the strictly required fragments
 			OutDescription += TEXT("\nMissing required fragments: ");
-			(Requirements.RequiredAllFragments - ArchetypeComposition.Fragments).DebugGetStringDesc(OutDescription);
+			(Requirements.RequiredAllFragments - ArchetypeComposition.GetFragments()).DebugGetStringDesc(OutDescription);
 		}
 
-		if (Requirements.RequiredAnyFragments.IsEmpty() == false && ArchetypeComposition.Fragments.HasAny(Requirements.RequiredAnyFragments) == false)
+		if (Requirements.RequiredAnyFragments.IsEmpty() == false && ArchetypeComposition.GetFragments().HasAny(Requirements.RequiredAnyFragments) == false)
 		{
 			// missing all of the "any" fragments
 			OutDescription += TEXT("\nMissing all \'any\' fragments: ");
 			Requirements.RequiredAnyFragments.DebugGetStringDesc(OutDescription);
 		}
 
-		if (ArchetypeComposition.Tags.HasAll(Requirements.RequiredAllTags) == false)
+		if (ArchetypeComposition.GetTags().HasAll(Requirements.RequiredAllTags) == false)
 		{
 			// missing one of the strictly required tags
 			OutDescription += TEXT("\nMissing required tags: ");
-			(Requirements.RequiredAllTags - ArchetypeComposition.Tags).DebugGetStringDesc(OutDescription);
+			(Requirements.RequiredAllTags - ArchetypeComposition.GetTags()).DebugGetStringDesc(OutDescription);
 		}
 
-		if (Requirements.RequiredAnyTags.IsEmpty() == false && ArchetypeComposition.Tags.HasAny(Requirements.RequiredAnyTags) == false)
+		if (Requirements.RequiredAnyTags.IsEmpty() == false && ArchetypeComposition.GetTags().HasAny(Requirements.RequiredAnyTags) == false)
 		{
 			// missing all of the "any" tags
 			OutDescription += TEXT("\nMissing all \'any\' tags: ");
 			Requirements.RequiredAnyTags.DebugGetStringDesc(OutDescription);
 		}
 
-		if (ArchetypeComposition.ChunkFragments.HasAll(Requirements.RequiredAllChunkFragments) == false)
+		if (ArchetypeComposition.GetChunkFragments().HasAll(Requirements.RequiredAllChunkFragments) == false)
 		{
 			// missing one of the strictly required chunk fragments
 			OutDescription += TEXT("\nMissing required chunk fragments: ");
-			(Requirements.RequiredAllChunkFragments - ArchetypeComposition.ChunkFragments).DebugGetStringDesc(OutDescription);
+			(Requirements.RequiredAllChunkFragments - ArchetypeComposition.GetChunkFragments()).DebugGetStringDesc(OutDescription);
 		}
 
-		if (ArchetypeComposition.SharedFragments.HasAll(Requirements.RequiredAllSharedFragments) == false)
+		if (ArchetypeComposition.GetSharedFragments().HasAll(Requirements.RequiredAllSharedFragments) == false)
 		{
 			// missing one of the strictly required Shared fragments
 			OutDescription += TEXT("\nMissing required Shared fragments: ");
-			(Requirements.RequiredAllSharedFragments - ArchetypeComposition.SharedFragments).DebugGetStringDesc(OutDescription);
+			(Requirements.RequiredAllSharedFragments - ArchetypeComposition.GetSharedFragments()).DebugGetStringDesc(OutDescription);
 		}
 
-		if (ArchetypeComposition.ConstSharedFragments.HasAll(Requirements.RequiredAllConstSharedFragments) == false)
+		if (ArchetypeComposition.GetConstSharedFragments().HasAll(Requirements.RequiredAllConstSharedFragments) == false)
 		{
 			// missing one of the strictly required Shared fragments
 			OutDescription += TEXT("\nMissing required Shared fragments: ");
-			(Requirements.RequiredAllConstSharedFragments - ArchetypeComposition.ConstSharedFragments).DebugGetStringDesc(OutDescription);
+			(Requirements.RequiredAllConstSharedFragments - ArchetypeComposition.GetConstSharedFragments()).DebugGetStringDesc(OutDescription);
 		}
 	}
 	// else we check if there are any optionals and if so test them
@@ -1199,43 +1199,43 @@ const UScriptStruct* FMassDebugger::GetFragmentTypeFromName(FName FragmentName)
 		TArray<FMassArchetypeHandle> ArchetypeHandles = GetAllArchetypes(*Environment.EntityManager.Pin());
 		for (FMassArchetypeHandle& ArchetypeHandle : ArchetypeHandles)
 		{
-			FMassArchetypeCompositionDescriptor Composition = GetArchetypeComposition(ArchetypeHandle);
+			const FMassArchetypeCompositionDescriptor& Composition = GetArchetypeComposition(ArchetypeHandle);
 
-			FMassFragmentBitSet::FIndexIterator It = Composition.Fragments.GetIndexIterator();
+			FMassFragmentBitSet::FIndexIterator It = Composition.GetFragments().GetIndexIterator();
 			while (It)
 			{
-				FName StructName = Composition.Fragments.DebugGetStructTypeName(*It);
-				const UScriptStruct* StructType = Composition.Fragments.GetTypeAtIndex(*It);
+				FName StructName = Composition.GetFragments().DebugGetStructTypeName(*It);
+				const UScriptStruct* StructType = Composition.GetFragments().GetTypeAtIndex(*It);
 				
 				FragmentsByName.Add(StructName, StructType);
 				++It;
 			}
 
-			FMassChunkFragmentBitSet::FIndexIterator ChunkIt = Composition.ChunkFragments.GetIndexIterator();
+			FMassChunkFragmentBitSet::FIndexIterator ChunkIt = Composition.GetChunkFragments().GetIndexIterator();
 			while (ChunkIt)
 			{
-				FName StructName = Composition.ChunkFragments.DebugGetStructTypeName(*ChunkIt);
-				const UScriptStruct* StructType = Composition.ChunkFragments.GetTypeAtIndex(*ChunkIt);
+				FName StructName = Composition.GetChunkFragments().DebugGetStructTypeName(*ChunkIt);
+				const UScriptStruct* StructType = Composition.GetChunkFragments().GetTypeAtIndex(*ChunkIt);
 
 				FragmentsByName.Add(StructName, StructType);
 				++ChunkIt;
 			}
 
-			FMassSharedFragmentBitSet::FIndexIterator SharedFragIt = Composition.SharedFragments.GetIndexIterator();
+			FMassSharedFragmentBitSet::FIndexIterator SharedFragIt = Composition.GetSharedFragments().GetIndexIterator();
 			while (SharedFragIt)
 			{
-				FName StructName = Composition.SharedFragments.DebugGetStructTypeName(*SharedFragIt);
-				const UScriptStruct* StructType = Composition.SharedFragments.GetTypeAtIndex(*SharedFragIt);
+				FName StructName = Composition.GetSharedFragments().DebugGetStructTypeName(*SharedFragIt);
+				const UScriptStruct* StructType = Composition.GetSharedFragments().GetTypeAtIndex(*SharedFragIt);
 
 				FragmentsByName.Add(StructName, StructType);
 				++SharedFragIt;
 			}
 
-			FMassConstSharedFragmentBitSet::FIndexIterator ConstSharedFragIt = Composition.ConstSharedFragments.GetIndexIterator();
+			FMassConstSharedFragmentBitSet::FIndexIterator ConstSharedFragIt = Composition.GetConstSharedFragments().GetIndexIterator();
 			while (ConstSharedFragIt)
 			{
-				FName StructName = Composition.ConstSharedFragments.DebugGetStructTypeName(*ConstSharedFragIt);
-				const UScriptStruct* StructType = Composition.ConstSharedFragments.GetTypeAtIndex(*ConstSharedFragIt);
+				FName StructName = Composition.GetConstSharedFragments().DebugGetStructTypeName(*ConstSharedFragIt);
+				const UScriptStruct* StructType = Composition.GetConstSharedFragments().GetTypeAtIndex(*ConstSharedFragIt);
 
 				FragmentsByName.Add(StructName, StructType);
 				++ConstSharedFragIt;
diff --git a/Engine/Source/Runtime/MassEntity/Private/MassEntityBuilder.cpp b/Engine/Source/Runtime/MassEntity/Private/MassEntityBuilder.cpp
--- a/Engine/Source/Runtime/MassEntity/Private/MassEntityBuilder.cpp
+++ b/Engine/Source/Runtime/MassEntity/Private/MassEntityBuilder.cpp
@@ -106,9 +106,9 @@ namespace UE::Mass
 			, TConstArrayView<FInstancedStruct> InitialFragmentValues, TConstArrayView<FConstSharedStruct> ConstSharedFragments, TConstArrayView<FSharedStruct> SharedFragments)
 		{
 			const UObject* LogOwner = InEntityManager->GetOwner();
-			bool bIssuesFound = CheckStructContainer<FMassFragment>(InitialFragmentValues, Composition.Fragments, LogOwner);
-			bIssuesFound |= CheckStructContainer<FMassConstSharedFragment>(ConstSharedFragments, Composition.ConstSharedFragments, LogOwner);
-			bIssuesFound |= CheckStructContainer<FMassSharedFragment>(SharedFragments, Composition.SharedFragments, LogOwner);
+			bool bIssuesFound = CheckStructContainer<FMassFragment>(InitialFragmentValues, Composition.GetFragments(), LogOwner);
+			bIssuesFound |= CheckStructContainer<FMassConstSharedFragment>(ConstSharedFragments, Composition.GetConstSharedFragments(), LogOwner);
+			bIssuesFound |= CheckStructContainer<FMassSharedFragment>(SharedFragments, Composition.GetSharedFragments(), LogOwner);
 			
 			return !bIssuesFound;
 		}
@@ -571,24 +571,24 @@ FEntityBuilder& FEntityBuilder::Add(TNotNull<const UScriptStruct*> ElementType)
 {
 	if (UE::Mass::IsA<FMassTag>(ElementType))
 	{
-		Composition.Tags.Add(*ElementType);
+		Composition.GetTags().Add(*ElementType);
 	}
 	else if (UE::Mass::IsA<FMassFragment>(ElementType))
 	{
-		Composition.Fragments.Add(*ElementType);
+		Composition.GetFragments().Add(*ElementType);
 	}
 	else if (UE::Mass::IsA<FMassSharedFragment>(ElementType))
 	{
-		Composition.SharedFragments.Add(*ElementType);
+		Composition.GetSharedFragments().Add(*ElementType);
 	}
 	else if (UE::Mass::IsA<FMassConstSharedFragment>(ElementType))
 	{
-		Composition.ConstSharedFragments.Add(*ElementType);
+		Composition.GetConstSharedFragments().Add(*ElementType);
 	}
 	else if (ensureMsgf(UE::Mass::IsA<FMassChunkFragment>(ElementType)
 		, TEXT("Unhandled element type %s"), *ElementType->GetName()))
 	{
-		Composition.ChunkFragments.Add(*ElementType);
+		Composition.GetChunkFragments().Add(*ElementType);
 	}
 
 	State = EState::ReadyToCommit;
diff --git a/Engine/Source/Runtime/MassEntity/Private/MassEntityManager.cpp b/Engine/Source/Runtime/MassEntity/Private/MassEntityManager.cpp
--- a/Engine/Source/Runtime/MassEntity/Private/MassEntityManager.cpp
+++ b/Engine/Source/Runtime/MassEntity/Private/MassEntityManager.cpp
@@ -364,8 +364,8 @@ FMassArchetypeHandle FMassEntityManager::GetOrCreateSuitableArchetype(const FMas
 		|| ConstSharedFragmentBitSet != ArchetypeData.GetConstSharedFragmentBitSet())
 	{
 		FMassArchetypeCompositionDescriptor NewDescriptor = ArchetypeData.GetCompositionDescriptor();
-		NewDescriptor.SharedFragments = SharedFragmentBitSet;
-		NewDescriptor.ConstSharedFragments = ConstSharedFragmentBitSet;
+		NewDescriptor.SetSharedFragments(SharedFragmentBitSet);
+		NewDescriptor.SetConstSharedFragments(ConstSharedFragmentBitSet);
 		return CreateArchetype(NewDescriptor, CreationParams);
 	}
 	return ArchetypeHandle;
@@ -430,7 +430,7 @@ FMassArchetypeHandle FMassEntityManager::InternalCreateSimilarArchetype(const TS
 	checkSlow(SourceArchetype.IsValid());
 	const FMassArchetypeData& SourceArchetypeRef = *SourceArchetype.Get();
 	FMassArchetypeCompositionDescriptor NewComposition(SourceArchetypeRef.GetCompositionDescriptor());
-	NewComposition.Tags = OverrideTags;
+	NewComposition.SetTags(OverrideTags);
 
 	return InternalCreateSimilarArchetype(SourceArchetypeRef, MoveTemp(NewComposition), SourceArchetypeRef.GetGroups());
 }
@@ -440,7 +440,7 @@ FMassArchetypeHandle FMassEntityManager::InternalCreateSimilarArchetype(const TS
 	checkSlow(SourceArchetype.IsValid());
 	const FMassArchetypeData& SourceArchetypeRef = *SourceArchetype.Get();
 	FMassArchetypeCompositionDescriptor NewComposition(SourceArchetypeRef.GetCompositionDescriptor());
-	NewComposition.Fragments = OverrideFragments;
+	NewComposition.SetFragments(OverrideFragments);
 
 	return InternalCreateSimilarArchetype(SourceArchetypeRef, MoveTemp(NewComposition), SourceArchetypeRef.GetGroups());
 }
@@ -508,15 +508,15 @@ void FMassEntityManager::InternalAppendFragmentsAndTagsToArchetypeCompositionDes
 	{
 		if (UE::Mass::IsA<FMassFragment>(Type))
 		{
-			InOutComposition.Fragments.Add(*Type);
+			InOutComposition.GetFragments().Add(*Type);
 		}
 		else if (UE::Mass::IsA<FMassTag>(Type))
 		{
-			InOutComposition.Tags.Add(*Type);
+			InOutComposition.GetTags().Add(*Type);
 		}
 		else if (UE::Mass::IsA<FMassChunkFragment>(Type))
 		{
-			InOutComposition.ChunkFragments.Add(*Type);
+			InOutComposition.GetChunkFragments().Add(*Type);
 		}
 		else
 		{
@@ -648,11 +648,11 @@ void FMassEntityManager::BuildEntity(FMassEntityHandle Entity, TConstArrayView<F
 	FMassArchetypeCompositionDescriptor Composition(FragmentInstanceList, FMassTagBitSet(), FMassChunkFragmentBitSet(), FMassSharedFragmentBitSet(), FMassConstSharedFragmentBitSet());
 	for (const FConstSharedStruct& SharedFragment : SharedFragmentValues.GetConstSharedFragments())
 	{
-		Composition.ConstSharedFragments.Add(*SharedFragment.GetScriptStruct());
+		Composition.GetConstSharedFragments().Add(*SharedFragment.GetScriptStruct());
 	}
 	for (const FSharedStruct& SharedFragment : SharedFragmentValues.GetSharedFragments())
 	{
-		Composition.SharedFragments.Add(*SharedFragment.GetScriptStruct());
+		Composition.GetSharedFragments().Add(*SharedFragment.GetScriptStruct());
 	}
 
 	const FMassArchetypeHandle& ArchetypeHandle = CreateArchetype(Composition);
@@ -691,11 +691,11 @@ TSharedRef<FMassEntityManager::FEntityCreationContext> FMassEntityManager::Batch
 	FMassArchetypeCompositionDescriptor Composition(FragmentsAffected, FMassTagBitSet(), FMassChunkFragmentBitSet(), FMassSharedFragmentBitSet(), FMassConstSharedFragmentBitSet());
 	for (const FConstSharedStruct& SharedFragment : SharedFragmentValues.GetConstSharedFragments())
 	{
-		Composition.ConstSharedFragments.Add(*SharedFragment.GetScriptStruct());
+		Composition.GetConstSharedFragments().Add(*SharedFragment.GetScriptStruct());
 	}
 	for (const FSharedStruct& SharedFragment : SharedFragmentValues.GetSharedFragments())
 	{
-		Composition.SharedFragments.Add(*SharedFragment.GetScriptStruct());
+		Composition.GetSharedFragments().Add(*SharedFragment.GetScriptStruct());
 	}
 
 	return BatchBuildEntities(EncodedEntitiesWithPayload, MoveTemp(Composition), SharedFragmentValues, CreationParams);
@@ -1052,8 +1052,8 @@ void FMassEntityManager::AddCompositionToEntity_GetDelta(FMassEntityHandle Entit
 
 	InOutDescriptor.Remove(OldArchetype->GetCompositionDescriptor());
 
-	ensureMsgf(InOutDescriptor.ChunkFragments.IsEmpty(), TEXT("Adding new chunk fragments is not supported"));
-	ensureMsgf(InOutDescriptor.SharedFragments.IsEmpty() 
+	ensureMsgf(InOutDescriptor.GetChunkFragments().IsEmpty(), TEXT("Adding new chunk fragments is not supported"));
+	ensureMsgf(InOutDescriptor.GetSharedFragments().IsEmpty() 
 		|| (AddedSharedFragmentValues && AddedSharedFragmentValues->DoesMatchComposition(InOutDescriptor))
 		, TEXT("When adding new shared fragments it's required to provide values for said fragments"));
 
@@ -1103,7 +1103,7 @@ void FMassEntityManager::RemoveCompositionFromEntity(FMassEntityHandle Entity, c
 		FMassArchetypeCompositionDescriptor NewDescriptor = OldArchetype->GetCompositionDescriptor();
 		NewDescriptor.Remove(InDescriptor);
 
-		ensureMsgf(InDescriptor.ChunkFragments.IsEmpty(), TEXT("Removing chunk fragments is not supported"));
+		ensureMsgf(InDescriptor.GetChunkFragments().IsEmpty(), TEXT("Removing chunk fragments is not supported"));
 
 		if (NewDescriptor.IsEquivalent(OldArchetype->GetCompositionDescriptor()) == false)
 		{
@@ -1116,7 +1116,7 @@ void FMassEntityManager::RemoveCompositionFromEntity(FMassEntityHandle Entity, c
 				// Move the entity over
 				FMassArchetypeData& NewArchetype = FMassArchetypeHelper::ArchetypeDataFromHandleChecked(NewArchetypeHandle);
 				NewArchetype.CopyDebugNamesFrom(*OldArchetype);
-				if (InDescriptor.SharedFragments.IsEmpty() && InDescriptor.ConstSharedFragments.IsEmpty())
+				if (InDescriptor.GetSharedFragments().IsEmpty() && InDescriptor.GetConstSharedFragments().IsEmpty())
 				{
 					OldArchetype->MoveEntityToAnotherArchetype(Entity, NewArchetype);
 				}
@@ -1335,10 +1335,10 @@ bool FMassEntityManager::DoesEntityHaveElement(FMassEntityHandle Entity, TNotNul
 	check(CurrentArchetype);
 	if (UE::Mass::IsA<FMassTag>(ElementType))
 	{
-		return CurrentArchetype->GetCompositionDescriptor().Tags.Contains(*ElementType);
+		return CurrentArchetype->GetCompositionDescriptor().GetTags().Contains(*ElementType);
 	}
 	
-	return CurrentArchetype->GetCompositionDescriptor().Fragments.Contains(*ElementType);
+	return CurrentArchetype->GetCompositionDescriptor().GetFragments().Contains(*ElementType);
 }
 
 void FMassEntityManager::AddElementToEntities(TConstArrayView<FMassEntityHandle> Entities, TNotNull<const UScriptStruct*> ElementType)
@@ -1419,7 +1419,7 @@ bool FMassEntityManager::AddConstSharedFragmentToEntity(const FMassEntityHandle
 
 	const UScriptStruct* StructType = InConstSharedFragment.GetScriptStruct();
 	CA_ASSUME(StructType);
-	if (CurrentArchetype->GetCompositionDescriptor().ConstSharedFragments.Contains(*StructType))
+	if (CurrentArchetype->GetCompositionDescriptor().GetConstSharedFragments().Contains(*StructType))
 	{
 		const FMassArchetypeSharedFragmentValues& SharedFragmentValues = CurrentArchetype->GetSharedFragmentValues(Entity);
 		FConstSharedStruct ExistingConstSharedStruct = SharedFragmentValues.GetConstSharedFragmentStruct(StructType);
@@ -1433,7 +1433,7 @@ bool FMassEntityManager::AddConstSharedFragmentToEntity(const FMassEntityHandle
 	}
 	
 	FMassArchetypeCompositionDescriptor NewComposition(CurrentArchetype->GetCompositionDescriptor());
-	NewComposition.ConstSharedFragments.Add(*StructType);
+	NewComposition.GetConstSharedFragments().Add(*StructType);
 	const FMassArchetypeHandle NewArchetypeHandle = CreateArchetype(NewComposition, FMassArchetypeCreationParams(*CurrentArchetype));
 	check(NewArchetypeHandle.IsValid());
 	FMassArchetypeData* NewArchetype = NewArchetypeHandle.DataPtr.Get();
@@ -1467,14 +1467,14 @@ bool FMassEntityManager::RemoveConstSharedFragmentFromEntity(const FMassEntityHa
 	FMassArchetypeData* CurrentArchetype = GetEntityStorageInterface().GetArchetypeAsShared(Entity.Index).Get();
 	check(CurrentArchetype);
 	
-	if (!CurrentArchetype->GetCompositionDescriptor().ConstSharedFragments.Contains(ConstSharedFragmentType))
+	if (!CurrentArchetype->GetCompositionDescriptor().GetConstSharedFragments().Contains(ConstSharedFragmentType))
 	{
 		// Nothing to do. Returning false to indicate nothing has been removed, as per function's documentation 
 		return false;
 	}
 
 	FMassArchetypeCompositionDescriptor NewComposition(CurrentArchetype->GetCompositionDescriptor());
-	NewComposition.ConstSharedFragments.Remove(ConstSharedFragmentType);
+	NewComposition.GetConstSharedFragments().Remove(ConstSharedFragmentType);
 	const FMassArchetypeHandle NewArchetypeHandle = CreateArchetype(NewComposition);
 	check(NewArchetypeHandle.IsValid());
 	FMassArchetypeData* NewArchetype = NewArchetypeHandle.DataPtr.Get();
@@ -1512,7 +1512,7 @@ bool FMassEntityManager::AddSharedFragmentToEntity(const FMassEntityHandle Entit
 
 	const UScriptStruct* StructType = InSharedFragment.GetScriptStruct();
 	CA_ASSUME(StructType);
-	if (CurrentArchetype->GetCompositionDescriptor().SharedFragments.Contains(*StructType))
+	if (CurrentArchetype->GetCompositionDescriptor().GetSharedFragments().Contains(*StructType))
 	{
 		const FMassArchetypeSharedFragmentValues& SharedFragmentValues = CurrentArchetype->GetSharedFragmentValues(Entity);
 		FConstSharedStruct ExistingSharedStruct = SharedFragmentValues.GetSharedFragmentStruct(StructType);
@@ -1526,7 +1526,7 @@ bool FMassEntityManager::AddSharedFragmentToEntity(const FMassEntityHandle Entit
 	}
 	
 	FMassArchetypeCompositionDescriptor NewComposition(CurrentArchetype->GetCompositionDescriptor());
-	NewComposition.SharedFragments.Add(*StructType);
+	NewComposition.GetSharedFragments().Add(*StructType);
 	const FMassArchetypeHandle NewArchetypeHandle = CreateArchetype(NewComposition, FMassArchetypeCreationParams(*CurrentArchetype));
 	check(NewArchetypeHandle.IsValid());
 	FMassArchetypeData* NewArchetype = NewArchetypeHandle.DataPtr.Get();
@@ -1560,14 +1560,14 @@ bool FMassEntityManager::RemoveSharedFragmentFromEntity(const FMassEntityHandle
 	FMassArchetypeData* CurrentArchetype = GetEntityStorageInterface().GetArchetypeAsShared(Entity.Index).Get();
 	check(CurrentArchetype);
 	
-	if (!CurrentArchetype->GetCompositionDescriptor().SharedFragments.Contains(SharedFragmentType))
+	if (!CurrentArchetype->GetCompositionDescriptor().GetSharedFragments().Contains(SharedFragmentType))
 	{
 		// Nothing to do. Returning false to indicate nothing has been removed, as per function's documentation 
 		return false;
 	}
 
 	FMassArchetypeCompositionDescriptor NewComposition(CurrentArchetype->GetCompositionDescriptor());
-	NewComposition.SharedFragments.Remove(SharedFragmentType);
+	NewComposition.GetSharedFragments().Remove(SharedFragmentType);
 	const FMassArchetypeHandle NewArchetypeHandle = CreateArchetype(NewComposition);
 	check(NewArchetypeHandle.IsValid());
 	FMassArchetypeData* NewArchetype = NewArchetypeHandle.DataPtr.Get();
@@ -1788,8 +1788,8 @@ void FMassEntityManager::BatchAddSharedFragmentsForEntities(TConstArrayView<FMas
 		testableCheckfReturn(CurrentArchetype, continue, TEXT("Adding shared fragments to archetype-less entities is not supported"));
 
 		FMassArchetypeCompositionDescriptor NewComposition(CurrentArchetype->GetCompositionDescriptor());
-		NewComposition.SharedFragments += AddedFragmentValues.GetSharedFragmentBitSet();
-		NewComposition.ConstSharedFragments += AddedFragmentValues.GetConstSharedFragmentBitSet();
+		NewComposition.GetSharedFragments() += AddedFragmentValues.GetSharedFragmentBitSet();
+		NewComposition.GetConstSharedFragments() += AddedFragmentValues.GetConstSharedFragmentBitSet();
 
 		const FMassArchetypeHandle NewArchetypeHandle = CreateArchetype(NewComposition, FMassArchetypeCreationParams(*CurrentArchetype));
 		check(NewArchetypeHandle.IsValid());
@@ -2179,7 +2179,7 @@ void FMassEntityManager::DebugGetArchetypeFragmentTypes(const FMassArchetypeHand
 	if (Archetype.IsValid())
 	{
 		const FMassArchetypeData& ArchetypeData = FMassArchetypeHelper::ArchetypeDataFromHandleChecked(Archetype);
-		ArchetypeData.GetCompositionDescriptor().Fragments.ExportTypes(InOutFragmentList);
+		ArchetypeData.GetCompositionDescriptor().GetFragments().ExportTypes(InOutFragmentList);
 	}
 }
 
diff --git a/Engine/Source/Runtime/MassEntity/Private/MassEntityTrace.cpp b/Engine/Source/Runtime/MassEntity/Private/MassEntityTrace.cpp
--- a/Engine/Source/Runtime/MassEntity/Private/MassEntityTrace.cpp
+++ b/Engine/Source/Runtime/MassEntity/Private/MassEntityTrace.cpp
@@ -239,12 +239,12 @@ void FMassTrace::OnPhaseEnd(uint64 PhaseId)
 uint64 FMassTrace::OutputRegisterArchetype(uint64 ArchetypeID, const FMassArchetypeCompositionDescriptor& CompositionDescriptor)
 {
 	TArray<uint64> FragmentsScratch;
-	FragmentsScratch.Reserve(CompositionDescriptor.Fragments.CountStoredTypes() + CompositionDescriptor.Tags.CountStoredTypes());
+	FragmentsScratch.Reserve(CompositionDescriptor.GetFragments().CountStoredTypes() + CompositionDescriptor.GetTags().CountStoredTypes());
 
-	auto FragmentIterator = CompositionDescriptor.Fragments.GetIndexIterator();
+	auto FragmentIterator = CompositionDescriptor.GetFragments().GetIndexIterator();
 	while (FragmentIterator)
 	{
-		const UScriptStruct* FragmentStruct = CompositionDescriptor.Fragments.GetTypeAtIndex(*FragmentIterator);
+		const UScriptStruct* FragmentStruct = CompositionDescriptor.GetFragments().GetTypeAtIndex(*FragmentIterator);
 
 		OutputRegisterFragment(FragmentStruct);
 
@@ -253,10 +253,10 @@ uint64 FMassTrace::OutputRegisterArchetype(uint64 ArchetypeID, const FMassArchet
 		++FragmentIterator;
 	}
 
-	auto TagIterator = CompositionDescriptor.Tags.GetIndexIterator();
+	auto TagIterator = CompositionDescriptor.GetTags().GetIndexIterator();
 	while (TagIterator)
 	{
-		const UScriptStruct* FragmentStruct = CompositionDescriptor.Tags.GetTypeAtIndex(*TagIterator);
+		const UScriptStruct* FragmentStruct = CompositionDescriptor.GetTags().GetTypeAtIndex(*TagIterator);
 		{
 			OutputRegisterFragment(FragmentStruct);
 		}
diff --git a/Engine/Source/Runtime/MassEntity/Private/MassEntityTypes.cpp b/Engine/Source/Runtime/MassEntity/Private/MassEntityTypes.cpp
--- a/Engine/Source/Runtime/MassEntity/Private/MassEntityTypes.cpp
+++ b/Engine/Source/Runtime/MassEntity/Private/MassEntityTypes.cpp
@@ -61,6 +61,7 @@ uint32 FMassArchetypeCompositionDescriptor::CalculateHash(const FMassFragmentBit
 				, ConstSharedFragmentsHash));
 }
 
+PRAGMA_DISABLE_DEPRECATION_WARNINGS
 int32 FMassArchetypeCompositionDescriptor::CountStoredTypes() const
 {
 	return Fragments.CountStoredTypes()
@@ -117,6 +118,7 @@ void FMassArchetypeCompositionDescriptor::DebugOutputDescription(FOutputDevice&
 	Ar.SetAutoEmitLineTerminator(bAutoLineEnd);
 #endif // WITH_MASSENTITY_DEBUG
 }
+PRAGMA_ENABLE_DEPRECATION_WARNINGS
 
 //-----------------------------------------------------------------------------
 // FMassArchetypeSharedFragmentValues
diff --git a/Engine/Source/Runtime/MassEntity/Private/MassObserverManager.cpp b/Engine/Source/Runtime/MassEntity/Private/MassObserverManager.cpp
--- a/Engine/Source/Runtime/MassEntity/Private/MassObserverManager.cpp
+++ b/Engine/Source/Runtime/MassEntity/Private/MassObserverManager.cpp
@@ -145,25 +145,25 @@ namespace UE::Mass::ObserverManager
 
 		void operator()(const FMassArchetypeCompositionDescriptor& Change, const FEntityCollection& Entities)
 		{
-			if (Change.Fragments.IsEmpty() == false)
+			if (Change.GetFragments().IsEmpty() == false)
 			{
-				(*this)(Change.Fragments, Entities);
+				(*this)(Change.GetFragments(), Entities);
 			}
-			if (Change.Tags.IsEmpty() == false)
+			if (Change.GetTags().IsEmpty() == false)
 			{
-				(*this)(Change.Tags, Entities);
+				(*this)(Change.GetTags(), Entities);
 			}
 		}
 
 		void operator()(const FMassArchetypeCompositionDescriptor& Change, const FMassEntityHandle EntityHandle)
 		{
-			if (Change.Fragments.IsEmpty() == false)
+			if (Change.GetFragments().IsEmpty() == false)
 			{
-				(*this)(Change.Fragments, EntityHandle);
+				(*this)(Change.GetFragments(), EntityHandle);
 			}
-			if (Change.Tags.IsEmpty() == false)
+			if (Change.GetTags().IsEmpty() == false)
 			{
-				(*this)(Change.Tags, EntityHandle);
+				(*this)(Change.GetTags(), EntityHandle);
 			}
 		}
 
@@ -439,8 +439,8 @@ bool FMassObserverManager::OnCompositionChanged(FCollectionRefOrHandle&& EntityC
 
 	const int32 OperationIndex = static_cast<int32>(Operation);
 
-	FMassFragmentBitSet FragmentOverlap = ObservedFragments[OperationIndex].GetOverlap(CompositionDelta.Fragments);
-	FMassTagBitSet TagOverlap = ObservedTags[OperationIndex].GetOverlap(CompositionDelta.Tags);
+	FMassFragmentBitSet FragmentOverlap = ObservedFragments[OperationIndex].GetOverlap(CompositionDelta.GetFragments());
+	FMassTagBitSet TagOverlap = ObservedTags[OperationIndex].GetOverlap(CompositionDelta.GetTags());
 	const bool bHasFragmentsOverlap = !FragmentOverlap.IsEmpty();
 	const bool bHasTagsOverlap = !TagOverlap.IsEmpty();
 
@@ -565,8 +565,8 @@ bool FMassObserverManager::OnCollectionsCreatedImpl(UE::Mass::FProcessingContext
 		checkfSlow(Collection.IsUpToDate(), TEXT("Out-of-date FMassArchetypeEntityCollection used. Stored information is unreliable."));
 
 		const FMassArchetypeCompositionDescriptor& ArchetypeComposition = EntityManager.GetArchetypeComposition(Collection.GetArchetype());
-		FragmentOverlap += ArchetypeComposition.Fragments;
-		TagOverlap += ArchetypeComposition.Tags;
+		FragmentOverlap += ArchetypeComposition.GetFragments();
+		TagOverlap += ArchetypeComposition.GetTags();
 	}
 	FragmentOverlap = ObservedFragments[OperationIndex].GetOverlap(FragmentOverlap);
 	TagOverlap = ObservedTags[OperationIndex].GetOverlap(TagOverlap);
diff --git a/Engine/Source/Runtime/MassEntity/Private/MassRequirements.cpp b/Engine/Source/Runtime/MassEntity/Private/MassRequirements.cpp
--- a/Engine/Source/Runtime/MassEntity/Private/MassRequirements.cpp
+++ b/Engine/Source/Runtime/MassEntity/Private/MassRequirements.cpp
@@ -227,11 +227,11 @@ bool FMassFragmentRequirements::IsEmpty() const
 bool FMassFragmentRequirements::DoesMatchAnyOptionals(const FMassArchetypeCompositionDescriptor& ArchetypeComposition) const
 {
 	return bHasOptionalRequirements
-		&& (ArchetypeComposition.Fragments.HasAny(RequiredOptionalFragments)
-			|| ArchetypeComposition.Tags.HasAny(RequiredOptionalTags)
-			|| ArchetypeComposition.ChunkFragments.HasAny(RequiredOptionalChunkFragments)
-			|| ArchetypeComposition.SharedFragments.HasAny(RequiredOptionalSharedFragments)
-			|| ArchetypeComposition.ConstSharedFragments.HasAny(RequiredOptionalConstSharedFragments));
+		&& (ArchetypeComposition.GetFragments().HasAny(RequiredOptionalFragments)
+			|| ArchetypeComposition.GetTags().HasAny(RequiredOptionalTags)
+			|| ArchetypeComposition.GetChunkFragments().HasAny(RequiredOptionalChunkFragments)
+			|| ArchetypeComposition.GetSharedFragments().HasAny(RequiredOptionalSharedFragments)
+			|| ArchetypeComposition.GetConstSharedFragments().HasAny(RequiredOptionalConstSharedFragments));
 }
 
 bool FMassFragmentRequirements::DoesArchetypeMatchRequirements(const FMassArchetypeHandle& ArchetypeHandle) const
@@ -248,23 +248,23 @@ bool FMassFragmentRequirements::DoesArchetypeMatchRequirements(const FMassArchet
 	CacheProperties();
 
 	const bool bPassNegativeFilter = bHasNegativeRequirements == false
-		|| (ArchetypeComposition.Fragments.HasNone(RequiredNoneFragments)
-			&& ArchetypeComposition.Tags.HasNone(RequiredNoneTags)
-			&& ArchetypeComposition.ChunkFragments.HasNone(RequiredNoneChunkFragments)
-			&& ArchetypeComposition.SharedFragments.HasNone(RequiredNoneSharedFragments)
-			&& ArchetypeComposition.ConstSharedFragments.HasNone(RequiredNoneConstSharedFragments));
+		|| (ArchetypeComposition.GetFragments().HasNone(RequiredNoneFragments)
+			&& ArchetypeComposition.GetTags().HasNone(RequiredNoneTags)
+			&& ArchetypeComposition.GetChunkFragments().HasNone(RequiredNoneChunkFragments)
+			&& ArchetypeComposition.GetSharedFragments().HasNone(RequiredNoneSharedFragments)
+			&& ArchetypeComposition.GetConstSharedFragments().HasNone(RequiredNoneConstSharedFragments));
 	
 	if (bPassNegativeFilter)
 	{
 		if (bHasPositiveRequirements)
 		{
-			return ArchetypeComposition.Fragments.HasAll(RequiredAllFragments)
-				&& (RequiredAnyFragments.IsEmpty() || ArchetypeComposition.Fragments.HasAny(RequiredAnyFragments))
-				&& ArchetypeComposition.Tags.HasAll(RequiredAllTags)
-				&& (RequiredAnyTags.IsEmpty() || ArchetypeComposition.Tags.HasAny(RequiredAnyTags))
-				&& ArchetypeComposition.ChunkFragments.HasAll(RequiredAllChunkFragments)
-				&& ArchetypeComposition.SharedFragments.HasAll(RequiredAllSharedFragments)
-				&& ArchetypeComposition.ConstSharedFragments.HasAll(RequiredAllConstSharedFragments);
+			return ArchetypeComposition.GetFragments().HasAll(RequiredAllFragments)
+				&& (RequiredAnyFragments.IsEmpty() || ArchetypeComposition.GetFragments().HasAny(RequiredAnyFragments))
+				&& ArchetypeComposition.GetTags().HasAll(RequiredAllTags)
+				&& (RequiredAnyTags.IsEmpty() || ArchetypeComposition.GetTags().HasAny(RequiredAnyTags))
+				&& ArchetypeComposition.GetChunkFragments().HasAll(RequiredAllChunkFragments)
+				&& ArchetypeComposition.GetSharedFragments().HasAll(RequiredAllSharedFragments)
+				&& ArchetypeComposition.GetConstSharedFragments().HasAll(RequiredAllConstSharedFragments);
 		}
 		else if (bHasOptionalRequirements)
 		{
diff --git a/Engine/Source/Runtime/MassEntity/Public/MassEntityBuilder.h b/Engine/Source/Runtime/MassEntity/Public/MassEntityBuilder.h
--- a/Engine/Source/Runtime/MassEntity/Public/MassEntityBuilder.h
+++ b/Engine/Source/Runtime/MassEntity/Public/MassEntityBuilder.h
@@ -351,7 +351,7 @@ struct FScopedEntityBuilder : FEntityBuilder
 template<CTag T>
 FEntityBuilder& FEntityBuilder::Add()
 {
-	Composition.Tags.Add<T>();
+	Composition.GetTags().Add<T>();
 	State = EState::ReadyToCommit;
 	CachedArchetypeHandle = FMassArchetypeHandle();
 	return *this;
@@ -360,7 +360,7 @@ FEntityBuilder& FEntityBuilder::Add()
 template<CChunkFragment T>
 FEntityBuilder& FEntityBuilder::Add()
 {
-	Composition.ChunkFragments.Add<T>();
+	Composition.GetChunkFragments().Add<T>();
 	State = EState::ReadyToCommit;
 	CachedArchetypeHandle = FMassArchetypeHandle();
 	return *this;
diff --git a/Engine/Source/Runtime/MassEntity/Public/MassEntityTypes.h b/Engine/Source/Runtime/MassEntity/Public/MassEntityTypes.h
--- a/Engine/Source/Runtime/MassEntity/Public/MassEntityTypes.h
+++ b/Engine/Source/Runtime/MassEntity/Public/MassEntityTypes.h
@@ -40,6 +40,7 @@ namespace UE::Mass
 
 /** The type summarily describing a composition of an entity or an archetype. It contains information on both the
  *  fragments and tags */
+PRAGMA_DISABLE_DEPRECATION_WARNINGS
 struct FMassArchetypeCompositionDescriptor
 {
 	FMassArchetypeCompositionDescriptor() = default;
@@ -202,10 +203,37 @@ struct FMassArchetypeCompositionDescriptor
 	template<typename T>
 	void Add();
 
+	const FMassFragmentBitSet& GetFragments() const;
+	const FMassTagBitSet& GetTags() const;
+	const FMassChunkFragmentBitSet& GetChunkFragments() const;
+	const FMassSharedFragmentBitSet& GetSharedFragments() const;
+	const FMassConstSharedFragmentBitSet& GetConstSharedFragments() const;
+
+	FMassFragmentBitSet& GetFragments();
+	FMassTagBitSet& GetTags();
+	FMassChunkFragmentBitSet& GetChunkFragments();
+	FMassSharedFragmentBitSet& GetSharedFragments();
+	FMassConstSharedFragmentBitSet& GetConstSharedFragments();
+
+	void SetFragments(const FMassFragmentBitSet& InBitSet);
+	void SetTags(const FMassTagBitSet& InBitSet);
+	void SetChunkFragments(const FMassChunkFragmentBitSet& InBitSet);
+	void SetSharedFragments(const FMassSharedFragmentBitSet& InBitSet);
+	void SetConstSharedFragments(const FMassConstSharedFragmentBitSet& InBitSet);
+
+	UE_DEPRECATED(5.7, "Direct access to FMassArchetypeCompositionDescriptor's bitsets is deprecated. Use any of the newly added getters instead.")
 	FMassFragmentBitSet Fragments;
+
+	UE_DEPRECATED(5.7, "Direct access to FMassArchetypeCompositionDescriptor's bitsets is deprecated. Use any of the newly added getters instead.")
 	FMassTagBitSet Tags;
+
+	UE_DEPRECATED(5.7, "Direct access to FMassArchetypeCompositionDescriptor's bitsets is deprecated. Use any of the newly added getters instead.")
 	FMassChunkFragmentBitSet ChunkFragments;
+
+	UE_DEPRECATED(5.7, "Direct access to FMassArchetypeCompositionDescriptor's bitsets is deprecated. Use any of the newly added getters instead.")
 	FMassSharedFragmentBitSet SharedFragments;
+
+	UE_DEPRECATED(5.7, "Direct access to FMassArchetypeCompositionDescriptor's bitsets is deprecated. Use any of the newly added getters instead.")
 	FMassConstSharedFragmentBitSet ConstSharedFragments;
 
 	UE_DEPRECATED(5.5, "This FMassArchetypeCompositionDescriptor constructor is deprecated. Please explicitly provide FConstSharedFragmentBitSet.")
@@ -229,6 +257,7 @@ struct FMassArchetypeCompositionDescriptor
 		ensureMsgf(false, TEXT("This constructor is defunct. Please update your implementation based on deprecation warning."));
 	}
 };
+PRAGMA_ENABLE_DEPRECATION_WARNINGS
 
 /** 
  * Wrapper for const and non-const shared fragment containers that tracks which struct types it holds (via a FMassSharedFragmentBitSet).
@@ -274,8 +303,8 @@ struct FMassArchetypeSharedFragmentValues
 	 */
 	bool DoesMatchComposition(const FMassArchetypeCompositionDescriptor& InDescriptor) const 
 	{
-		return HasExactSharedFragmentTypesMatch(InDescriptor.SharedFragments)
-			&& HasExactConstSharedFragmentTypesMatch(InDescriptor.ConstSharedFragments);
+		return HasExactSharedFragmentTypesMatch(InDescriptor.GetSharedFragments())
+			&& HasExactConstSharedFragmentTypesMatch(InDescriptor.GetConstSharedFragments());
 	}
 
 	inline bool IsEquivalent(const FMassArchetypeSharedFragmentValues& OtherSharedFragmentValues) const
@@ -403,7 +432,7 @@ struct FMassArchetypeSharedFragmentValues
 	 */
 	int32 Remove(const FMassArchetypeCompositionDescriptor& InDescriptor)
 	{
-		return Remove(InDescriptor.SharedFragments) + Remove(InDescriptor.ConstSharedFragments);
+		return Remove(InDescriptor.GetSharedFragments()) + Remove(InDescriptor.GetConstSharedFragments());
 	}
 
 	inline const TArray<FConstSharedStruct>& GetConstSharedFragments() const
@@ -773,6 +802,7 @@ struct FMassArchetypeCreationParams
 //-----------------------------------------------------------------------------
 // INLINES
 //-----------------------------------------------------------------------------
+PRAGMA_DISABLE_DEPRECATION_WARNINGS
 template<typename T>
 auto& FMassArchetypeCompositionDescriptor::GetContainer() const
 {
@@ -845,6 +875,82 @@ void FMassArchetypeCompositionDescriptor::Add()
 	GetContainer<FElementType>().template Add<T>();
 }
 
+inline const FMassFragmentBitSet& FMassArchetypeCompositionDescriptor::GetFragments() const 
+{ 
+	return Fragments; 
+}
+
+inline const FMassTagBitSet& FMassArchetypeCompositionDescriptor::GetTags() const 
+{ 
+	return Tags; 
+}
+
+inline const FMassChunkFragmentBitSet& FMassArchetypeCompositionDescriptor::GetChunkFragments() const 
+{ 
+	return ChunkFragments; 
+}
+
+inline const FMassSharedFragmentBitSet& FMassArchetypeCompositionDescriptor::GetSharedFragments() const 
+{ 
+	return SharedFragments; 
+}
+
+inline const FMassConstSharedFragmentBitSet& FMassArchetypeCompositionDescriptor::GetConstSharedFragments() const 
+{ 
+	return ConstSharedFragments; 
+}
+
+inline FMassFragmentBitSet& FMassArchetypeCompositionDescriptor::GetFragments() 
+{ 
+	return Fragments; 
+}
+
+inline FMassTagBitSet& FMassArchetypeCompositionDescriptor::GetTags() 
+{ 
+	return Tags; 
+}
+
+inline FMassChunkFragmentBitSet& FMassArchetypeCompositionDescriptor::GetChunkFragments() 
+{ 
+	return ChunkFragments; 
+}
+
+inline FMassSharedFragmentBitSet& FMassArchetypeCompositionDescriptor::GetSharedFragments() 
+{ 
+	return SharedFragments; 
+}
+
+inline FMassConstSharedFragmentBitSet& FMassArchetypeCompositionDescriptor::GetConstSharedFragments() 
+{ 
+	return ConstSharedFragments; 
+}
+
+inline void FMassArchetypeCompositionDescriptor::SetFragments(const FMassFragmentBitSet& InBitSet)
+{ 
+	Fragments = InBitSet; 
+}
+
+inline void FMassArchetypeCompositionDescriptor::SetTags(const FMassTagBitSet& InBitSet)
+{ 
+	Tags = InBitSet; 
+}
+
+inline void FMassArchetypeCompositionDescriptor::SetChunkFragments(const FMassChunkFragmentBitSet& InBitSet)
+{ 
+	ChunkFragments = InBitSet; 
+}
+
+inline void FMassArchetypeCompositionDescriptor::SetSharedFragments(const FMassSharedFragmentBitSet& InBitSet)
+{ 
+	SharedFragments = InBitSet; 
+}
+
+inline void FMassArchetypeCompositionDescriptor::SetConstSharedFragments(const FMassConstSharedFragmentBitSet& InBitSet)
+{ 
+	ConstSharedFragments = InBitSet; 
+}
+PRAGMA_ENABLE_DEPRECATION_WARNINGS
+
 inline SIZE_T FMassArchetypeSharedFragmentValues::GetAllocatedSize() const
 {
 	return SharedFragmentBitSet.GetAllocatedSize()
diff --git a/Engine/Source/Runtime/MassEntity/Public/MassExecutionContext.h b/Engine/Source/Runtime/MassEntity/Public/MassExecutionContext.h
--- a/Engine/Source/Runtime/MassEntity/Public/MassExecutionContext.h
+++ b/Engine/Source/Runtime/MassEntity/Public/MassExecutionContext.h
@@ -297,26 +297,26 @@ public:
 
 	bool DoesArchetypeHaveFragment(const UScriptStruct& FragmentType) const
 	{
-		return CurrentArchetypeCompositionDescriptor.Fragments.Contains(FragmentType);
+		return CurrentArchetypeCompositionDescriptor.GetFragments().Contains(FragmentType);
 	}
 
 	template<typename T>
 	bool DoesArchetypeHaveFragment() const
 	{
 		static_assert(UE::Mass::CFragment<T>, "Given struct is not of a valid fragment type.");
-		return CurrentArchetypeCompositionDescriptor.Fragments.Contains<T>();
+		return CurrentArchetypeCompositionDescriptor.GetFragments().Contains<T>();
 	}
 
 	bool DoesArchetypeHaveTag(const UScriptStruct& TagType) const
 	{
-		return CurrentArchetypeCompositionDescriptor.Tags.Contains(TagType);
+		return CurrentArchetypeCompositionDescriptor.GetTags().Contains(TagType);
 	}
 
 	template<typename T>
 	bool DoesArchetypeHaveTag() const
 	{
 		static_assert(UE::Mass::CTag<T>, "Given struct is not of a valid tag type.");
-		return CurrentArchetypeCompositionDescriptor.Tags.Contains<T>();
+		return CurrentArchetypeCompositionDescriptor.GetTags().Contains<T>();
 	}
 
 #if WITH_MASSENTITY_DEBUG
diff --git a/Engine/Source/Runtime/MassEntity/Public/MassObserverManager.h b/Engine/Source/Runtime/MassEntity/Public/MassObserverManager.h
--- a/Engine/Source/Runtime/MassEntity/Public/MassObserverManager.h
+++ b/Engine/Source/Runtime/MassEntity/Public/MassObserverManager.h
@@ -101,7 +101,7 @@ struct FMassObserverManager
 
 	bool HasObserversForComposition(const FMassArchetypeCompositionDescriptor& Composition, const EMassObservedOperation Operation) const
 	{
-		return HasObserversForBitSet(Composition.Fragments, Operation) || HasObserversForBitSet(Composition.Tags, Operation);
+		return HasObserversForBitSet(Composition.GetFragments(), Operation) || HasObserversForBitSet(Composition.GetTags(), Operation);
 	}
 
 	/** @return whether there are observers watching affected elements */
