Minor Mass optimization by caching the max entities per page number of bits to avoid a log2 at runtime.

diff --git a/Engine/Source/Runtime/MassEntity/Private/MassEntityManagerStorage.cpp b/Engine/Source/Runtime/MassEntity/Private/MassEntityManagerStorage.cpp
--- a/Engine/Source/Runtime/MassEntity/Private/MassEntityManagerStorage.cpp
+++ b/Engine/Source/Runtime/MassEntity/Private/MassEntityManagerStorage.cpp
@@ -244,6 +244,7 @@ namespace UE::Mass
 		check(FMath::IsPowerOfTwo(InInitializationParams.MaxEntitiesPerPage));
 		check(FMath::IsPowerOfTwo(InInitializationParams.MaxEntityCount));
 		MaxEntitiesPerPage = InInitializationParams.MaxEntitiesPerPage;
+		MaxEntitiesPerPageShift = FMath::FloorLog2(InInitializationParams.MaxEntitiesPerPage);
 		MaximumEntityCountShift = FMath::FloorLog2(InInitializationParams.MaxEntityCount);
 		checkf(MaximumEntityCountShift < 32, TEXT("Invalid maximum entity count, cannot exceed 31 bits"));
 
@@ -327,7 +328,7 @@ namespace UE::Mass
 		// Page Index is which page in the array of pages we need to access
 		if (Index >= 0)
 		{
-			const uint32 PageIndex = static_cast<uint32>(Index) >> FMath::FloorLog2(MaxEntitiesPerPage);
+			const uint32 PageIndex = static_cast<uint32>(Index) >> MaxEntitiesPerPageShift;
 			return PageIndex < PageCount;
 		}
 		return false;
@@ -604,7 +605,7 @@ namespace UE::Mass
 	{
 		check(Index >= 0);
 		// PageIndex is which Page in the array of pages we need to access
-		const uint32 PageIndex = static_cast<uint32>(Index) >> FMath::FloorLog2(MaxEntitiesPerPage);
+		const uint32 PageIndex = static_cast<uint32>(Index) >> MaxEntitiesPerPageShift;
 
 		// Convert the entity index into the index with respect to the page
 		const uint32 EntityOffset = (PageIndex * MaxEntitiesPerPage);
diff --git a/Engine/Source/Runtime/MassEntity/Public/MassEntityManagerStorage.h b/Engine/Source/Runtime/MassEntity/Public/MassEntityManagerStorage.h
--- a/Engine/Source/Runtime/MassEntity/Public/MassEntityManagerStorage.h
+++ b/Engine/Source/Runtime/MassEntity/Public/MassEntityManagerStorage.h
@@ -258,6 +258,7 @@ namespace UE::Mass
 		uint32 EntityCount = 0;
 		uint32 MaximumEntityCountShift = 0;
 		uint32 MaxEntitiesPerPage = 0;
+		uint32 MaxEntitiesPerPageShift = 0;
 		uint32 PageCount = 0;
 		/** ALWAYS acquire FreeListMutex before this one */
 		UE::FTransactionallySafeMutex PageAllocateMutex;
