[Mass] Fixed FMassRequirementAccessDetector::Initialize not handling const shared fragments.

#rnx

diff --git a/Engine/Source/Runtime/MassEntity/Private/MassRequirementAccessDetector.cpp b/Engine/Source/Runtime/MassEntity/Private/MassRequirementAccessDetector.cpp
--- a/Engine/Source/Runtime/MassEntity/Private/MassRequirementAccessDetector.cpp
+++ b/Engine/Source/Runtime/MassEntity/Private/MassRequirementAccessDetector.cpp
@@ -20,6 +20,7 @@ void FMassRequirementAccessDetector::Initialize()
 	AddDetectors(FMassFragmentBitSet::FStructTrackerWrapper::StructTracker);
 	AddDetectors(FMassChunkFragmentBitSet::FStructTrackerWrapper::StructTracker);
 	AddDetectors(FMassSharedFragmentBitSet::FStructTrackerWrapper::StructTracker);
+	AddDetectors(FMassConstSharedFragmentBitSet::FStructTrackerWrapper::StructTracker);
 	AddDetectors(FMassExternalSubsystemBitSet::FStructTrackerWrapper::StructTracker);
 }
 
