[Mass] Fixed FMassArchetypeData::BatchMoveEntitiesToAnotherArchetype's wrong behavior when ConstSharedFragmentToRemoveBitSet is the only non-null shared fragments change provided.

#rnx

diff --git a/Engine/Source/Runtime/MassEntity/Private/MassArchetypeData.cpp b/Engine/Source/Runtime/MassEntity/Private/MassArchetypeData.cpp
--- a/Engine/Source/Runtime/MassEntity/Private/MassArchetypeData.cpp
+++ b/Engine/Source/Runtime/MassEntity/Private/MassArchetypeData.cpp
@@ -1195,13 +1195,13 @@ void FMassArchetypeData::BatchMoveEntitiesToAnotherArchetype(const FMassArchetyp
 		ResultSubChunk.ChunkIndex = 0;
 		ResultSubChunk.Length = 0;
 		int32 NumberMoved = 0;
-		const bool ChangeSharedFragments = SharedFragmentValuesToAdd || SharedFragmentToRemoveBitSet;
+		const bool bChangeSharedFragments = SharedFragmentValuesToAdd || SharedFragmentToRemoveBitSet || ConstSharedFragmentToRemoveBitSet;
 
 		do
 		{
 			const int32 IndexWithinChunk = EntityRange.SubchunkStart + NumberMoved;
 
-			if (ChangeSharedFragments == false)
+			if (bChangeSharedFragments == false)
 			{
 				ResultSubChunk = NewArchetype.PrepareNextEntitiesSpanInternal(MakeArrayView(DyingEntityPtr + NumberMoved, RangeLength - NumberMoved)
 					, Chunk.GetSharedFragmentValues(), ResultSubChunk.ChunkIndex);
