[Mass]
Added FMassEntityQuery::FExecutionLimiter to limit execution to a set entity count limit and then resume from that point on the next cycle.
Added unit tests for FMassEntityQuery::FExecutionLimiter

diff --git a/Engine/Source/Developer/MassEntityTestSuite/Private/MassQueryTest.cpp b/Engine/Source/Developer/MassEntityTestSuite/Private/MassQueryTest.cpp
--- a/Engine/Source/Developer/MassEntityTestSuite/Private/MassQueryTest.cpp
+++ b/Engine/Source/Developer/MassEntityTestSuite/Private/MassQueryTest.cpp
@@ -950,6 +950,90 @@ struct FQueryTest_ExportHandles : FEntityTestBase
 };
 IMPLEMENT_AI_INSTANT_TEST(FQueryTest_ExportHandles, "System.Mass.Query.ExportHandles");
 
+struct FQueryTest_ExecutionLimiter : FEntityTestBase
+{
+	virtual bool InstantTest() override
+	{
+		CA_ASSUME(EntityManager);
+
+		const int32 EntityLimit = 1;
+		const int32 CyclesToSaturate = 20; // no assumptions about chunk size, the minimum cycles will be the total entity count needing processing
+		const int32 FloatsArchetypeCreated = 7;
+		const int32 IntsArchetypeCreated = 11;
+		const int32 FloatsIntsArchetypeCreated = 13;
+		TArray<FMassEntityHandle> EntitiesCreated;
+		EntityManager->BatchCreateEntities(IntsArchetype, IntsArchetypeCreated, EntitiesCreated);
+		// clear to store only the float-related entities
+		EntitiesCreated.Reset();
+		EntityManager->BatchCreateEntities(FloatsArchetype, FloatsArchetypeCreated, EntitiesCreated);
+		EntityManager->BatchCreateEntities(FloatsIntsArchetype, FloatsIntsArchetypeCreated, EntitiesCreated);
+
+		int TotalProcessed = 0;
+		FMassExecutionContext ExecContext(*EntityManager.Get());
+		FMassEntityQuery Query(EntityManager.ToSharedRef(), { FTestFragment_Float::StaticStruct() });
+		// limit the execution to a single entity per cycle (this will process the entire chunk, but only one no matter how many entities it contains)
+		UE::Mass::FExecutionLimiter Limiter(EntityLimit);
+		Query.ForEachEntityChunk(ExecContext, Limiter, [&TotalProcessed](FMassExecutionContext& Context)
+			{
+				TotalProcessed += Context.GetNumEntities();
+				TArrayView<FTestFragment_Float> Floats = Context.GetMutableFragmentView<FTestFragment_Float>();
+
+				for (int32 i = 0; i < Context.GetNumEntities(); ++i)
+				{
+					Floats[i].Value = 13.f;
+				}
+			});
+
+		AITEST_TRUE("Limiter should prevent all entities from being processed", TotalProcessed < FloatsIntsArchetypeCreated + FloatsArchetypeCreated);
+		AITEST_TRUE("Limiter should process at least as many entities as requested", TotalProcessed >= EntityLimit);
+
+		// run a few iterations to ensure all entities are processed:
+		for (int i = 0; i < CyclesToSaturate; i++)
+		{
+			Query.ForEachEntityChunk(ExecContext, Limiter, [&TotalProcessed](FMassExecutionContext& Context)
+				{
+					TotalProcessed += Context.GetNumEntities();
+					TArrayView<FTestFragment_Float> Floats = Context.GetMutableFragmentView<FTestFragment_Float>();
+
+					for (int32 i = 0; i < Context.GetNumEntities(); ++i)
+					{
+						Floats[i].Value = 13.f;
+					}
+				});
+		}
+
+		for (FMassEntityHandle& Entity : EntitiesCreated)
+		{
+			const FTestFragment_Float& TestedFragment = EntityManager->GetFragmentDataChecked<FTestFragment_Float>(Entity);
+			AITEST_EQUAL("Every fragment value should have changed to the expected value", TestedFragment.Value, 13.f);
+		}
+
+		Limiter.EntityLimit = 1000000;
+		TotalProcessed = 0;
+		Query.ForEachEntityChunk(ExecContext, Limiter, [&TotalProcessed](FMassExecutionContext& Context)
+			{
+				TotalProcessed += Context.GetNumEntities();
+				TArrayView<FTestFragment_Float> Floats = Context.GetMutableFragmentView<FTestFragment_Float>();
+
+				for (int32 i = 0; i < Context.GetNumEntities(); ++i)
+				{
+					Floats[i].Value = 13.f;
+				}
+			});
+
+		AITEST_TRUE("All entities should be processed with a high enough limit", TotalProcessed == FloatsIntsArchetypeCreated + FloatsArchetypeCreated);
+
+		for (FMassEntityHandle& Entity : EntitiesCreated)
+		{
+			const FTestFragment_Float& TestedFragment = EntityManager->GetFragmentDataChecked<FTestFragment_Float>(Entity);
+			AITEST_EQUAL("Every fragment value should have changed to the expected value", TestedFragment.Value, 13.f);
+		}
+
+		return true;
+	}
+};
+IMPLEMENT_AI_INSTANT_TEST(FQueryTest_ExecutionLimiter, "System.Mass.Query.ExecutionLimiter");
+
 } // FMassQueryTest
 
 #undef LOCTEXT_NAMESPACE
diff --git a/Engine/Source/Runtime/MassEntity/Internal/MassArchetypeData.h b/Engine/Source/Runtime/MassEntity/Internal/MassArchetypeData.h
--- a/Engine/Source/Runtime/MassEntity/Internal/MassArchetypeData.h
+++ b/Engine/Source/Runtime/MassEntity/Internal/MassArchetypeData.h
@@ -307,7 +307,7 @@ public:
 	void ExecuteFunction(FMassExecutionContext& RunContext, const FMassExecuteFunction& Function, const FMassQueryRequirementIndicesMapping& RequirementMapping
 		, FMassArchetypeEntityCollection::FConstEntityRangeArrayView EntityRangeContainer, const FMassChunkConditionFunction& ChunkCondition);
 	void ExecuteFunction(FMassExecutionContext& RunContext, const FMassExecuteFunction& Function, const FMassQueryRequirementIndicesMapping& RequirementMapping
-		, const FMassChunkConditionFunction& ChunkCondition);
+		, const FMassChunkConditionFunction& ChunkCondition, UE::Mass::FExecutionLimiter* ExecutionLimiter = nullptr);
 
 	void ExecutionFunctionForChunk(FMassExecutionContext& RunContext, const FMassExecuteFunction& Function, const FMassQueryRequirementIndicesMapping& RequirementMapping
 		, const FMassArchetypeEntityCollection::FArchetypeEntityRange& EntityRange, const FMassChunkConditionFunction& ChunkCondition = FMassChunkConditionFunction());
diff --git a/Engine/Source/Runtime/MassEntity/Private/MassArchetypeData.cpp b/Engine/Source/Runtime/MassEntity/Private/MassArchetypeData.cpp
--- a/Engine/Source/Runtime/MassEntity/Private/MassArchetypeData.cpp
+++ b/Engine/Source/Runtime/MassEntity/Private/MassArchetypeData.cpp
@@ -496,7 +496,7 @@ void FMassArchetypeData::ExecuteFunction(FMassExecutionContext& RunContext, cons
 	RunContext.DebugSetColor(DebugColor);
 #endif // WITH_MASSENTITY_DEBUG
 
-	uint32 PrevSharedFragmentValuesHash = UINT32_MAX;
+	uint32 PrevSharedFragmentValuesHash = TNumericLimits<uint32>::Max();
 	for (FMassArchetypeChunkIterator ChunkIterator(EntityRangeContainer); ChunkIterator; ++ChunkIterator)
 	{
 		FMassArchetypeChunk& Chunk = Chunks[ChunkIterator->ChunkIndex];
@@ -526,7 +526,7 @@ void FMassArchetypeData::ExecuteFunction(FMassExecutionContext& RunContext, cons
 	}
 }
 
-void FMassArchetypeData::ExecuteFunction(FMassExecutionContext& RunContext, const FMassExecuteFunction& Function, const FMassQueryRequirementIndicesMapping& RequirementMapping, const FMassChunkConditionFunction& ChunkCondition)
+void FMassArchetypeData::ExecuteFunction(FMassExecutionContext& RunContext, const FMassExecuteFunction& Function, const FMassQueryRequirementIndicesMapping& RequirementMapping, const FMassChunkConditionFunction& ChunkCondition, UE::Mass::FExecutionLimiter* ExecutionLimiter)
 {
 	if (GetNumEntities() == 0)
 	{
@@ -538,10 +538,24 @@ void FMassArchetypeData::ExecuteFunction(FMassExecutionContext& RunContext, cons
 	RunContext.DebugSetColor(DebugColor);
 #endif // WITH_MASSENTITY_DEBUG
 
-	uint32 PrevSharedFragmentValuesHash = UINT32_MAX;
-	for (FMassArchetypeChunk& Chunk : Chunks)
+	int32 ChunkIndex = 0;
+	int32 EntityCountRemaining = TNumericLimits<int32>::Max();
+	int32 MaxChunkIndex = Chunks.Num();
+
+	if (ExecutionLimiter)
+	{
+		ChunkIndex = FMath::Max(ExecutionLimiter->ChunkIndex, 0);
+		EntityCountRemaining = ExecutionLimiter->EntityCountRemaining;
+		MaxChunkIndex = FMath::Min(MaxChunkIndex, ExecutionLimiter->MaxChunkIndex);
+	}
+
+	uint32 PrevSharedFragmentValuesHash = TNumericLimits<uint32>::Max();
+	for (; ChunkIndex < MaxChunkIndex && EntityCountRemaining > 0; ChunkIndex++)
 	{
-		if (Chunk.GetNumInstances())
+		FMassArchetypeChunk& Chunk = Chunks[ChunkIndex];
+		const int32 EntityCount = Chunk.GetNumInstances();
+
+		if (EntityCount)
 		{
 			const uint32 SharedFragmentValuesHash = GetTypeHash(Chunk.GetSharedFragmentValues());
 			if (PrevSharedFragmentValuesHash != SharedFragmentValuesHash)
@@ -558,9 +572,18 @@ void FMassArchetypeData::ExecuteFunction(FMassExecutionContext& RunContext, cons
 			{
 				BindEntityRequirements(RunContext, RequirementMapping.EntityFragments, Chunk, 0, Chunk.GetNumInstances());
 				Function(RunContext);
+
+				EntityCountRemaining -= EntityCount;
 			}
 		}
 	}
+
+	if (ExecutionLimiter)
+	{
+		// set the limiter to continue on the next chunk if this archetype has unprocessed chunks
+		ExecutionLimiter->ChunkIndex = ++ChunkIndex;
+		ExecutionLimiter->EntityCountRemaining = EntityCountRemaining;
+	}
 }
 
 void FMassArchetypeData::ExecutionFunctionForChunk(FMassExecutionContext& RunContext, const FMassExecuteFunction& Function, const FMassQueryRequirementIndicesMapping& RequirementMapping, const FMassArchetypeEntityCollection::FArchetypeEntityRange& EntityRange, const FMassChunkConditionFunction& ChunkCondition)
diff --git a/Engine/Source/Runtime/MassEntity/Private/MassEntityQuery.cpp b/Engine/Source/Runtime/MassEntity/Private/MassEntityQuery.cpp
--- a/Engine/Source/Runtime/MassEntity/Private/MassEntityQuery.cpp
+++ b/Engine/Source/Runtime/MassEntity/Private/MassEntityQuery.cpp
@@ -267,7 +267,7 @@ void FMassEntityQuery::ForEachEntityChunk(FMassExecutionContext& ExecutionContex
 			// do need those calculations.
 		CacheArchetypes();
 
-		// if there's a chunk collection set by the external code - usse that
+		// if there's a chunk collection set by the external code - use that
 		if (ExecutionContext.GetEntityCollection().IsEmpty() == false)
 		{
 			const FMassArchetypeHandle& ArchetypeHandle = ExecutionContext.GetEntityCollection().GetArchetype();
@@ -314,6 +314,100 @@ void FMassEntityQuery::ForEachEntityChunk(FMassExecutionContext& ExecutionContex
 	}
 }
 
+void FMassEntityQuery::ForEachEntityChunk(FMassExecutionContext& ExecutionContext, UE::Mass::FExecutionLimiter& Limiter, const FMassExecuteFunction& ExecuteFunction)
+{
+	UE_TRACE_SCOPED_MASS_QUERY_FOR_EACH()
+
+	checkf(CachedEntityManager == ExecutionContext.GetSharedEntityManager() && CachedEntityManager.IsValid()
+		, TEXT("FMassEntityQuery needs to be initialized with a valid EntityManager and the execution context has to come from the same manager"));
+
+#if WITH_MASSENTITY_DEBUG
+	checkf(ExecutionContext.GetExecutionType() == ExpectedContextType && (ExpectedContextType == EMassExecutionContextType::Local || bRegistered)
+		, TEXT("ExecutionContextType mismatch, make sure all the queries run as part of processor execution are registered with some processor with a FMassEntityQuery::RegisterWithProcessor call"));
+#endif
+
+	FScopedEntityQueryContext ScopedQueryContext(*this, ExecutionContext);
+
+	if (ScopedQueryContext.IsSuccessfullySetUp() == false)
+	{
+		// required subsystems are not available, bail out.
+		return;
+	}
+
+	// FExecutionLimiter is not supported for use with manually set entity collections.
+	checkf(FMassFragmentRequirements::IsEmpty() == false, TEXT("FExecutionLimiter is not supported for use with manually set entity collections."));
+	checkf(ExecutionContext.GetEntityCollection().IsEmpty(), TEXT("FExecutionLimiter is not supported for use with manually set entity collections."));
+
+	// note that the following function will usually only resort to verifying that the data is up to date by
+	// checking the version number. In rare cases when it would result in non trivial cost we actually
+	// do need those calculations.
+	CacheArchetypes();
+
+	// it's important to set requirements after caching archetypes due to that call potentially sorting the requirements and the order is relevant here.
+	ExecutionContext.ApplyFragmentRequirements(*this);
+
+	// note that this checkSlow is here on purpose, for debugging purposes, not data verification purposes.
+	const int32 NumArchetypes = ValidArchetypes.Num();
+	checkSlow(OrderedArchetypeIndices.Num() == NumArchetypes);
+
+	if (NumArchetypes == 0)
+	{
+		return;
+	}
+
+	if (Limiter.ChunkIndex < 0 || Limiter.ArchetypeIndex < 0 || Limiter.ArchetypeIndex >= NumArchetypes)
+	{
+		Limiter.ArchetypeIndex = 0;
+		Limiter.ChunkIndex = 0;
+	}
+	const int32 StartingChunkIndex = Limiter.ChunkIndex;
+	const int32 StartingArchetypeIndex = Limiter.ArchetypeIndex;
+	Limiter.EntityCountRemaining = Limiter.EntityLimit;
+	Limiter.MaxChunkIndex = TNumericLimits<int32>::Max();
+	bool bLooped = false;
+	bool bDone = false;
+
+	do
+	{
+		const int32 ArchetypeIndex = OrderedArchetypeIndices[Limiter.ArchetypeIndex];
+		const FMassArchetypeHandle& ArchetypeHandle = ValidArchetypes[ArchetypeIndex];
+		FMassArchetypeData& ArchetypeData = FMassArchetypeHelper::ArchetypeDataFromHandleChecked(ArchetypeHandle);
+
+		if (bLooped && StartingChunkIndex != 0)
+		{
+			// archetype was started past the first chunk so if we loop we need to process it again
+			// but stop before the starting chunk is reprocessed
+			Limiter.MaxChunkIndex = StartingChunkIndex;
+			bDone = true;
+		}
+
+		UE_TRACE_SCOPED_MASS_QUERY_FOR_EACH_REPORT_ARCHETYPE(ArchetypeData)
+
+		ArchetypeData.ExecuteFunction(ExecutionContext, ExecuteFunction, ArchetypeFragmentMapping[ArchetypeIndex], ChunkCondition, &Limiter);
+		ExecutionContext.ClearFragmentViews(*this);
+		
+		if (Limiter.ChunkIndex >= ArchetypeData.GetChunkCount())
+		{
+			// all chunks in the archetype were processed, so move to the next
+			Limiter.ChunkIndex = 0;
+			if (++Limiter.ArchetypeIndex >= NumArchetypes)
+			{
+				Limiter.ArchetypeIndex = 0;
+			}
+			if (Limiter.ArchetypeIndex == StartingArchetypeIndex)
+			{
+				if (StartingChunkIndex == 0)
+				{
+					// no need to process the starting archetype again because the initial pass hit all the chunks
+					bDone = true;
+				}
+				bLooped = true;
+			}
+		}
+		
+	} while (Limiter.EntityCountRemaining > 0 && !bDone);
+}
+
 void FMassEntityQuery::ParallelForEachEntityChunkInCollection(TConstArrayView<FMassArchetypeEntityCollection> EntityCollections
 	, FMassExecutionContext& ExecutionContext, const FMassExecuteFunction& ExecuteFunction
 	, const EParallelExecutionFlags Flags)
diff --git a/Engine/Source/Runtime/MassEntity/Public/MassEntityQuery.h b/Engine/Source/Runtime/MassEntity/Public/MassEntityQuery.h
--- a/Engine/Source/Runtime/MassEntity/Public/MassEntityQuery.h
+++ b/Engine/Source/Runtime/MassEntity/Public/MassEntityQuery.h
@@ -87,6 +87,15 @@ public:
 
 	/** Runs ExecuteFunction on all entities matching Requirements */
 	UE_API void ForEachEntityChunk(FMassExecutionContext& ExecutionContext, const FMassExecuteFunction& ExecuteFunction);
+
+	/**
+	 * Runs ExecuteFunction on entities matching Requirements up to the EntityLimit specified in the ExecutionLimiter.
+	 * Resumes from last index set in Limiter.
+	 * Will always complete the current chunk once the limit is reached.
+	 * Does not account for changes to entity count or organization between iterations (Entities could be skipped).
+	 * Identical to the unlimited version if the Limit exceeds the number of applicable entities.
+	 */
+	UE_API void ForEachEntityChunk(FMassExecutionContext& ExecutionContext, UE::Mass::FExecutionLimiter& Limiter, const FMassExecuteFunction& ExecuteFunction);
 	
 	/** Will first verify that the archetype given with Collection matches the query's requirements, and if so will run the other, more generic ForEachEntityChunk implementation */
 	UE_API void ForEachEntityChunk(const FMassArchetypeEntityCollection& EntityCollection
diff --git a/Engine/Source/Runtime/MassEntity/Public/MassEntityTypes.h b/Engine/Source/Runtime/MassEntity/Public/MassEntityTypes.h
--- a/Engine/Source/Runtime/MassEntity/Public/MassEntityTypes.h
+++ b/Engine/Source/Runtime/MassEntity/Public/MassEntityTypes.h
@@ -32,10 +32,39 @@ DECLARE_STRUCTTYPEBITSET_EXPORTED(MASSENTITY_API, FMassSharedFragmentBitSet, FMa
 DECLARE_STRUCTTYPEBITSET_EXPORTED(MASSENTITY_API, FMassConstSharedFragmentBitSet, FMassConstSharedFragment);
 DECLARE_CLASSTYPEBITSET_EXPORTED(MASSENTITY_API, FMassExternalSubsystemBitSet, USubsystem);
 
+struct FMassArchetypeData;
+struct FMassEntityQuery;
+
 namespace UE::Mass
 {
 	template<typename T>
 	static constexpr bool TAlwaysFalse = false;
+
+	/**
+	 * FExecutionLimiter is used to limit the execution of a query to a set entity count.
+	 */
+	struct FExecutionLimiter
+	{
+		friend struct ::FMassArchetypeData;
+		friend struct ::FMassEntityQuery;
+
+		explicit FExecutionLimiter(int32 InEntityLimit)
+			: EntityLimit(InEntityLimit)
+			, ChunkIndex(INDEX_NONE)
+			, ArchetypeIndex(INDEX_NONE)
+			, MaxChunkIndex(INDEX_NONE)
+			, EntityCountRemaining(0)
+		{
+		}
+		
+		int32 EntityLimit;
+
+	private:
+		int32 ChunkIndex;
+		int32 ArchetypeIndex;
+		int32 MaxChunkIndex;
+		int32 EntityCountRemaining;
+	};
 }
 
 /** The type summarily describing a composition of an entity or an archetype. It contains information on both the
