[Mass]
Fixed inconsistencies between single and multi-threaded deferred command flushing in FMassProcessingPhase::ExecuteTick

diff --git a/Engine/Source/Runtime/MassEntity/Private/MassProcessingPhaseManager.cpp b/Engine/Source/Runtime/MassEntity/Private/MassProcessingPhaseManager.cpp
--- a/Engine/Source/Runtime/MassEntity/Private/MassProcessingPhaseManager.cpp
+++ b/Engine/Source/Runtime/MassEntity/Private/MassProcessingPhaseManager.cpp
@@ -78,7 +78,6 @@ void FMassProcessingPhase::ExecuteTick(float DeltaTime, ELevelTick TickType, ENa
 	check(PhaseProcessor);
 	
 	FMassEntityManager& EntityManager = PhaseManager->GetEntityManagerRef();
-	FMassProcessingContext Context(EntityManager, DeltaTime);
 
 	bIsDuringMassProcessing = true;
 
@@ -87,6 +86,7 @@ void FMassProcessingPhase::ExecuteTick(float DeltaTime, ELevelTick TickType, ENa
 		bool bWorkRequested = false;
 		if (PhaseProcessor->IsEmpty() == false)
 		{
+			FMassProcessingContext Context(EntityManager, DeltaTime);
 			const FGraphEventRef PipelineCompletionEvent = UE::Mass::Executor::TriggerParallelTasks(*PhaseProcessor, MoveTemp(Context), [this, DeltaTime]()
 				{
 					OnParallelExecutionDone(DeltaTime);
@@ -108,6 +108,10 @@ void FMassProcessingPhase::ExecuteTick(float DeltaTime, ELevelTick TickType, ENa
 	{
 		if (PhaseManager->IsPaused() == false)
 		{
+			// note that it's important to create the processing context in this scope
+			// so that it wraps up its destruction before we call OnPhaseEnd, which in turn will cause
+			// the main EntityManager's command buffer to flush
+			FMassProcessingContext Context(EntityManager, DeltaTime);
 			UE::Mass::Executor::Run(*PhaseProcessor, Context);
 		}
 
