Mass Manager Optimizations
- Removed some very costly calls to CheckIfEntityIsActive and merged some checks to minimize the impacts on performance.
- FMassEntityManager::IsEntityActive now route to the entity storage interface to avoid multiple calls to FConcurrentEntityStorage::LookupEntity.
- Replaced multiplications by power of two by left shifts.

diff --git a/Engine/Source/Runtime/MassEntity/Private/MassEntityManager.cpp b/Engine/Source/Runtime/MassEntity/Private/MassEntityManager.cpp
--- a/Engine/Source/Runtime/MassEntity/Private/MassEntityManager.cpp
+++ b/Engine/Source/Runtime/MassEntity/Private/MassEntityManager.cpp
@@ -1901,7 +1901,6 @@ void* FMassEntityManager::InternalGetFragmentDataPtr(FMassEntityHandle Entity, c
 {
 	// note that FragmentType is guaranteed to be of valid type - it's either statically checked by the template versions
 	// or `checkf`ed by the non-template one
-	CheckIfEntityIsActive(Entity);
 	const FMassArchetypeData* CurrentArchetype = GetEntityStorageInterface().GetArchetype(Entity.Index);
 	check(CurrentArchetype);
 	return CurrentArchetype->GetFragmentDataForEntity(FragmentType, Entity.Index);
@@ -1911,7 +1910,6 @@ const FConstSharedStruct* FMassEntityManager::InternalGetConstSharedFragmentPtr(
 {
 	// note that ConstSharedFragmentType is guaranteed to be of valid type - it's either statically checked by the template versions
 	// or `checkf`ed by the non-template one
-	CheckIfEntityIsActive(Entity);
 	const FMassArchetypeData* CurrentArchetype = GetEntityStorageInterface().GetArchetype(Entity.Index);
 	check(CurrentArchetype);
 	const FConstSharedStruct* SharedFragment = CurrentArchetype->GetSharedFragmentValues(Entity).GetConstSharedFragments().FindByPredicate(FStructTypeEqualOperator(ConstSharedFragmentType));
@@ -1922,13 +1920,22 @@ const FSharedStruct* FMassEntityManager::InternalGetSharedFragmentPtr(FMassEntit
 {
 	// note that SharedFragmentType is guaranteed to be of valid type - it's either statically checked by the template versions
 	// or `checkf`ed by the non-template one
-	CheckIfEntityIsActive(Entity);
 	const FMassArchetypeData* CurrentArchetype = GetEntityStorageInterface().GetArchetype(Entity.Index);
 	check(CurrentArchetype);
 	const FSharedStruct* SharedFragment = CurrentArchetype->GetSharedFragmentValues(Entity).GetSharedFragments().FindByPredicate(FStructTypeEqualOperator(SharedFragmentType));
 	return SharedFragment;
 }
 
+bool FMassEntityManager::IsEntityActive(FMassEntityHandle Entity) const
+{
+	if (Entity.Index != UE::Mass::Private::InvalidEntityIndex)
+	{
+		const UE::Mass::FStorageType& EntityStorageInterface = GetEntityStorageInterface();
+		return EntityStorageInterface.IsEntityActive(Entity);
+	}
+	return false;
+}
+
 bool FMassEntityManager::IsEntityValid(FMassEntityHandle Entity) const
 {
 	return (Entity.Index != UE::Mass::Private::InvalidEntityIndex) 
@@ -1949,17 +1956,6 @@ bool FMassEntityManager::IsEntityReserved(FMassEntityHandle EntityHandle) const
 	return GetEntityStorageInterface().GetEntityState(EntityHandle.Index) == UE::Mass::IEntityStorageInterface::EEntityState::Reserved;
 }
 
-void FMassEntityManager::CheckIfEntityIsValid(FMassEntityHandle Entity) const
-{
-	checkf(IsEntityValid(Entity), TEXT("Invalid entity (ID: %d, SN:%d, %s)"), Entity.Index, Entity.SerialNumber,
-		   (Entity.Index == 0) ? TEXT("was never initialized") : TEXT("already destroyed"));
-}
-
-void FMassEntityManager::CheckIfEntityIsActive(FMassEntityHandle Entity) const
-{
-	checkf(IsEntityBuilt(Entity), TEXT("Entity not yet created(ID: %d, SN:%d)"), Entity.Index, Entity.SerialNumber);
-}
-
 FMassEntityHandle FMassEntityManager::CreateEntityIndexHandle(const int32 EntityIndex) const
 {
 	return (GetEntityStorageInterface().IsValidIndex(EntityIndex)
diff --git a/Engine/Source/Runtime/MassEntity/Private/MassEntityManagerStorage.cpp b/Engine/Source/Runtime/MassEntity/Private/MassEntityManagerStorage.cpp
--- a/Engine/Source/Runtime/MassEntity/Private/MassEntityManagerStorage.cpp
+++ b/Engine/Source/Runtime/MassEntity/Private/MassEntityManagerStorage.cpp
@@ -99,6 +99,13 @@ namespace UE::Mass
 			&& Entities[EntityHandle.Index].SerialNumber == EntityHandle.SerialNumber;
 	}
 
+	bool FSingleThreadedEntityStorage::IsEntityActive(FMassEntityHandle EntityHandle) const
+	{
+		return IsValidIndex(EntityHandle.Index)
+			&& GetSerialNumber(EntityHandle.Index) == EntityHandle.SerialNumber
+			&& GetEntityState(EntityHandle.Index) == UE::Mass::IEntityStorageInterface::EEntityState::Created;
+	}
+
 	SIZE_T FSingleThreadedEntityStorage::GetAllocatedSize() const
 	{
 		return Entities.GetAllocatedSize() + EntityFreeIndexList.GetAllocatedSize();
@@ -299,25 +306,29 @@ namespace UE::Mass
 		LookupEntity(Index).CurrentArchetype = Archetype;
 	}
 
-	IEntityStorageInterface::EEntityState FConcurrentEntityStorage::GetEntityState(int32 Index) const
+	IEntityStorageInterface::EEntityState FConcurrentEntityStorage::GetEntityStateInternal(const FEntityData& EntityData) const
 	{
 		//
 		// || Archetype || bIsAllocated || Result    |
 		//  |  nullptr   |      0       |  Free     |
 		//  |  nullptr   |      1       |  Reserved |
 		//  | !nullptr   |      1       |  Created  |
-		//
-	
-		const FEntityData& EntityData = LookupEntity(Index);
+		//	
 		if (EntityData.CurrentArchetype != nullptr)
 		{
 			return EEntityState::Created;
 		}
+
 		return EntityData.bIsAllocated
 			? EEntityState::Reserved
 			: EEntityState::Free;
 	}
 
+	IEntityStorageInterface::EEntityState FConcurrentEntityStorage::GetEntityState(int32 Index) const
+	{
+		return GetEntityStateInternal(LookupEntity(Index));
+	}
+
 	int32 FConcurrentEntityStorage::GetSerialNumber(int32 Index) const
 	{
 		return LookupEntity(Index).GenerationId;
@@ -340,6 +351,17 @@ namespace UE::Mass
 			&& LookupEntity(EntityHandle.Index).GetSerialNumber() == EntityHandle.SerialNumber;
 	}
 
+	bool FConcurrentEntityStorage::IsEntityActive(FMassEntityHandle EntityHandle) const
+	{
+		if (IsValidIndex(EntityHandle.Index))
+		{
+			const FEntityData& EntityData = LookupEntity(EntityHandle.Index);
+			return EntityData.GetSerialNumber() == EntityHandle.SerialNumber
+				&& GetEntityStateInternal(EntityData) == UE::Mass::IEntityStorageInterface::EEntityState::Created;
+		}
+		return false;
+	}
+
 	SIZE_T FConcurrentEntityStorage::GetAllocatedSize() const
 	{
 		const SIZE_T EntityFreeListSizeBytes = EntityFreeIndexList.GetAllocatedSize();
@@ -369,7 +391,7 @@ namespace UE::Mass
 
 		// Allocate new page
 		const uint32 NewPageIndex = PageCount;
-		checkf((NewPageIndex + 1) * MaxEntitiesPerPage < (1llu << MaximumEntityCountShift), TEXT("Exhausted number of entities"));
+		checkf(((NewPageIndex + 1) << MaxEntitiesPerPageShift) < (1u << MaximumEntityCountShift), TEXT("Exhausted number of entities"));
 
 		const uint64 PageSize = ComputePageSize();
 		FEntityData* Page = static_cast<FEntityData*>(FMemory::Malloc(PageSize, alignof(FEntityData)));
@@ -394,7 +416,7 @@ namespace UE::Mass
 		int32 NewEntityIndexStart;
 		if (LIKELY(NewPageIndex != 0))
 		{
-			NewEntityIndexStart = NewPageIndex * MaxEntitiesPerPage;
+			NewEntityIndexStart = NewPageIndex << MaxEntitiesPerPageShift;
 		}
 		else
 		{
@@ -405,7 +427,7 @@ namespace UE::Mass
 			++SentinelEntity->GenerationId;
 		}
 		
-		const int32 NewEntityIndexEnd = (NewPageIndex + 1) * MaxEntitiesPerPage;
+		const int32 NewEntityIndexEnd = (NewPageIndex + 1) << MaxEntitiesPerPageShift;
 
 		EntityFreeIndexList.Reserve(NewEntityIndexEnd - NewEntityIndexStart);
 
@@ -603,15 +625,15 @@ namespace UE::Mass
 
 	FConcurrentEntityStorage::FEntityData& FConcurrentEntityStorage::LookupEntity(int32 Index)
 	{
-		check(Index >= 0);
 		// PageIndex is which Page in the array of pages we need to access
 		const uint32 PageIndex = static_cast<uint32>(Index) >> MaxEntitiesPerPageShift;
 
 		// Convert the entity index into the index with respect to the page
-		const uint32 EntityOffset = (PageIndex * MaxEntitiesPerPage);
-		check(Index >= static_cast<int32>(EntityOffset)); // Check against negative values
+		const uint32 EntityOffset = (PageIndex << MaxEntitiesPerPageShift);
 		const uint32 InternalPageIndex = static_cast<uint32>(Index) - EntityOffset;
 
+		check((Index >= 0) && (Index >= static_cast<int32>(EntityOffset))); // Check against negative values;
+
 		// Pointer to start of page
 		FEntityData* PageStart = EntityPages[PageIndex];
 		FEntityData& EntityData = PageStart[InternalPageIndex];
@@ -625,7 +647,7 @@ namespace UE::Mass
 
 	uint64 FConcurrentEntityStorage::ComputePageSize() const
 	{
-		return sizeof(FEntityData) * MaxEntitiesPerPage;
+		return sizeof(FEntityData) << MaxEntitiesPerPageShift;
 	}
 
 #if WITH_MASSENTITY_DEBUG
diff --git a/Engine/Source/Runtime/MassEntity/Public/MassEntityManager.h b/Engine/Source/Runtime/MassEntity/Public/MassEntityManager.h
--- a/Engine/Source/Runtime/MassEntity/Public/MassEntityManager.h
+++ b/Engine/Source/Runtime/MassEntity/Public/MassEntityManager.h
@@ -532,10 +532,7 @@ public:
 	 * @return whether the given handle represents a valid and built entity
 	 *	 (i.e., the handle is valid and the entity represent has been constructed already)
 	 */
-	bool IsEntityActive(FMassEntityHandle EntityHandle) const 
-	{
-		return IsEntityValid(EntityHandle) && IsEntityBuilt(EntityHandle);
-	}
+	UE_API bool IsEntityActive(FMassEntityHandle EntityHandle) const;
 
 	/**
 	 * @return whether the given entity handle is valid, i.e. it points
@@ -553,10 +550,17 @@ public:
 	UE_API bool IsEntityReserved(FMassEntityHandle EntityHandle) const;
 
 	/** Asserts that IsEntityValid */
-	UE_API void CheckIfEntityIsValid(FMassEntityHandle EntityHandle) const;
+	inline void CheckIfEntityIsValid(FMassEntityHandle EntityHandle) const
+	{
+		checkf(IsEntityValid(EntityHandle), TEXT("Invalid entity (ID: %d, SN:%d, %s)"), EntityHandle.Index, EntityHandle.SerialNumber,
+			   (EntityHandle.Index == 0) ? TEXT("was never initialized") : TEXT("already destroyed"));
+	}
 
 	/** Asserts that IsEntityBuilt */
-	UE_API void CheckIfEntityIsActive(FMassEntityHandle EntityHandle) const;
+	inline void CheckIfEntityIsActive(FMassEntityHandle EntityHandle) const
+	{
+		checkf(IsEntityBuilt(EntityHandle), TEXT("Entity not yet created(ID: %d, SN:%d)"), EntityHandle.Index, EntityHandle.SerialNumber);
+	}
 
 	/**
 	 * Generate valid, up-to-date entity handle for the entity at given index.
diff --git a/Engine/Source/Runtime/MassEntity/Public/MassEntityManagerStorage.h b/Engine/Source/Runtime/MassEntity/Public/MassEntityManagerStorage.h
--- a/Engine/Source/Runtime/MassEntity/Public/MassEntityManagerStorage.h
+++ b/Engine/Source/Runtime/MassEntity/Public/MassEntityManagerStorage.h
@@ -88,6 +88,8 @@ namespace UE::Mass
 		 */
 		virtual bool IsValidHandle(FMassEntityHandle EntityHandle) const = 0;
 
+		virtual bool IsEntityActive(FMassEntityHandle EntityHandle) const = 0;
+
 		virtual SIZE_T GetAllocatedSize() const = 0;
 
 		/** Checks if entity at Index is built */
@@ -146,6 +148,7 @@ namespace UE::Mass
 		virtual int32 GetSerialNumber(int32 Index) const override;
 		virtual bool IsValidIndex(int32 Index) const override;
 		virtual bool IsValidHandle(FMassEntityHandle EntityHandle) const override;
+		virtual bool IsEntityActive(FMassEntityHandle EntityHandle) const override;
 		virtual SIZE_T GetAllocatedSize() const override;
 		virtual bool IsValid(int32 Index) const override;
 		virtual FMassEntityHandle AcquireOne() override;
@@ -209,6 +212,7 @@ namespace UE::Mass
 		virtual int32 GetSerialNumber(int32 Index) const override;
 		virtual bool IsValidIndex(int32 Index) const override;
 		virtual bool IsValidHandle(FMassEntityHandle EntityHandle) const override;
+		virtual bool IsEntityActive(FMassEntityHandle EntityHandle) const override;
 		virtual SIZE_T GetAllocatedSize() const override;
 		virtual bool IsValid(int32 Index) const override;
 		virtual FMassEntityHandle AcquireOne() override;
@@ -243,6 +247,8 @@ namespace UE::Mass
 			bool operator==(const FEntityData& Other) const;
 		};
 
+		EEntityState GetEntityStateInternal(const FEntityData& EntityData) const;
+
 		FEntityData& LookupEntity(int32 Index);
 		const FEntityData& LookupEntity(int32 Index) const;
 
