Minor optimization in FMassArchetypeData::MakeEntityHandle, don't go through TArray::operator[] twice for the same element.

diff --git a/Engine/Source/Runtime/MassEntity/Internal/MassArchetypeData.h b/Engine/Source/Runtime/MassEntity/Internal/MassArchetypeData.h
--- a/Engine/Source/Runtime/MassEntity/Internal/MassArchetypeData.h
+++ b/Engine/Source/Runtime/MassEntity/Internal/MassArchetypeData.h
@@ -435,9 +435,10 @@ public:
 	{
 		const int32 AbsoluteIndex = EntityMap.FindChecked(EntityIndex);
 		const int32 ChunkIndex = AbsoluteIndex / NumEntitiesPerChunk;
-	
-		return FMassEntityInChunkDataHandle(Chunks[ChunkIndex].GetRawMemory(), AbsoluteIndex % NumEntitiesPerChunk
-			, ChunkIndex, Chunks[ChunkIndex].GetSerialModificationNumber());
+		const FMassArchetypeChunk& Chunk = Chunks[ChunkIndex];
+
+		return FMassEntityInChunkDataHandle(Chunk.GetRawMemory(), AbsoluteIndex % NumEntitiesPerChunk
+			, ChunkIndex, Chunk.GetSerialModificationNumber());
 	}
 
 	FORCEINLINE FMassEntityInChunkDataHandle MakeEntityHandle(const FMassEntityHandle Entity) const
