[Mass] Trivial change to CHECK_SYNC_API_RETURN macro in MassEntityManager.cpp, now supporting dummy static struct instances.

#rnx

diff --git a/Engine/Source/Runtime/MassEntity/Private/MassEntityManager.cpp b/Engine/Source/Runtime/MassEntity/Private/MassEntityManager.cpp
--- a/Engine/Source/Runtime/MassEntity/Private/MassEntityManager.cpp
+++ b/Engine/Source/Runtime/MassEntity/Private/MassEntityManager.cpp
@@ -20,7 +20,7 @@
 #include "MassObserverNotificationTypes.h"
 
 #define CHECK_SYNC_API() testableCheckf(IsProcessing() == false, TEXT("Synchronous API function %hs called during mass processing. Use asynchronous API instead."), __FUNCTION__)
-#define CHECK_SYNC_API_RETURN(ReturnValue) testableCheckfReturn(IsProcessing() == false, return ReturnValue, TEXT("Synchronous API function %hs called during mass processing. Use asynchronous API instead."), __FUNCTION__)
+#define CHECK_SYNC_API_RETURN(ReturnValue) testableCheckfReturn(IsProcessing() == false, ReturnValue , TEXT("Synchronous API function %hs called during mass processing. Use asynchronous API instead."), __FUNCTION__)
 #define CHECK_ELEMENT(ElementType) checkf(UE::Mass::Private::IsElement(ElementType), TEXT("%hs: Only tags and fragments are considered 'elements', type provided: %s") \
 		, __FUNCTION__ , *ElementType->GetName())
 
@@ -582,7 +582,7 @@ void FMassEntityManager::DoEntityCompaction(const double TimeAllowed)
 
 FMassEntityHandle FMassEntityManager::CreateEntity(const FMassArchetypeHandle& ArchetypeHandle, const FMassArchetypeSharedFragmentValues& SharedFragmentValues)
 {
-	CHECK_SYNC_API_RETURN({});
+	CHECK_SYNC_API_RETURN(return {});
 	check(ArchetypeHandle.IsValid());
 	
 	MASS_BREAKPOINT(UE::Mass::Debug::FBreakpoint::CheckCreateEntityBreakpoints(ArchetypeHandle));
@@ -597,7 +597,7 @@ FMassEntityHandle FMassEntityManager::CreateEntity(const FMassArchetypeHandle& A
 
 FMassEntityHandle FMassEntityManager::CreateEntity(TConstArrayView<FInstancedStruct> FragmentInstanceList, const FMassArchetypeSharedFragmentValues& SharedFragmentValues, const FMassArchetypeCreationParams& CreationParams)
 {
-	CHECK_SYNC_API_RETURN({});
+	CHECK_SYNC_API_RETURN(return {});
 	check(FragmentInstanceList.Num() > 0);
 
 	const FMassArchetypeHandle& ArchetypeHandle = CreateArchetype(FMassArchetypeCompositionDescriptor(FragmentInstanceList,
@@ -688,7 +688,7 @@ int32 FMassEntityManager::BatchReserveEntities(TArrayView<FMassEntityHandle> InO
 TSharedRef<FMassEntityManager::FEntityCreationContext> FMassEntityManager::BatchBuildEntities(const FMassArchetypeEntityCollectionWithPayload& EncodedEntitiesWithPayload
 	, const FMassFragmentBitSet& FragmentsAffected, const FMassArchetypeSharedFragmentValues& SharedFragmentValues, const FMassArchetypeCreationParams& CreationParams)
 {
-	CHECK_SYNC_API_RETURN(FMassObserverManager::FCreationContext::DebugCreateDummyCreationContext());
+	CHECK_SYNC_API_RETURN(return FMassObserverManager::FCreationContext::DebugCreateDummyCreationContext());
 	check(SharedFragmentValues.IsSorted());
 
 	FMassArchetypeCompositionDescriptor Composition(FragmentsAffected, FMassTagBitSet(), FMassChunkFragmentBitSet(), FMassSharedFragmentBitSet(), FMassConstSharedFragmentBitSet());
@@ -708,7 +708,7 @@ TSharedRef<FMassEntityManager::FEntityCreationContext> FMassEntityManager::Batch
 	, const FMassArchetypeCompositionDescriptor& Composition
 	, const FMassArchetypeSharedFragmentValues& SharedFragmentValues, const FMassArchetypeCreationParams& CreationParams)
 {
-	CHECK_SYNC_API_RETURN(FMassObserverManager::FCreationContext::DebugCreateDummyCreationContext());
+	CHECK_SYNC_API_RETURN(return FMassObserverManager::FCreationContext::DebugCreateDummyCreationContext());
 
 	TRACE_CPUPROFILER_EVENT_SCOPE(Mass_BatchBuildEntities);
 
@@ -764,7 +764,7 @@ TSharedRef<FMassEntityManager::FEntityCreationContext> FMassEntityManager::Batch
 TSharedRef<FMassEntityManager::FEntityCreationContext> FMassEntityManager::BatchCreateReservedEntities(const FMassArchetypeHandle& ArchetypeHandle
 	, const FMassArchetypeSharedFragmentValues& SharedFragmentValues, TConstArrayView<FMassEntityHandle> ReservedEntities)
 {
-	CHECK_SYNC_API_RETURN(FMassObserverManager::FCreationContext::DebugCreateDummyCreationContext());
+	CHECK_SYNC_API_RETURN(return FMassObserverManager::FCreationContext::DebugCreateDummyCreationContext());
 	checkf(!ReservedEntities.IsEmpty(), TEXT("No reserved entities given to batch create."));
 
 	TRACE_CPUPROFILER_EVENT_SCOPE(Mass_BatchCreateReservedEntities);
@@ -777,7 +777,7 @@ TSharedRef<FMassEntityManager::FEntityCreationContext> FMassEntityManager::Batch
 TSharedRef<FMassEntityManager::FEntityCreationContext> FMassEntityManager::BatchCreateEntities(const FMassArchetypeHandle& ArchetypeHandle
 	, const FMassArchetypeSharedFragmentValues& SharedFragmentValues, const int32 Count, TArray<FMassEntityHandle>& InOutEntities)
 {
-	CHECK_SYNC_API_RETURN(FMassObserverManager::FCreationContext::DebugCreateDummyCreationContext());
+	CHECK_SYNC_API_RETURN(return FMassObserverManager::FCreationContext::DebugCreateDummyCreationContext());
 	testableCheckfReturn(ArchetypeHandle.IsValid(), return FMassObserverManager::FCreationContext::DebugCreateDummyCreationContext()
 		, TEXT("%hs expecting a valid ArchetypeHandle"), __FUNCTION__);
 
@@ -1416,7 +1416,7 @@ void FMassEntityManager::RemoveElementFromEntity(FMassEntityHandle Entity, TNotN
 
 bool FMassEntityManager::AddConstSharedFragmentToEntity(const FMassEntityHandle Entity, const FConstSharedStruct& InConstSharedFragment)
 {
-	CHECK_SYNC_API_RETURN(false);
+	CHECK_SYNC_API_RETURN(return false);
 
 	if (!ensureMsgf(InConstSharedFragment.IsValid(), TEXT("%hs parameter Fragment is expected to be valid"), __FUNCTION__))
 	{
@@ -1466,7 +1466,7 @@ bool FMassEntityManager::AddConstSharedFragmentToEntity(const FMassEntityHandle
 
 bool FMassEntityManager::RemoveConstSharedFragmentFromEntity(const FMassEntityHandle Entity, const UScriptStruct& ConstSharedFragmentType)
 {
-	CHECK_SYNC_API_RETURN(false);
+	CHECK_SYNC_API_RETURN(return false);
 
 	if (!ensureMsgf(UE::Mass::IsA<FMassConstSharedFragment>(&ConstSharedFragmentType), TEXT("%hs parameter ConstSharedFragmentType is expected to be a FMassConstSharedFragment"), __FUNCTION__))
 	{
@@ -1509,7 +1509,7 @@ bool FMassEntityManager::RemoveConstSharedFragmentFromEntity(const FMassEntityHa
 
 bool FMassEntityManager::AddSharedFragmentToEntity(const FMassEntityHandle Entity, const FSharedStruct& InSharedFragment)
 {
-	CHECK_SYNC_API_RETURN(false);
+	CHECK_SYNC_API_RETURN(return false);
 
 	if (!ensureMsgf(InSharedFragment.IsValid(), TEXT("%hs parameter Fragment is expected to be valid"), __FUNCTION__))
 	{
@@ -1559,7 +1559,7 @@ bool FMassEntityManager::AddSharedFragmentToEntity(const FMassEntityHandle Entit
 
 bool FMassEntityManager::RemoveSharedFragmentFromEntity(const FMassEntityHandle Entity, const UScriptStruct& SharedFragmentType)
 {
-	CHECK_SYNC_API_RETURN(false);
+	CHECK_SYNC_API_RETURN(return false);
 
 	if (!ensureMsgf(UE::Mass::IsA<FMassSharedFragment>(&SharedFragmentType), TEXT("%hs parameter SharedFragmentType is expected to be a FMassSharedFragment"), __FUNCTION__))
 	{
