Fix missing loop for map handling

#rnx

diff --git a/Engine/Source/Runtime/MassEntity/Private/MassObserverRegistry.cpp b/Engine/Source/Runtime/MassEntity/Private/MassObserverRegistry.cpp
--- a/Engine/Source/Runtime/MassEntity/Private/MassObserverRegistry.cpp
+++ b/Engine/Source/Runtime/MassEntity/Private/MassObserverRegistry.cpp
@@ -47,9 +47,9 @@ void UMassObserverRegistry::OnModulePackagesUnloaded(TConstArrayView<UPackage*>
 {
 	TRACE_CPUPROFILER_EVENT_SCOPE(UMassObserverRegistry::OnModulePackagesUnloaded);
 
-	const auto ProcessObserversMap = [&Packages](FObserverClassesMap* ObserversMap)
+	const auto ProcessObserversMap = [&Packages](FObserverClassesMap& ObserversMap)
 	{
-		for (auto MapIt = ObserversMap->CreateIterator(); MapIt; ++MapIt)
+		for (auto MapIt = ObserversMap.CreateIterator(); MapIt; ++MapIt)
 		{
 			TArray<FSoftClassPath>& ObserverClasses = MapIt->Value;
 
@@ -77,6 +77,9 @@ void UMassObserverRegistry::OnModulePackagesUnloaded(TConstArrayView<UPackage*>
 		}
 	};
 
-	ProcessObserversMap(FragmentObserverMaps);
-	ProcessObserversMap(TagObserverMaps);
+	for (uint8 OperationIndex = 0; OperationIndex < static_cast<uint8>(EMassObservedOperation::MAX); ++OperationIndex)
+	{
+		ProcessObserversMap(FragmentObserverMaps[OperationIndex]);
+		ProcessObserversMap(TagObserverMaps[OperationIndex]);
+	}
 }
