[Mass] comment fixes in newly added TMassFragmentTraits

#rnx

diff --git a/Engine/Source/Runtime/MassEntity/Public/MassExternalSubsystemTraits.h b/Engine/Source/Runtime/MassEntity/Public/MassExternalSubsystemTraits.h
--- a/Engine/Source/Runtime/MassEntity/Public/MassExternalSubsystemTraits.h
+++ b/Engine/Source/Runtime/MassEntity/Public/MassExternalSubsystemTraits.h
@@ -60,8 +60,8 @@ struct TMassFragmentTraits final
 	{
 		// Fragment types are best kept trivially copyable for performance reasons.
 		// To enforce that we test this trait when checking if a given type is a valid fragment type.
-		// This test can be skipped by specifically opting out, which also makes the 
-		// type author's acknowledge that they know and accept that the type contains unorthodox elements.
+		// This test can be skipped by specifically opting out, which also documents that 
+		// making the given type non-trivially-copyable was a deliberate decision. 
 		AuthorAcceptsItsNotTriviallyCopyable = false
 	};
 };
\ No newline at end of file
