[Mass] String literal fix

#rnx
[RN] trivial

diff --git a/Engine/Source/Runtime/MassEntity/Public/MassEntityElementTypes.h b/Engine/Source/Runtime/MassEntity/Public/MassEntityElementTypes.h
--- a/Engine/Source/Runtime/MassEntity/Public/MassEntityElementTypes.h
+++ b/Engine/Source/Runtime/MassEntity/Public/MassEntityElementTypes.h
@@ -12,8 +12,7 @@ struct FMassFragment
 };
 
 // these are the messages we'll print out when static checks whether a given type is a fragment fails
-#define _MASS_INVALID_FRAGMENT_CORE_MESSAGE "Make sure to inherit from FMassFragment or one of its child-types and ensure that the struct is trivially copyable, \
-or opt out by specializing TMassFragmentTraits for this type and setting AuthorAcceptsItsNotTriviallyCopyable = true"
+#define _MASS_INVALID_FRAGMENT_CORE_MESSAGE "Make sure to inherit from FMassFragment or one of its child-types and ensure that the struct is trivially copyable, or opt out by specializing TMassFragmentTraits for this type and setting AuthorAcceptsItsNotTriviallyCopyable = true"
 #define MASS_INVALID_FRAGMENT_MSG  "Given struct doesn't represent a valid fragment type." _MASS_INVALID_FRAGMENT_CORE_MESSAGE
 #define MASS_INVALID_FRAGMENT_MSG_F  "Type %s is not a valid fragment type." _MASS_INVALID_FRAGMENT_CORE_MESSAGE
 
