[Mass] Fixed relation participating entity destruction handling when the relation's roles are configured identically and CleanUp is used as the destruction policy.

#rnx

diff --git a/Engine/Source/Developer/MassEntityTestSuite/Private/MassRelationTest.cpp b/Engine/Source/Developer/MassEntityTestSuite/Private/MassRelationTest.cpp
--- a/Engine/Source/Developer/MassEntityTestSuite/Private/MassRelationTest.cpp
+++ b/Engine/Source/Developer/MassEntityTestSuite/Private/MassRelationTest.cpp
@@ -699,6 +699,43 @@ namespace UE::Mass::Test::Relations
 	};
 	IMPLEMENT_AI_INSTANT_TEST(FReRegisteringRelationType, "System.Mass.Relations.Type.Registration");
 
+	struct FSymmetricCleanup : FRelationTestBase
+	{
+		virtual void TweakRelation(FRelationTypeTraits& TempRelationTraits) override
+		{
+			TempRelationTraits.RoleTraits[0].DestructionPolicy = ERemovalPolicy::CleanUp;
+			TempRelationTraits.RoleTraits[1].DestructionPolicy = ERemovalPolicy::CleanUp;
+		}
+
+		virtual bool InstantTest() override
+		{
+			const FMassArchetypeHandle OriginalArchetype = IntsArchetype;
+
+			TArray<FMassEntityHandle> Entities;
+			EntityManager->BatchCreateEntities(OriginalArchetype, 2, Entities);
+
+			constexpr int32 ObjectIndex = 1;
+			constexpr int32 SubjecIndex = 0;
+			RelationshipManagerPtr->CreateRelationInstance(RelationTypeHandle, Entities[SubjecIndex], Entities[ObjectIndex]);
+
+			TArray<FMassEntityHandle> SubjectEntities = RelationshipManagerPtr->GetRelationSubjects(RelationTypeHandle, Entities[ObjectIndex]);
+			AITEST_EQUAL("Expected number of object's subjects", SubjectEntities.Num(), 1);
+			AITEST_TRUE("Object's subject meets expectations", SubjectEntities[0] == Entities[SubjecIndex]);
+
+			TArray<FMassEntityHandle> ObjectEntities = RelationshipManagerPtr->GetRelationObjects(RelationTypeHandle, Entities[SubjecIndex]);
+			AITEST_EQUAL("Expected number of subject's objects", ObjectEntities.Num(), 1);
+			AITEST_TRUE("Subject's object meets expectations", ObjectEntities[0] == Entities[ObjectIndex]);
+
+			EntityManager->DestroyEntity(Entities[ObjectIndex]);
+
+			ObjectEntities = RelationshipManagerPtr->GetRelationObjects(RelationTypeHandle, Entities[SubjecIndex]);
+			AITEST_EQUAL("After object's destruction: Expected number of subject's objects", ObjectEntities.Num(), 0);
+			
+			return true;
+		}
+	};
+	IMPLEMENT_AI_INSTANT_TEST(FSymmetricCleanup, "System.Mass.Relations.SymmetricCleanup");
+
 }// UE::Mass::Test::Relations
 
 UE_ENABLE_OPTIMIZATION_SHIP
diff --git a/Engine/Source/Runtime/MassEntity/Private/MassRelationObservers.cpp b/Engine/Source/Runtime/MassEntity/Private/MassRelationObservers.cpp
--- a/Engine/Source/Runtime/MassEntity/Private/MassRelationObservers.cpp
+++ b/Engine/Source/Runtime/MassEntity/Private/MassRelationObservers.cpp
@@ -38,20 +38,27 @@ void UMassRelationObserver::ConfigureQueries(const TSharedRef<FMassEntityManager
 		EntityQuery.AddRequirement(ObservedType, ObservedTypeAccess);
 	}
 
-	const UE::Mass::FRelationManager& RelationManager = EntityManager->GetRelationManager();
-	const UE::Mass::FRelationData& RelationData = RelationManager.GetRelationDataChecked(RelationTypeHandle);
-	const TNotNull<const UScriptStruct*> RelationTag = RelationData.Traits.GetRelationTagType();
-	if (RelationTag != ObservedType)
+	if (bAutoAddRelationFragmentRequirement || bAutoAddRelationTagRequirement)
 	{
-		EntityQuery.AddTagRequirement(*RelationTag, EMassFragmentPresence::All);
-	}
+		const UE::Mass::FRelationManager& RelationManager = EntityManager->GetRelationManager();
+		const UE::Mass::FRelationData& RelationData = RelationManager.GetRelationDataChecked(RelationTypeHandle);
 
-	if (bAutoAddRelationFragment)
-	{
-		const TNotNull<const UScriptStruct*> RelationFragmentType = RelationData.Traits.RelationFragmentType;
-		if (UE::Mass::IsA<FMassFragment>(RelationFragmentType) && RelationFragmentType != ObservedType)
+		if (bAutoAddRelationTagRequirement)
 		{
-			EntityQuery.AddRequirement(RelationFragmentType, RelationFragmentAccessType, EMassFragmentPresence::All);
+			const TNotNull<const UScriptStruct*> RelationTag = RelationData.Traits.GetRelationTagType();
+			if (RelationTag != ObservedType)
+			{
+				EntityQuery.AddTagRequirement(*RelationTag, EMassFragmentPresence::All);
+			}
+		}
+
+		if (bAutoAddRelationFragmentRequirement)
+		{
+			const TNotNull<const UScriptStruct*> RelationFragmentType = RelationData.Traits.RelationFragmentType;
+			if (UE::Mass::IsA<FMassFragment>(RelationFragmentType) && RelationFragmentType != ObservedType)
+			{
+				EntityQuery.AddRequirement(RelationFragmentType, RelationFragmentAccessType, EMassFragmentPresence::All);
+			}
 		}
 	}
 
@@ -173,7 +180,8 @@ void UMassRelationEntityDestruction::Execute(FMassEntityManager& EntityManager,
 UMassRelationRoleDestruction::UMassRelationRoleDestruction()
 {
 	ObservedOperations = EMassObservedOperationFlags::DestroyEntity;
-	bAutoAddRelationFragment = false;
+	bAutoAddRelationFragmentRequirement = false;
+	bAutoAddRelationTagRequirement = false;
 };
 
 void UMassRelationRoleDestruction::Execute(FMassEntityManager& EntityManager, FMassExecutionContext& Context)
@@ -188,7 +196,8 @@ void UMassRelationRoleDestruction::ConfigureQueries(const TSharedRef<FMassEntity
 	// @todo missing feature - this would cause the query to always fail to find archetypes. We need a way to
 	//		add "external requirements" to queries, that we'd use to calculate dependencies, but not use for binding
 
-	if (ExcludedRelationFragmentType)
+	if (ensureMsgf(ExcludedRelationFragmentType
+		, TEXT("We don't expect ExcludedRelationFragmentType to be null. Make sure ConfigureRelationObserver has been called first.")))
 	{
 		EntityQuery.AddRequirement(ExcludedRelationFragmentType, EMassFragmentAccess::None, EMassFragmentPresence::None);
 	}
@@ -196,12 +205,17 @@ void UMassRelationRoleDestruction::ConfigureQueries(const TSharedRef<FMassEntity
 
 bool UMassRelationRoleDestruction::ConfigureRelationObserver(UE::Mass::FTypeHandle InRegisteredTypeHandle, const UE::Mass::FRelationTypeTraits& Traits)
 {
-	check(RelationRole != UE::Mass::ERelationRole::MAX && static_cast<uint8>(UE::Mass::ERelationRole::MAX) == 2);
+	static_assert(static_cast<uint8>(UE::Mass::ERelationRole::MAX) == 2, "Current implementation relies on there being only two roles.");
 
-	const int32 RoleAsIndex = static_cast<int32>(RelationRole);
+	const int32 RoleAsIndex = RelationRole != UE::Mass::ERelationRole::MAX ? static_cast<int32>(RelationRole) : 0;
 	const int32 OppositeRoleIndex = (RoleAsIndex + 1) % 2;
 	bool bExecutionFunctionAssigned = false;
 
+	// this processor is specifically implemented for handling destruction of regular entities
+	// so we're using Traits.RelationFragmentType, which is only added to the relation entities,
+	// to filter those out. This will ensure we only get entities that played a Role in the given relationship.
+	ExcludedRelationFragmentType = Traits.RelationFragmentType;
+
 	// note there we're deliberately not calling Super::ConfigureRelationObserver, since we're going to observe role-specific elements
 	RelationTypeHandle = InRegisteredTypeHandle;
 	if (Traits.RoleTraits[RoleAsIndex].Element)
@@ -212,46 +226,104 @@ bool UMassRelationRoleDestruction::ConfigureRelationObserver(UE::Mass::FTypeHand
 	{
 		TNotNull<const UScriptStruct*> RelationTypeTag = RelationTypeHandle.GetScriptStruct();
 		ObservedType = RelationTypeTag;
-		ExcludedRelationFragmentType = Traits.RelationFragmentType;
 	}
 
 	switch (Traits.RoleTraits[RoleAsIndex].DestructionPolicy)
 	{
 		case UE::Mass::ERemovalPolicy::CleanUp:
 			RelationFragmentAccessType = EMassFragmentAccess::None;
-			ExecuteFunction = [InRegisteredTypeHandle, OppositeRoleIndex](FMassExecutionContext& ExecutionContext)
+			if (RelationRole != UE::Mass::Relations::ERelationRole::MAX)
 			{
-				UE::Mass::FRelationManager& RelationManager = ExecutionContext.GetEntityManagerChecked().GetRelationManager();
-				UE::Mass::FRelationData& RelationData = RelationManager.GetRelationDataChecked(InRegisteredTypeHandle);
-				for (FMassExecutionContext::FEntityIterator EntityIt = ExecutionContext.CreateEntityIterator(); EntityIt; ++EntityIt)
-				{
-					TArray<FMassRelationRoleInstanceHandle>& Container = RelationData.RoleMap.FindChecked(ExecutionContext.GetEntity(EntityIt))[OppositeRoleIndex];
-					ExecutionContext.Defer().DestroyEntities(RelationManager.GetRelationEntities(Container));
-					Container.Reset();
-				}
-			};
+				ExecuteFunction = [InRegisteredTypeHandle, OppositeRoleIndex](FMassExecutionContext& ExecutionContext)
+					{
+						UE::Mass::FRelationManager& RelationManager = ExecutionContext.GetEntityManagerChecked().GetRelationManager();
+						UE::Mass::FRelationData& RelationData = RelationManager.GetRelationDataChecked(InRegisteredTypeHandle);
+						for (FMassExecutionContext::FEntityIterator EntityIt = ExecutionContext.CreateEntityIterator(); EntityIt; ++EntityIt)
+						{
+							TArray<FMassRelationRoleInstanceHandle>& Container = RelationData.RoleMap.FindChecked(ExecutionContext.GetEntity(EntityIt))[OppositeRoleIndex];
+							ExecutionContext.Defer().DestroyEntities(RelationManager.GetRelationEntities(Container));
+							Container.Empty();
+						}
+					};
+			}
+			else
+			{
+				// this observer will only get called once and needs to handle both sides of the relation
+				ExecuteFunction = [InRegisteredTypeHandle](FMassExecutionContext& ExecutionContext)
+					{
+						UE::Mass::FRelationManager& RelationManager = ExecutionContext.GetEntityManagerChecked().GetRelationManager();
+						UE::Mass::FRelationData& RelationData = RelationManager.GetRelationDataChecked(InRegisteredTypeHandle);
+						for (FMassExecutionContext::FEntityIterator EntityIt = ExecutionContext.CreateEntityIterator(); EntityIt; ++EntityIt)
+						{
+							UE::Mass::FRelationData::FRelationInstanceData& InstanceData = RelationData.RoleMap.FindChecked(ExecutionContext.GetEntity(EntityIt));
+
+							TArray<FMassRelationRoleInstanceHandle>& SubjectsContainer = InstanceData[static_cast<uint8>(UE::Mass::ERelationRole::Subject)];
+							// this will instigate destruction of relations where Entity is the object
+							ExecutionContext.Defer().DestroyEntities(RelationManager.GetRelationEntities(SubjectsContainer));
+							SubjectsContainer.Empty();
+
+							TArray<FMassRelationRoleInstanceHandle>& ObjectsContainer = InstanceData[static_cast<uint8>(UE::Mass::ERelationRole::Object)];
+							// this will instigate destruction of relations where Entity is the subject
+							ExecutionContext.Defer().DestroyEntities(RelationManager.GetRelationEntities(ObjectsContainer));
+							SubjectsContainer.Empty();
+						}
+					};
+			}
 			bExecutionFunctionAssigned = true;
 			break;
 		case UE::Mass::ERemovalPolicy::Destroy:
 			RelationFragmentAccessType = EMassFragmentAccess::ReadOnly;
 			// destroy the other side of the relationship, and the relation entity
-			ExecuteFunction = [InRegisteredTypeHandle, RoleAsIndex = RoleAsIndex](FMassExecutionContext& ExecutionContext)
+			if (RelationRole != UE::Mass::Relations::ERelationRole::MAX)
 			{
-				const int32 OppositeRoleIndex = (RoleAsIndex + 1) % 2;
-				UE::Mass::FRelationManager& RelationManager = ExecutionContext.GetEntityManagerChecked().GetRelationManager();
-				UE::Mass::FRelationData& RelationData = RelationManager.GetRelationDataChecked(InRegisteredTypeHandle);
+				ExecuteFunction = [InRegisteredTypeHandle, OppositeRoleIndex](FMassExecutionContext& ExecutionContext)
+				{
+					UE::Mass::FRelationManager& RelationManager = ExecutionContext.GetEntityManagerChecked().GetRelationManager();
+					UE::Mass::FRelationData& RelationData = RelationManager.GetRelationDataChecked(InRegisteredTypeHandle);
 
-				TArray<FMassEntityHandle> EntitiesToDestroy;
+					TArray<FMassEntityHandle> EntitiesToDestroy;
 
-				for (FMassExecutionContext::FEntityIterator EntityIt = ExecutionContext.CreateEntityIterator(); EntityIt; ++EntityIt)
+					for (FMassExecutionContext::FEntityIterator EntityIt = ExecutionContext.CreateEntityIterator(); EntityIt; ++EntityIt)
+					{
+						TArray<FMassRelationRoleInstanceHandle>& Container = RelationData.RoleMap.FindChecked(ExecutionContext.GetEntity(EntityIt))[OppositeRoleIndex];
+						RelationManager.GetRelationEntities(Container, EntitiesToDestroy);
+						RelationManager.GetRoleEntities(Container, EntitiesToDestroy);
+						Container.Empty();
+					}
+
+					ExecutionContext.Defer().DestroyEntities(EntitiesToDestroy);
+				};
+			}
+			else
+			{
+				// this implementation handles both sides of relations
+				ExecuteFunction = [InRegisteredTypeHandle](FMassExecutionContext& ExecutionContext)
 				{
-					TArray<FMassRelationRoleInstanceHandle>& Container = RelationData.RoleMap.FindChecked(ExecutionContext.GetEntity(EntityIt))[OppositeRoleIndex];
-					RelationManager.GetRelationEntities(Container, EntitiesToDestroy);
-					RelationManager.GetRoleEntities(Container, EntitiesToDestroy);
-				}
+					UE::Mass::FRelationManager& RelationManager = ExecutionContext.GetEntityManagerChecked().GetRelationManager();
+					UE::Mass::FRelationData& RelationData = RelationManager.GetRelationDataChecked(InRegisteredTypeHandle);
 
-				ExecutionContext.Defer().DestroyEntities(EntitiesToDestroy);
-			};
+					TArray<FMassEntityHandle> EntitiesToDestroy;
+
+					for (FMassExecutionContext::FEntityIterator EntityIt = ExecutionContext.CreateEntityIterator(); EntityIt; ++EntityIt)
+					{
+						UE::Mass::FRelationData::FRelationInstanceData& InstanceData = RelationData.RoleMap.FindChecked(ExecutionContext.GetEntity(EntityIt));
+
+						TArray<FMassRelationRoleInstanceHandle>& SubjectsContainer = InstanceData[static_cast<uint8>(UE::Mass::ERelationRole::Subject)];
+						// this will instigate destruction of both the relation entities and the relation subjects
+						RelationManager.GetRelationEntities(SubjectsContainer, EntitiesToDestroy);
+						RelationManager.GetRoleEntities(SubjectsContainer, EntitiesToDestroy);
+						SubjectsContainer.Empty();
+
+						TArray<FMassRelationRoleInstanceHandle>& ObjectsContainer = InstanceData[static_cast<uint8>(UE::Mass::ERelationRole::Object)];
+						// this will instigate destruction of both the relation entities and the relation objects
+						RelationManager.GetRelationEntities(ObjectsContainer, EntitiesToDestroy);
+						RelationManager.GetRoleEntities(ObjectsContainer, EntitiesToDestroy);
+						ObjectsContainer.Empty();
+					}
+
+					ExecutionContext.Defer().DestroyEntities(EntitiesToDestroy);
+				};
+			}
 			bExecutionFunctionAssigned = true;
 			break;
 		case UE::Mass::ERemovalPolicy::Splice:
@@ -299,22 +371,35 @@ bool UMassRelationRoleDestruction::ConfigureRelationObserver(UE::Mass::FTypeHand
 
 void UMassRelationRoleDestruction::AddObserverInstances(FMassObserverManager& ObserverManager, UE::Mass::FTypeHandle InRegisteredTypeHandle, const UE::Mass::FRelationTypeTraits& Traits)
 {
-	bool NeedsSameObserver[static_cast<int32>(UE::Mass::ERelationRole::MAX)] = {
-		false
-		, (Traits.RoleTraits[1].Element == Traits.RoleTraits[0].Element && Traits.RoleTraits[1].DestructionPolicy == Traits.RoleTraits[0].DestructionPolicy)
-	};
+	static_assert(static_cast<int32>(UE::Mass::ERelationRole::MAX) == 2, "This implementation is tailored the there being only two roles to any relation");
 
-	for (int32 RoleIndex = 0; RoleIndex < static_cast<int32>(UE::Mass::ERelationRole::MAX); ++RoleIndex)
+	// is we both roles have the same policy and the same element, we can get away with creating only one observer
+	// as long as we configure it appropriately
+	if (Traits.RoleTraits[1].Element == Traits.RoleTraits[0].Element && Traits.RoleTraits[1].DestructionPolicy == Traits.RoleTraits[0].DestructionPolicy)
 	{
-		if (Traits.RoleTraits[RoleIndex].DestructionPolicy != UE::Mass::ERemovalPolicy::Custom
-		   && !NeedsSameObserver[RoleIndex])
+		if (Traits.RoleTraits[0].DestructionPolicy != UE::Mass::ERemovalPolicy::Custom)
 		{
 			UMassRelationRoleDestruction* ObserverProcessor = NewObject<UMassRelationRoleDestruction>();
-			ObserverProcessor->RelationRole = static_cast<UE::Mass::ERelationRole>(RoleIndex);
+			ObserverProcessor->RelationRole = UE::Mass::ERelationRole::MAX;
 			if (ObserverProcessor->ConfigureRelationObserver(InRegisteredTypeHandle, Traits))
 			{
 				ObserverManager.AddObserverInstance(ObserverProcessor);
 			}
 		}
 	}
+	else
+	{
+		for (int32 RoleIndex = 0; RoleIndex < static_cast<int32>(UE::Mass::ERelationRole::MAX); ++RoleIndex)
+		{
+			if (Traits.RoleTraits[RoleIndex].DestructionPolicy != UE::Mass::ERemovalPolicy::Custom)
+			{
+				UMassRelationRoleDestruction* ObserverProcessor = NewObject<UMassRelationRoleDestruction>();
+				ObserverProcessor->RelationRole = static_cast<UE::Mass::ERelationRole>(RoleIndex);
+				if (ObserverProcessor->ConfigureRelationObserver(InRegisteredTypeHandle, Traits))
+				{
+					ObserverManager.AddObserverInstance(ObserverProcessor);
+				}
+			}
+		}
+	}
 }
diff --git a/Engine/Source/Runtime/MassEntity/Public/MassRelationObservers.h b/Engine/Source/Runtime/MassEntity/Public/MassRelationObservers.h
--- a/Engine/Source/Runtime/MassEntity/Public/MassRelationObservers.h
+++ b/Engine/Source/Runtime/MassEntity/Public/MassRelationObservers.h
@@ -36,7 +36,8 @@ protected:
 	/** Relevant only when ObservedType is a fragment */
 	EMassFragmentAccess ObservedTypeAccess = EMassFragmentAccess::ReadWrite;
 	EMassFragmentAccess RelationFragmentAccessType = EMassFragmentAccess::ReadWrite;
-	bool bAutoAddRelationFragment = true;
+	bool bAutoAddRelationFragmentRequirement = true;
+	bool bAutoAddRelationTagRequirement = true;
 
 	FString DebugDescription;
 };
@@ -94,7 +95,16 @@ protected:
 	UE_API virtual void ConfigureQueries(const TSharedRef<FMassEntityManager>& EntityManager) override;
 
 	FMassExecuteFunction ExecuteFunction;
+
+	/**
+	 * indicates which side of the relation the given destructor handles. MAX indicates both sides.
+	 */
 	UE::Mass::Relations::ERelationRole RelationRole = UE::Mass::Relations::ERelationRole::MAX;
+
+	/**
+	 * indicates the relation-specific fragment, that characterizes relation-entities. We use this
+	 * information to filer these entities out - this observer is meant to only handle "role" entities
+	 */
 	const UScriptStruct* ExcludedRelationFragmentType = nullptr;
 };
 
