[Mass] fixed a silly mistake in new relations code

#rnx

diff --git a/Engine/Source/Runtime/MassEntity/Private/MassRelationObservers.cpp b/Engine/Source/Runtime/MassEntity/Private/MassRelationObservers.cpp
--- a/Engine/Source/Runtime/MassEntity/Private/MassRelationObservers.cpp
+++ b/Engine/Source/Runtime/MassEntity/Private/MassRelationObservers.cpp
@@ -265,7 +265,7 @@ bool UMassRelationRoleDestruction::ConfigureRelationObserver(UE::Mass::FTypeHand
 							TArray<FMassRelationRoleInstanceHandle>& ObjectsContainer = InstanceData[static_cast<uint8>(UE::Mass::ERelationRole::Object)];
 							// this will instigate destruction of relations where Entity is the subject
 							ExecutionContext.Defer().DestroyEntities(RelationManager.GetRelationEntities(ObjectsContainer));
-							SubjectsContainer.Empty();
+							ObjectsContainer.Empty();
 						}
 					};
 			}
