[Mass] Removing CDOs from Mass settings on modules' unload

diff --git a/Engine/Source/Runtime/MassEntity/Private/MassEntitySettings.cpp b/Engine/Source/Runtime/MassEntity/Private/MassEntitySettings.cpp
--- a/Engine/Source/Runtime/MassEntity/Private/MassEntitySettings.cpp
+++ b/Engine/Source/Runtime/MassEntity/Private/MassEntitySettings.cpp
@@ -4,8 +4,11 @@
 #include "MassProcessingPhaseManager.h"
 #include "VisualLogger/VisualLogger.h"
 #include "UObject/UObjectHash.h"
+#include "UObject/UObjectGlobals.h"
 #include "Misc/CoreDelegates.h"
 #include "MassArchetypeData.h"
+#include "Algo/RemoveIf.h"
+
 
 #include UE_INLINE_GENERATED_CPP_BY_NAME(MassEntitySettings)
 
@@ -26,6 +29,10 @@ UMassEntitySettings::UMassEntitySettings(const FObjectInitializer& ObjectInitial
 	}
 
 	FCoreDelegates::OnPostEngineInit.AddUObject(this, &UMassEntitySettings::OnPostEngineInit);
+
+	// we need to get notified about modules being unloaded (like Game Feature Plugins) so that we can remove
+	// stored CDOs originating from the modules being removed.
+	FCoreUObjectDelegates::CompiledInUObjectsRemovedDelegate.AddUObject(this, &UMassEntitySettings::OnModulePackagesUnloaded);
 }
 
 void UMassEntitySettings::PostInitProperties()
@@ -37,6 +44,7 @@ void UMassEntitySettings::PostInitProperties()
 void UMassEntitySettings::BeginDestroy()
 {
 	FCoreDelegates::OnPostEngineInit.RemoveAll(this);
+	FCoreUObjectDelegates::CompiledInUObjectsRemovedDelegate.RemoveAll(this);
 	Super::BeginDestroy();
 }
 
@@ -46,6 +54,16 @@ void UMassEntitySettings::OnPostEngineInit()
 	BuildProcessorListAndPhases();
 }
 
+void UMassEntitySettings::OnModulePackagesUnloaded(TConstArrayView<UPackage*> Packages)
+{
+	ProcessorCDOs.SetNum(Algo::RemoveIf(ProcessorCDOs,
+		[&Packages](const TObjectPtr<UMassProcessor>& CDO)
+		{
+			return !CDO || Packages.Contains(CDO->GetPackage());
+		})
+	);
+}
+
 void UMassEntitySettings::BuildProcessorListAndPhases()
 {
 	if (bInitialized == true || bEngineInitialized == false)
diff --git a/Engine/Source/Runtime/MassEntity/Public/MassEntitySettings.h b/Engine/Source/Runtime/MassEntity/Public/MassEntitySettings.h
--- a/Engine/Source/Runtime/MassEntity/Public/MassEntitySettings.h
+++ b/Engine/Source/Runtime/MassEntity/Public/MassEntitySettings.h
@@ -55,6 +55,9 @@ protected:
 	MASSENTITY_API void BuildPhases();
 	MASSENTITY_API void BuildProcessorList();
 
+private:
+	void OnModulePackagesUnloaded(TConstArrayView<UPackage*> Packages);
+
 public:
 	UPROPERTY(EditDefaultsOnly, Category = Mass, config, AdvancedDisplay)
 	uint32 ChunkMemorySize = 128 * 1024;
