[Mass] fix to processor CDOs removal from Mass settings on modules unload. The previous change left CDOs in phase configs.

#rnx

diff --git a/Engine/Source/Runtime/MassEntity/Private/MassEntitySettings.cpp b/Engine/Source/Runtime/MassEntity/Private/MassEntitySettings.cpp
--- a/Engine/Source/Runtime/MassEntity/Private/MassEntitySettings.cpp
+++ b/Engine/Source/Runtime/MassEntity/Private/MassEntitySettings.cpp
@@ -56,12 +56,33 @@ void UMassEntitySettings::OnPostEngineInit()
 
 void UMassEntitySettings::OnModulePackagesUnloaded(TConstArrayView<UPackage*> Packages)
 {
+	const int32 InitialNum = ProcessorCDOs.Num();
+
 	ProcessorCDOs.SetNum(Algo::RemoveIf(ProcessorCDOs,
 		[&Packages](const TObjectPtr<UMassProcessor>& CDO)
 		{
 			return !CDO || Packages.Contains(CDO->GetPackage());
 		})
 	);
+
+	if (ProcessorCDOs.Num() != InitialNum)
+	{
+		// rebuild the phase configs
+		for (FMassProcessingPhaseConfig& PhaseConfig : ProcessingPhasesConfig)
+		{
+			PhaseConfig.ProcessorCDOs.Reset();
+		}
+
+		for (TObjectPtr<UMassProcessor> ProcessorCDO : ProcessorCDOs)
+		{
+			if (ProcessorCDO->ShouldAutoAddToGlobalList())
+			{
+				ProcessingPhasesConfig[static_cast<int32>(ProcessorCDO->GetProcessingPhase())].ProcessorCDOs.Add(ProcessorCDO);
+			}
+		}
+
+		BuildPhases();
+	}
 }
 
 void UMassEntitySettings::BuildProcessorListAndPhases()
