[MassDebugger]
Fixed crash switching environments with "Auto Update Entity Data" enabled on the "Entities" tab.

#Synced-CL 48355042

diff --git a/Engine/Source/Runtime/MassEntity/Private/MassDebugger.cpp b/Engine/Source/Runtime/MassEntity/Private/MassDebugger.cpp
--- a/Engine/Source/Runtime/MassEntity/Private/MassDebugger.cpp
+++ b/Engine/Source/Runtime/MassEntity/Private/MassDebugger.cpp
@@ -1417,6 +1417,10 @@ TSharedPtr<FStructOnScope> FMassDebugger::GetFragmentData(const FMassEntityManag
 
 bool FMassDebugger::GetFragmentData(const FMassEntityManager& EntityManager, const UScriptStruct* FragmentType, FMassEntityHandle Entity, TSharedPtr<FStructOnScope>& OutStructData)
 {
+	if (!EntityManager.IsEntityValid(Entity))
+	{
+		return false;
+	}
 	TSharedPtr<FMassArchetypeData> Archetype = EntityManager.DebugGetEntityStorageInterface().GetArchetypeAsShared(Entity.Index);
 	if (Archetype.IsValid())
 	{
@@ -1438,12 +1442,15 @@ bool FMassDebugger::GetFragmentData(const FMassEntityManager& EntityManager, con
 
 const FMassArchetypeSharedFragmentValues& FMassDebugger::GetSharedFragmentValues(const FMassEntityManager& EntityManager, FMassEntityHandle Entity)
 {
-	TSharedPtr<FMassArchetypeData> Archetype = EntityManager.DebugGetEntityStorageInterface().GetArchetypeAsShared(Entity.Index);
-	if (Archetype.IsValid())
+	if (EntityManager.IsEntityValid(Entity))
 	{
-		return Archetype->GetSharedFragmentValues(Entity);
+		TSharedPtr<FMassArchetypeData> Archetype = EntityManager.DebugGetEntityStorageInterface().GetArchetypeAsShared(Entity.Index);
+		if (Archetype.IsValid())
+		{
+			return Archetype->GetSharedFragmentValues(Entity);
+		}
 	}
-
+	
 	static FMassArchetypeSharedFragmentValues Dummy;
 	return Dummy;
 }
@@ -1460,6 +1467,10 @@ TSharedPtr<FStructOnScope> FMassDebugger::GetSharedFragmentData(const FMassEntit
 
 bool FMassDebugger::GetSharedFragmentData(const FMassEntityManager& EntityManager, const UScriptStruct* FragmentType, FMassEntityHandle Entity, TSharedPtr<FStructOnScope>& OutStructData)
 {
+	if (!EntityManager.IsEntityValid(Entity))
+	{
+		return false;
+	}
 	TSharedPtr<FMassArchetypeData> Archetype = EntityManager.DebugGetEntityStorageInterface().GetArchetypeAsShared(Entity.Index);
 	if (Archetype.IsValid())
 	{
@@ -1493,6 +1504,10 @@ TSharedPtr<FStructOnScope> FMassDebugger::GetConstSharedFragmentData(const FMass
 
 bool FMassDebugger::GetConstSharedFragmentData(const FMassEntityManager& EntityManager, const UScriptStruct* FragmentType, FMassEntityHandle Entity, TSharedPtr<FStructOnScope>& OutStructData)
 {
+	if (!EntityManager.IsEntityValid(Entity))
+	{
+		return false;
+	}
 	TSharedPtr<FMassArchetypeData> Archetype = EntityManager.DebugGetEntityStorageInterface().GetArchetypeAsShared(Entity.Index);
 	if (Archetype.IsValid())
 	{
