[Mass]
Added MT lock to shared fragment value hash to avoid contention on lazy hash calculations that run in parallel.

diff --git a/Engine/Source/Runtime/MassEntity/Public/MassEntityTypes.h b/Engine/Source/Runtime/MassEntity/Public/MassEntityTypes.h
--- a/Engine/Source/Runtime/MassEntity/Public/MassEntityTypes.h
+++ b/Engine/Source/Runtime/MassEntity/Public/MassEntityTypes.h
@@ -6,6 +6,8 @@
 #include "MassProcessingTypes.h"
 #include "StructUtils/StructArrayView.h"
 #include "Subsystems/Subsystem.h"
+#include "Misc/ScopeRWLock.h"
+#include "Misc/TransactionallySafeRWLock.h"
 #if UE_ENABLE_INCLUDE_ORDER_DEPRECATED_IN_5_6
 #include "MassExternalSubsystemTraits.h"
 #endif // UE_ENABLE_INCLUDE_ORDER_DEPRECATED_IN_5_6
@@ -296,10 +298,58 @@ PRAGMA_ENABLE_DEPRECATION_WARNINGS
 struct FMassArchetypeSharedFragmentValues
 {
 	FMassArchetypeSharedFragmentValues() = default;
-	FMassArchetypeSharedFragmentValues(const FMassArchetypeSharedFragmentValues& OtherFragmentValues) = default;
-	FMassArchetypeSharedFragmentValues(FMassArchetypeSharedFragmentValues&& OtherFragmentValues) = default;
-	FMassArchetypeSharedFragmentValues& operator=(const FMassArchetypeSharedFragmentValues& OtherFragmentValues) = default;
-	FMassArchetypeSharedFragmentValues& operator=(FMassArchetypeSharedFragmentValues&& OtherFragmentValues) = default;
+	
+	FMassArchetypeSharedFragmentValues(const FMassArchetypeSharedFragmentValues& Other)
+		: HashCache(Other.HashCache)
+		, bSorted(Other.bSorted)
+		, SharedFragmentBitSet(Other.SharedFragmentBitSet)
+		, ConstSharedFragmentBitSet(Other.ConstSharedFragmentBitSet)
+		, ConstSharedFragments(Other.ConstSharedFragments)
+		, SharedFragments(Other.SharedFragments)
+	{
+	}
+
+	FMassArchetypeSharedFragmentValues(FMassArchetypeSharedFragmentValues&& Other)
+		: HashCache(Other.HashCache)
+		, bSorted(Other.bSorted)
+		, SharedFragmentBitSet(MoveTemp(Other.SharedFragmentBitSet))
+		, ConstSharedFragmentBitSet(MoveTemp(Other.ConstSharedFragmentBitSet))
+		, ConstSharedFragments(MoveTemp(Other.ConstSharedFragments))
+		, SharedFragments(MoveTemp(Other.SharedFragments))
+	{
+	}
+
+	FMassArchetypeSharedFragmentValues& operator=(const FMassArchetypeSharedFragmentValues& Other)
+	{
+		if (this != &Other)
+		{
+			UE::TWriteScopeLock WriteLock{ HashLock };
+
+			SharedFragmentBitSet = Other.SharedFragmentBitSet;
+			ConstSharedFragmentBitSet = Other.ConstSharedFragmentBitSet;
+			ConstSharedFragments = Other.ConstSharedFragments;
+			SharedFragments = Other.SharedFragments;
+			HashCache = Other.HashCache;
+			bSorted = Other.bSorted;
+		}
+		return *this;
+	}
+
+	FMassArchetypeSharedFragmentValues& operator=(FMassArchetypeSharedFragmentValues&& Other) noexcept
+	{
+		if (this != &Other)
+		{
+			UE::TWriteScopeLock WriteLock{ HashLock };
+
+			SharedFragmentBitSet = MoveTemp(Other.SharedFragmentBitSet);
+			ConstSharedFragmentBitSet = MoveTemp(Other.ConstSharedFragmentBitSet);
+			ConstSharedFragments = MoveTemp(Other.ConstSharedFragments);
+			SharedFragments = MoveTemp(Other.SharedFragments);
+			HashCache = Other.HashCache;
+			bSorted = Other.bSorted;
+		}
+		return *this;
+	}
 
 	inline bool HasExactFragmentTypesMatch(const FMassSharedFragmentBitSet& InSharedFragmentBitSet, const FMassConstSharedFragmentBitSet& InConstSharedFragmentBitSet) const
 	{
@@ -509,6 +559,7 @@ struct FMassArchetypeSharedFragmentValues
 
 	inline void DirtyHashCache()
 	{
+		UE::TWriteScopeLock ScopeLock{ HashLock };
 		HashCache = UINT32_MAX;
 		// we consider a single shared fragment as being "sorted"
 		bSorted = (SharedFragments.Num() + ConstSharedFragments.Num() <= 1) ;
@@ -516,16 +567,30 @@ struct FMassArchetypeSharedFragmentValues
 
 	inline void CacheHash() const
 	{
-		if (HashCache == UINT32_MAX)
 		{
-			HashCache = CalculateHash();
+			UE::TReadScopeLock ReadLock{ HashLock };
+			if (HashCache != UINT32_MAX)
+			{
+				return;
+			}
+		}
+
+		{
+			UE::TWriteScopeLock WriteLock{ HashLock };
+			if (HashCache == UINT32_MAX)
+			{
+				HashCache = CalculateHash();
+			}
 		}
 	}
 
 	friend inline uint32 GetTypeHash(const FMassArchetypeSharedFragmentValues& SharedFragmentValues)
 	{
 		SharedFragmentValues.CacheHash();
-		return SharedFragmentValues.HashCache;
+		{
+			UE::TReadScopeLock ReadLock{ SharedFragmentValues.HashLock };
+			return SharedFragmentValues.HashCache;
+		}
 	}
 
 	MASSENTITY_API uint32 CalculateHash() const;
@@ -560,6 +625,9 @@ protected:
 	TArray<FConstSharedStruct> ConstSharedFragments;
 	TArray<FSharedStruct> SharedFragments;
 
+private:
+	mutable FTransactionallySafeRWLock HashLock;
+
 public:
 	//-----------------------------------------------------------------------------
 	// DEPRECATED
@@ -1000,6 +1068,8 @@ inline bool FMassArchetypeSharedFragmentValues::IsEmpty() const
 
 inline void FMassArchetypeSharedFragmentValues::Reset()
 {
+	UE::TWriteScopeLock ScopeLock{ HashLock };
+
 	HashCache = UINT32_MAX;
 	bSorted = false; 
 	SharedFragmentBitSet.Reset();
