[Mass]
Remove extraneous conditional in shared fragment hash caching

diff --git a/Engine/Source/Runtime/MassEntity/Public/MassEntityTypes.h b/Engine/Source/Runtime/MassEntity/Public/MassEntityTypes.h
--- a/Engine/Source/Runtime/MassEntity/Public/MassEntityTypes.h
+++ b/Engine/Source/Runtime/MassEntity/Public/MassEntityTypes.h
@@ -577,10 +577,7 @@ struct FMassArchetypeSharedFragmentValues
 
 		{
 			UE::TWriteScopeLock WriteLock{ HashLock };
-			if (HashCache == UINT32_MAX)
-			{
-				HashCache = CalculateHash();
-			}
+			HashCache = CalculateHash();
 		}
 	}
 
