[MassEntity] added details when processor or query subsystem requirements are not met
- UE::Mass::ErrorReporting::FScopedOutputDevice can be used to provide a threadsafe global output device to allow nested call to report errors
- UE::Mass::ErrorReporting::GetOutputDevice() can be use to report the error message

#rnx

diff --git a/Engine/Source/Runtime/MassEntity/Private/MassEntityQuery.cpp b/Engine/Source/Runtime/MassEntity/Private/MassEntityQuery.cpp
--- a/Engine/Source/Runtime/MassEntity/Private/MassEntityQuery.cpp
+++ b/Engine/Source/Runtime/MassEntity/Private/MassEntityQuery.cpp
@@ -43,7 +43,22 @@ struct FScopedEntityQueryContext
 	{
 		CachedExecutionContext.PushQuery(InQuery);
 
+		// Setup output device to collect errors that can be reported by CacheSubsystemRequirements
+		FStringOutputDevice ErrorOutput;
+		UE::Mass::ErrorReporting::FScopedOutputDevice ScopedOutput(ErrorOutput);
 		bSubsystemRequirementsCached = CachedExecutionContext.CacheSubsystemRequirements(InQuery);
+		if (!bSubsystemRequirementsCached)
+		{
+#if WITH_MASSENTITY_DEBUG
+			UE_VLOG_UELOG(ExecutionContext.GetSharedEntityManager().Get().GetOwner(), LogMass, Log
+				, TEXT("Mass query '%s' can't be successfully set up due to subsystem requirements not being met: %s")
+				, *ExecutionContext.DebugGetExecutionDesc()
+				, *ErrorOutput);
+#else
+			UE_VLOG_UELOG(ExecutionContext.GetSharedEntityManager().Get().GetOwner(), LogMass, Log
+				, TEXT("Mass query can't be successfully set up due to subsystem requirements not being met: %s"), *ErrorOutput);
+#endif
+		}
 	}
 
 	~FScopedEntityQueryContext()
diff --git a/Engine/Source/Runtime/MassEntity/Private/MassProcessor.cpp b/Engine/Source/Runtime/MassEntity/Private/MassProcessor.cpp
--- a/Engine/Source/Runtime/MassEntity/Private/MassProcessor.cpp
+++ b/Engine/Source/Runtime/MassEntity/Private/MassProcessor.cpp
@@ -277,17 +277,27 @@ void UMassProcessor::CallExecute(FMassEntityManager& EntityManager, FMassExecuti
 		Context.DebugSetExecutionDesc(DebugDescription);
 		Context.DebugSetProcessor(this);
 #endif
-		// CacheSubsystemRequirements will return true only if all requirements declared with ProcessorRequirements are met
-		// meaning if it fails there's no point in calling Execute.
-		// Note that we're not testing individual queries in OwnedQueries - processors can function just fine with some 
-		// of their queries not having anything to do.
-		if (Context.CacheSubsystemRequirements(ProcessorRequirements))
+
+		bool bAllRequirementsMet = false;
 		{
-			Execute(EntityManager, Context);
+			// Setup output device to collect errors that can be reported by CacheSubsystemRequirements
+			FStringOutputDevice ErrorOutput;
+			UE::Mass::ErrorReporting::FScopedOutputDevice ScopedOutput(ErrorOutput);
+
+			// CacheSubsystemRequirements will return true only if all requirements declared with ProcessorRequirements are met
+			// meaning if it fails there's no point in calling Execute.
+			// Note that we're not testing individual queries in OwnedQueries - processors can function just fine with some
+			// of their queries not having anything to do.
+			bAllRequirementsMet = Context.CacheSubsystemRequirements(ProcessorRequirements);
+			if (!bAllRequirementsMet)
+			{
+				UE_VLOG_UELOG(this, LogMass, Log, TEXT("%s Skipping Execute due to subsystem requirements not being met: %s"), *GetProcessorName(), *ErrorOutput);
+			}
 		}
-		else
+
+		if (bAllRequirementsMet)
 		{
-			UE_VLOG_UELOG(this, LogMass, VeryVerbose, TEXT("%s Skipping Execute due to subsystem requirements not being met"), *GetProcessorName());
+			Execute(EntityManager, Context);
 		}
 
 		if (ActivationState == EActivationState::OneShot)
@@ -733,4 +743,4 @@ void UMassProcessor::Initialize(UObject& Owner)
 void UMassCompositeProcessor::SetChildProcessors(TArray<UMassProcessor*>&& InProcessors)
 {
 	SetChildProcessors(MakeArrayView(InProcessors.GetData(), InProcessors.Num()));
-}
\ No newline at end of file
+}
diff --git a/Engine/Source/Runtime/MassEntity/Private/MassRequirements.cpp b/Engine/Source/Runtime/MassEntity/Private/MassRequirements.cpp
--- a/Engine/Source/Runtime/MassEntity/Private/MassRequirements.cpp
+++ b/Engine/Source/Runtime/MassEntity/Private/MassRequirements.cpp
@@ -5,6 +5,7 @@
 #include "MassProcessorDependencySolver.h"
 #include "MassTypeManager.h"
 #include "MassEntityManager.h"
+#include "Misc/OutputDeviceNull.h"
 
 #include UE_INLINE_GENERATED_CPP_BY_NAME(MassRequirements)
 
@@ -12,6 +13,28 @@
 #include "MassRequirementAccessDetector.h"
 #endif // WITH_MASSENTITY_DEBUG
 
+namespace UE::Mass::ErrorReporting
+{
+thread_local FOutputDevice* GOutputDevice = nullptr;
+FOutputDeviceNull GNullOutput;
+
+FOutputDevice& GetOutputDevice()
+{
+	return GOutputDevice ? *GOutputDevice : GNullOutput;
+}
+
+FScopedOutputDevice::FScopedOutputDevice(FOutputDevice& OutputDevice)
+{
+	ActiveOutputDevice = GOutputDevice;
+	GOutputDevice = &OutputDevice;
+}
+
+FScopedOutputDevice::~FScopedOutputDevice()
+{
+	GOutputDevice = ActiveOutputDevice;
+}
+
+} // UE::Mass::ErrorReporting
 
 namespace UE::Mass::Private
 {
diff --git a/Engine/Source/Runtime/MassEntity/Private/MassSubsystemAccess.cpp b/Engine/Source/Runtime/MassEntity/Private/MassSubsystemAccess.cpp
--- a/Engine/Source/Runtime/MassEntity/Private/MassSubsystemAccess.cpp
+++ b/Engine/Source/Runtime/MassEntity/Private/MassSubsystemAccess.cpp
@@ -104,8 +104,14 @@ bool FMassSubsystemAccess::CacheSubsystem(const uint32 SystemIndex)
 
 	if (SubsystemSubclass)
 	{
-		USubsystem* SystemInstance = FMassSubsystemAccess::FetchSubsystemInstance(World.Get(), SubsystemSubclass);
+		USubsystem* SystemInstance = FetchSubsystemInstance(World.Get(), SubsystemSubclass);
 		Subsystems[SystemIndex] = SystemInstance;
+
+		if (SystemInstance == nullptr)
+		{
+			UE::Mass::ErrorReporting::GetOutputDevice().Logf(TEXT("%s "), *SubsystemClass->GetName());
+		}
+
 		return SystemInstance != nullptr;
 	}
 
diff --git a/Engine/Source/Runtime/MassEntity/Public/MassRequirements.h b/Engine/Source/Runtime/MassEntity/Public/MassRequirements.h
--- a/Engine/Source/Runtime/MassEntity/Public/MassRequirements.h
+++ b/Engine/Source/Runtime/MassEntity/Public/MassRequirements.h
@@ -13,6 +13,26 @@ struct FMassExecutionRequirements;
 struct FMassRequirementAccessDetector;
 class USubsystem;
 
+namespace UE::Mass::ErrorReporting
+{
+extern FOutputDevice& GetOutputDevice();
+
+/**
+ * Helper struct that can be used to set up a threadsafe global output device
+ * so nested calls can use UE::Mass::ErrorReporting::GetOutputDevice() to report
+ * error messages.
+ */
+struct FScopedOutputDevice
+{
+	explicit FScopedOutputDevice(FOutputDevice& OutputDevice);
+	~FScopedOutputDevice();
+
+private:
+	FOutputDevice* ActiveOutputDevice = nullptr;
+};
+
+} //UE::Mass::ErrorReporting
+
 UENUM()
 enum class EMassFragmentAccess : uint8
 {
