Fixed Mass FQueryExecutor's FMutableSharedFragmentAccess using wrong accessor type

diff --git a/Engine/Source/Runtime/MassEntity/Public/MassQueryExecutor.h b/Engine/Source/Runtime/MassEntity/Public/MassQueryExecutor.h
--- a/Engine/Source/Runtime/MassEntity/Public/MassQueryExecutor.h
+++ b/Engine/Source/Runtime/MassEntity/Public/MassQueryExecutor.h
@@ -465,7 +465,7 @@ struct FMutableSharedFragmentAccess
 
 	inline void SetupForChunk(FMassExecutionContext& Context)
 	{
-		Fragment = &(Context.GetSharedFragment<TFragment>());
+		Fragment = &(Context.GetMutableSharedFragment<TFragment>());
 	}
 
 	TFragment* Fragment = nullptr;
