Fixed typos in MassSettings doc

diff --git a/Engine/Source/Runtime/MassEntity/Public/MassSettings.h b/Engine/Source/Runtime/MassEntity/Public/MassSettings.h
--- a/Engine/Source/Runtime/MassEntity/Public/MassSettings.h
+++ b/Engine/Source/Runtime/MassEntity/Public/MassSettings.h
@@ -9,8 +9,9 @@
 
 
 /** 
- * A common parrent for Mass's per-module settings. Classes extending this class will automatically get registered 
- * with- and show under Mass settings in Project Settings.
+ * A common parent for Mass's per-module settings.
+ * Classes extending this class will automatically get registered
+ * with the Mass Settings and show under its section in Project Settings.
  */
 UCLASS(MinimalAPI, Abstract, config = Mass, defaultconfig, collapseCategories)
 class UMassModuleSettings : public UObject
