Deprecate GET_MASS_CONFIG_VALUE
- There's no need to hide use of UMassEntitySettings. This is unnecessary boilerplate and it leads to duplicate lookups of the class instance when there multiple uses of the class.

Deprecate static methods of UMassEntitySettings.
- Again, this leads to duplicate lookups.
- Unfortunately we can't have both a static and instance method with the same name so this needs to go through a two step deprecation. The instance version is going to have an awkward name for 2 minor versions.

diff --git a/Engine/Plugins/Runtime/MassGameplay/Source/MassSimulation/Private/MassSimulationSubsystem.cpp b/Engine/Plugins/Runtime/MassGameplay/Source/MassSimulation/Private/MassSimulationSubsystem.cpp
--- a/Engine/Plugins/Runtime/MassGameplay/Source/MassSimulation/Private/MassSimulationSubsystem.cpp
+++ b/Engine/Plugins/Runtime/MassGameplay/Source/MassSimulation/Private/MassSimulationSubsystem.cpp
@@ -186,12 +186,12 @@ void UMassSimulationSubsystem::HandleSimulationTickingEnabledCVarChange(IConsole
 
 void UMassSimulationSubsystem::RebuildTickPipeline()
 {
-	TConstArrayView<FMassProcessingPhaseConfig> ProcessingPhasesConfig = GET_MASS_CONFIG_VALUE(GetProcessingPhasesConfig());
+	UMassEntitySettings* Settings = GetMutableDefault<UMassEntitySettings>();
+	TConstArrayView<FMassProcessingPhaseConfig> ProcessingPhasesConfig = Settings->GetProcessingPhasesConfig();
 	FString DependencyGraphFileName;
 
 #if WITH_EDITOR
 	const UWorld* World = GetWorld();
-	const UMassEntitySettings* Settings = GetMutableDefault<UMassEntitySettings>();
 	if (World != nullptr && Settings != nullptr && !Settings->DumpDependencyGraphFileName.IsEmpty())
 	{
 		DependencyGraphFileName = FString::Printf(TEXT("%s_%s"), *Settings->DumpDependencyGraphFileName, *ToString(World->GetNetMode()));
diff --git a/Engine/Source/Editor/MassEntityEditor/Private/MassEntityEditorSubsystem.cpp b/Engine/Source/Editor/MassEntityEditor/Private/MassEntityEditorSubsystem.cpp
--- a/Engine/Source/Editor/MassEntityEditor/Private/MassEntityEditorSubsystem.cpp
+++ b/Engine/Source/Editor/MassEntityEditor/Private/MassEntityEditorSubsystem.cpp
@@ -89,13 +89,14 @@ UMassEntityEditorSubsystem::~UMassEntityEditorSubsystem()
 
 void UMassEntityEditorSubsystem::Initialize(FSubsystemCollectionBase& Collection)
 {
-	if (UMassEntitySettings::IsInitialized())
+	UMassEntitySettings* Settings = GetMutableDefault<UMassEntitySettings>();
+	if (Settings->IsInitialized2())
 	{
 		InitializeMassInternals();
 	}
 	else
 	{
-		OnSettingsInitializedHandle = UMassEntitySettings::GetOnInitializedEvent().AddUObject(this, &UMassEntityEditorSubsystem::InitializeMassInternals);
+		OnSettingsInitializedHandle = Settings->GetOnInitialized().AddUObject(this, &UMassEntityEditorSubsystem::InitializeMassInternals);
 	}
 
 	Super::Initialize(Collection);
@@ -125,7 +126,8 @@ void UMassEntityEditorSubsystem::InitializeMassInternals()
 	EntityManager->Initialize(InitParams);
 
 	// set up ProcessingPhasesConfig
-	TConstArrayView<FMassProcessingPhaseConfig> MainPhasesConfig = GET_MASS_CONFIG_VALUE(GetProcessingPhasesConfig());
+	UMassEntitySettings* Settings = GetMutableDefault<UMassEntitySettings>();
+	TConstArrayView<FMassProcessingPhaseConfig> MainPhasesConfig = Settings->GetProcessingPhasesConfig();
 	for (int PhaseIndex = 0; PhaseIndex < (int)EMassProcessingPhase::MAX; ++PhaseIndex)
 	{
 		ProcessingPhasesConfig[PhaseIndex] = MainPhasesConfig[PhaseIndex];
@@ -148,7 +150,8 @@ void UMassEntityEditorSubsystem::StopAndCleanUp()
 
 	if (OnSettingsInitializedHandle.IsValid())
 	{
-		UMassEntitySettings::GetOnInitializedEvent().Remove(OnSettingsInitializedHandle);
+		UMassEntitySettings* Settings = GetMutableDefault<UMassEntitySettings>();
+		Settings->GetOnInitialized().Remove(OnSettingsInitializedHandle);
 	}
 }
 
diff --git a/Engine/Source/Runtime/MassEntity/Private/MassArchetypeData.cpp b/Engine/Source/Runtime/MassEntity/Private/MassArchetypeData.cpp
--- a/Engine/Source/Runtime/MassEntity/Private/MassArchetypeData.cpp
+++ b/Engine/Source/Runtime/MassEntity/Private/MassArchetypeData.cpp
@@ -36,7 +36,7 @@ namespace UE::Mass
 FMassArchetypeData::FMassArchetypeData(const FMassArchetypeCreationParams& CreationParams)
 	: NumEntitiesPerChunk(UE::Mass::Private::UninitializedInt32)
 	, EntityListOffsetWithinChunk(UE::Mass::Private::UninitializedInt32)
-	, ChunkMemorySize(UE::Mass::SanitizeChunkMemorySize(CreationParams.ChunkMemorySize ? CreationParams.ChunkMemorySize : GET_MASS_CONFIG_VALUE(ChunkMemorySize)))
+	, ChunkMemorySize(UE::Mass::SanitizeChunkMemorySize(CreationParams.ChunkMemorySize ? CreationParams.ChunkMemorySize : GetDefault<UMassEntitySettings>()->ChunkMemorySize))
 {
 #if WITH_MASSENTITY_DEBUG
 	DebugNames.Add(CreationParams.DebugName);
@@ -1728,4 +1728,4 @@ bool FMassArchetypeHelper::DoesArchetypeMatchRequirements(const FMassArchetypeDa
 
 	//return bResult;
 }
-#endif // WITH_MASSENTITY_DEBUG
\ No newline at end of file
+#endif // WITH_MASSENTITY_DEBUG
diff --git a/Engine/Source/Runtime/MassEntity/Private/MassProcessor.cpp b/Engine/Source/Runtime/MassEntity/Private/MassProcessor.cpp
--- a/Engine/Source/Runtime/MassEntity/Private/MassProcessor.cpp
+++ b/Engine/Source/Runtime/MassEntity/Private/MassProcessor.cpp
@@ -604,8 +604,10 @@ void UMassCompositeProcessor::UpdateProcessorsCollection(TArrayView<FMassProcess
 	ChildPipeline.Reset();
 
 	const UWorld* World = GetWorld();
+	const UMassEntitySettings* Settings = GetDefault<UMassEntitySettings>();
+
 	const EProcessorExecutionFlags WorldExecutionFlags = UE::Mass::Utils::DetermineProcessorExecutionFlags(World, InWorldExecutionFlags);
-	const FMassProcessingPhaseConfig& PhaseConfig = GET_MASS_CONFIG_VALUE(GetProcessingPhaseConfig(ProcessingPhase));
+	const FMassProcessingPhaseConfig& PhaseConfig = Settings->GetProcessingPhaseConfig(ProcessingPhase);
 
 	for (FMassProcessorOrderInfo& ProcessorInfo : InOutOrderedProcessors)
 	{
diff --git a/Engine/Source/Runtime/MassEntity/Public/MassEntitySettings.h b/Engine/Source/Runtime/MassEntity/Public/MassEntitySettings.h
--- a/Engine/Source/Runtime/MassEntity/Public/MassEntitySettings.h
+++ b/Engine/Source/Runtime/MassEntity/Public/MassEntitySettings.h
@@ -10,19 +10,18 @@
 #endif // UE_ENABLE_INCLUDE_ORDER_DEPRECATED_IN_5_6
 #include "MassEntitySettings.generated.h"
 
-
-#define GET_MASS_CONFIG_VALUE(a) (GetMutableDefault<UMassEntitySettings>()->a)
-
-struct FPropertyChangedEvent;
-
+#define GET_MASS_CONFIG_VALUE(a) \
+	UE_DEPRECATED_MACRO(5.8, "Use GetDefault<UMassEntitySettings>(*) directly.") \
+	(GetMutableDefault<UMassEntitySettings>()->a)
 
 /**
- * Implements the settings for MassEntity plugin
+ * Implements the settings for the MassEntity module
  */
 UCLASS(MinimalAPI, config = Mass, defaultconfig, DisplayName = "Mass Entity")
 class UMassEntitySettings : public UMassModuleSettings
 {
 	GENERATED_BODY()
+
 public:
 #if WITH_EDITORONLY_DATA
 	DECLARE_MULTICAST_DELEGATE_OneParam(FOnSettingsChange, const FPropertyChangedEvent& /*PropertyChangedEvent*/);
@@ -37,13 +36,30 @@ public:
 	MASSENTITY_API TConstArrayView<FMassProcessingPhaseConfig> GetProcessingPhasesConfig();
 	const FMassProcessingPhaseConfig& GetProcessingPhaseConfig(const EMassProcessingPhase ProcessingPhase) const { check(ProcessingPhase != EMassProcessingPhase::MAX); return ProcessingPhasesConfig[int(ProcessingPhase)]; }
 
-	static FOnInitialized& GetOnInitializedEvent() { return GET_MASS_CONFIG_VALUE(OnInitializedEvent); }
+	FOnInitialized& GetOnInitialized() { return OnInitializedEvent; }
+
+	UE_DEPRECATED(5.8, "Use the instance method GetOnInitialized")
+	static FOnInitialized& GetOnInitializedEvent()
+	{
+		UMassEntitySettings* Settings = GetMutableDefault<UMassEntitySettings>();
+		return Settings->OnInitializedEvent;
+	}
+
 #if WITH_EDITOR
 	FOnSettingsChange& GetOnSettingsChange() { return OnSettingsChange; }
 
-	static bool IsInitialized() { return GET_MASS_CONFIG_VALUE(bInitialized); }
+	bool IsInitialized2() const { return bInitialized; }
+
+	UE_DEPRECATED(5.8, "Use the instance method IsInitialized2")
+	static bool IsInitialized()
+	{
+		const UMassEntitySettings* Settings = GetDefault<UMassEntitySettings>();
+		return Settings->bInitialized;
+	}
+#endif // WITH_EDITOR
 
 protected:
+#if WITH_EDITOR
 	MASSENTITY_API virtual void PostEditChangeProperty(struct FPropertyChangedEvent& PropertyChangedEvent) override;
 	MASSENTITY_API virtual void PostEditChangeChainProperty(struct FPropertyChangedChainEvent& PropertyChangedEvent) override;
 #endif // WITH_EDITOR
@@ -78,10 +94,11 @@ public:
 	UPROPERTY(VisibleAnywhere, Category = Mass, Transient, Instanced, EditFixedSize)
 	TArray<TObjectPtr<UMassProcessor>> ProcessorCDOs;
 
-#if WITH_EDITORONLY_DATA
 protected:
+#if WITH_EDITORONLY_DATA
 	FOnSettingsChange OnSettingsChange;
 #endif // WITH_EDITORONLY_DATA
+
 	bool bInitialized = false;
 	bool bEngineInitialized = false;
 
