Allow multiple registration of types with Mass
- This happens naturally with Game Feature Plugins. When a plugin is activated, deactivated, then activated again it may register the same type a second time. This is an expected and supported workflow.
- Previously, an ensure would fire when registering the same type a second time (with a special case for "built in" types).
- Now, we simply allow the duplicate registration.
- Ideally we would ensure that the registered type info is identical, but FRelationTypeTraits inside the variant in FTypeInfo, contains a TFunction which does not support operator== (and it appears non trivial to support).

diff --git a/Engine/Source/Runtime/MassEntity/Private/MassTypeManager.cpp b/Engine/Source/Runtime/MassEntity/Private/MassTypeManager.cpp
--- a/Engine/Source/Runtime/MassEntity/Private/MassTypeManager.cpp
+++ b/Engine/Source/Runtime/MassEntity/Private/MassTypeManager.cpp
@@ -48,15 +48,25 @@ namespace UE::Mass
 
 			TypeDataMap.Add(FTypeHandle(TypeHandle.TypeKey), MoveTemp(TypeInfo));
 		}
-		else if (testableEnsureMsgf(bBuiltInTypesRegistered == false, TEXT("Registered type overriding is supported only as part of built-in types registration.")))
+		else
 		{
-			// we're overriding the existing data with the new data in assumption it's more up-to-date.
-			// The most common occurence of this will be with already registered subsystems' subclasses.
-			// The subclasses can change the data registered on their behalf by the super class,
-			// but most of the time that won't be necessary.
+			// We're registering the same type multiple times. There are multiple cases where this can happen:
+			//
+			// * The most common occurrence of this will be with already registered subsystems' subclasses. The
+			//   subclasses can change the data registered on their behalf by the super class, but most of the time that
+			//   won't be necessary. We assume the new data is is more up-to-date and overwrite the existing data with
+			//   the new data.
+			//
+			// * A GameFeaturePlugin is activating for a second time (e.g. activate > deactivate > activate). In this
+			//   case we allow the duplicate registration as long as the type info is identical.
+			//
+			// In theory we should ensure that the new type info matches the old type info (for the second case only),
+			// but FRelationTypeTraits contains a TFunction and does not support operator==.
+
+			UE_LOG(LogMass, Verbose, TEXT("Re-registering type: %s"), *InType->GetName());
 			*ExistingData = MoveTemp(TypeInfo);
 		}
-		
+
 		return TypeHandle;
 	}
 
@@ -141,4 +151,4 @@ namespace UE::Mass
 	{
 		return UE::Mass::IsA<FMassRelation>(RelationOrElementType) && TypeDataMap.Contains(FTypeHandle(RelationOrElementType));
 	}
-}
\ No newline at end of file
+}
