Improve comment for behavior change in 46642028
- Also check that the type doesn't change between subsystem and non-subsystem in duplicate registration. This would invalidate cached data and our general expectation is that the type info hasn't changed. The intent is to make a future pass to handle type re-registration more robustly to support GFPs being re-activated.

diff --git a/Engine/Source/Runtime/MassEntity/Private/MassTypeManager.cpp b/Engine/Source/Runtime/MassEntity/Private/MassTypeManager.cpp
--- a/Engine/Source/Runtime/MassEntity/Private/MassTypeManager.cpp
+++ b/Engine/Source/Runtime/MassEntity/Private/MassTypeManager.cpp
@@ -54,15 +54,18 @@ namespace UE::Mass
 			//
 			// * The most common occurrence of this will be with already registered subsystems' subclasses. The
 			//   subclasses can change the data registered on their behalf by the super class, but most of the time that
-			//   won't be necessary. We assume the new data is is more up-to-date and overwrite the existing data with
-			//   the new data.
+			//   won't be necessary. We assume the new data is more up-to-date and overwrite the existing data.
 			//
 			// * A GameFeaturePlugin is activating for a second time (e.g. activate > deactivate > activate). In this
-			//   case we allow the duplicate registration as long as the type info is identical.
+			//   case we always allow the duplicate registration and overwrite the existing data.
 			//
 			// In theory we should ensure that the new type info matches the old type info (for the second case only),
 			// but FRelationTypeTraits contains a TFunction and does not support operator==.
 
+			bool bWasSubsystem = ExistingData->Traits.IsType<FSubsystemTypeTraits>();
+			bool bIsSubsystem = TypeInfo.Traits.IsType<FSubsystemTypeTraits>();
+			ensureAlways(bWasSubsystem == bIsSubsystem);
+
 			UE_LOG(LogMass, Verbose, TEXT("Re-registering type: %s"), *InType->GetName());
 			*ExistingData = MoveTemp(TypeInfo);
 		}
