[Mass]
Added Indirect fragment access feature to allow processors to safely access fragment data from entities outside their query results
Added linked entity feature to enforce grouping like entities based on shared entity access for more performant indirect access that behaves like mutable shared fragments
Fixed issue with dependency solver issue with tags requirements for multiple-query processors being stomped instead of aggregated
Fixe bug where processor-level subsystem access was not being considered by the dependency solver

diff --git a/Engine/Source/Developer/MassEntityTestSuite/Private/MassDependencyTest.cpp b/Engine/Source/Developer/MassEntityTestSuite/Private/MassDependencyTest.cpp
--- a/Engine/Source/Developer/MassEntityTestSuite/Private/MassDependencyTest.cpp
+++ b/Engine/Source/Developer/MassEntityTestSuite/Private/MassDependencyTest.cpp
@@ -238,6 +238,111 @@ struct FComplexScenario : FDependencySolverBase
 };
 IMPLEMENT_AI_INSTANT_TEST(FComplexScenario, "System.Mass.Dependencies.Complex");
 
+struct FThreadUnsafeWriteLinkedEntityFragment : FDependencySolverBase
+{
+	virtual bool SetUp() override
+	{
+		FDependencySolverBase::SetUp();
+		{
+			UMassTestProcessorBase* Proc = Processors.Add_GetRef(NewTestProcessor<UMassTestProcessor_A>(EntityManager));
+			// This creates the indirect fragment requirement but an indirect requirement alone is not valid
+			// The query requirements are valid because it also adds a requirement for the link type (const shared fragment)
+			Proc->EntityQuery.AddLinkedEntityRequirement<FTestFragment_Int, FMassEntityLinkFragment>(EMassFragmentAccess::ReadWrite);
+		}
+		{
+			UMassTestProcessorBase* Proc = Processors.Add_GetRef(NewTestProcessor<UMassTestProcessor_E>(EntityManager));
+			Proc->EntityQuery.AddRequirement<FTestFragment_Int>(EMassFragmentAccess::ReadWrite);
+		}
+		{
+			UMassTestProcessorBase* Proc = Processors.Add_GetRef(NewTestProcessor<UMassTestProcessor_B>(EntityManager));
+			Proc->EntityQuery.AddLinkedEntityRequirement<FTestFragment_Int, FMassEntityLinkFragment>(EMassFragmentAccess::ReadWrite);
+		}
+		{
+			UMassTestProcessorBase* Proc = Processors.Add_GetRef(NewTestProcessor<UMassTestProcessor_C>(EntityManager));
+			Proc->EntityQuery.AddLinkedEntityRequirement<FTestFragment_Int, FMassEntityLinkFragment>(EMassFragmentAccess::ReadWrite);
+		}
+		{
+			UMassTestProcessorBase* Proc = Processors.Add_GetRef(NewTestProcessor<UMassTestProcessor_D>(EntityManager));
+			Proc->EntityQuery.AddLinkedEntityRequirement<FTestFragment_Int, FMassEntityLinkFragment>(EMassFragmentAccess::ReadWrite);
+		}
+		{
+			UMassTestProcessorBase* Proc = Processors.Add_GetRef(NewTestProcessor<UMassTestProcessor_F>(EntityManager));
+			Proc->EntityQuery.AddRequirement<FTestFragment_Int>(EMassFragmentAccess::ReadWrite);
+		}
+
+		return true;
+	}
+
+	virtual bool InstantTest() override
+	{
+		Solve();
+
+		int32 EmptyCount = 0;
+		for (const FMassProcessorOrderInfo& OrderInfo : Result)
+		{
+			if (OrderInfo.Dependencies.IsEmpty())
+			{
+				EmptyCount++;
+			}
+		}
+
+		AITEST_TRUE("Dependent processor chain should have exactly one independent processor", EmptyCount == 1);
+
+		return true;
+	}
+};
+IMPLEMENT_AI_INSTANT_TEST(FThreadUnsafeWriteLinkedEntityFragment, "System.Mass.Dependencies.ThreadUnsafeWriteLinkedEntityFragment");
+
+struct FThreadUnsafeWriteIndirectFragment : FDependencySolverBase
+{
+	virtual bool SetUp() override
+	{
+		FDependencySolverBase::SetUp();
+		{
+			UMassTestProcessorBase* Proc = Processors.Add_GetRef(NewTestProcessor<UMassTestProcessor_E>(EntityManager));
+			Proc->EntityQuery.AddRequirement<FTestFragment_Int>(EMassFragmentAccess::ReadWrite);
+		}
+		{
+			UMassTestProcessorBase* Proc = Processors.Add_GetRef(NewTestProcessor<UMassTestProcessor_A>(EntityManager));
+			// This creates ONLY the indirect fragment requirement but an indirect requirement alone is not valid
+			Proc->EntityQuery.AddIndirectFragmentRequirement<FTestFragment_Int>(EMassFragmentAccess::ReadWrite);
+			// we must have another requirement of some sort, in practice we would want to access something with the entity handle for indirect access
+			// making this read-only ensures that the solver is only resolving indirect contention on FTestFragment_Int
+			Proc->EntityQuery.AddRequirement<FTestFragment_EntityHandle>(EMassFragmentAccess::ReadOnly);
+		}
+		{
+			UMassTestProcessorBase* Proc = Processors.Add_GetRef(NewTestProcessor<UMassTestProcessor_B>(EntityManager));
+			Proc->EntityQuery.AddIndirectFragmentRequirement<FTestFragment_Int>(EMassFragmentAccess::ReadWrite);
+			Proc->EntityQuery.AddRequirement<FTestFragment_EntityHandle>(EMassFragmentAccess::ReadOnly);
+		}
+		{
+			UMassTestProcessorBase* Proc = Processors.Add_GetRef(NewTestProcessor<UMassTestProcessor_C>(EntityManager));
+			Proc->EntityQuery.AddRequirement<FTestFragment_Int>(EMassFragmentAccess::ReadWrite);
+		}
+
+		return true;
+	}
+
+	virtual bool InstantTest() override
+	{
+		Solve();
+
+		int32 EmptyCount = 0;
+		for (const FMassProcessorOrderInfo& OrderInfo : Result)
+		{
+			if (OrderInfo.Dependencies.IsEmpty())
+			{
+				EmptyCount++;
+			}
+		}
+
+		AITEST_TRUE("Dependent processor chain should have exactly one independent processor", EmptyCount == 1);
+
+		return true;
+	}
+};
+IMPLEMENT_AI_INSTANT_TEST(FThreadUnsafeWriteIndirectFragment, "System.Mass.Dependencies.ThreadUnsafeWriteIndirectFragment");
+
 struct FThreadUnsafeWriteSubsystem : FDependencySolverBase
 {
 	virtual bool SetUp() override
diff --git a/Engine/Source/Developer/MassEntityTestSuite/Private/MassQueryTest.cpp b/Engine/Source/Developer/MassEntityTestSuite/Private/MassQueryTest.cpp
--- a/Engine/Source/Developer/MassEntityTestSuite/Private/MassQueryTest.cpp
+++ b/Engine/Source/Developer/MassEntityTestSuite/Private/MassQueryTest.cpp
@@ -8,6 +8,7 @@
 #include "MassExecutor.h"
 #include "MassExecutionContext.h"
 #include "MassTypeManager.h"
+#include "MassEntityLinkFragments.h"
 #include "Algo/Compare.h"
 
 #define LOCTEXT_NAMESPACE "MassTest"
@@ -1034,6 +1035,140 @@ struct FQueryTest_ExecutionLimiter : FEntityTestBase
 };
 IMPLEMENT_AI_INSTANT_TEST(FQueryTest_ExecutionLimiter, "System.Mass.Query.ExecutionLimiter");
 
+
+struct FQueryTest_LinkedFragmentAccess : FEntityTestBase
+{
+	virtual bool InstantTest() override
+	{
+		CA_ASSUME(EntityManager);
+		const int32 DirectIntValue = 1234;
+		const int32 LinkedIntValue = 9898;
+		const float DirectFloatValue = 10.0f;
+		const float InirectFloatValue = -10.0f;
+
+		TArray<const UScriptStruct*> Fragments = { FTestFragment_Float::StaticStruct(), FTestFragment_Int::StaticStruct() };
+		const FMassArchetypeHandle Archetype = EntityManager->CreateArchetype(Fragments);
+		FMassEntityHandle LinkedEntity = EntityManager->CreateEntity(Archetype);
+		FTestFragment_Float& LinkedFloatFragment = EntityManager->GetFragmentDataChecked<FTestFragment_Float>(LinkedEntity);
+		FTestFragment_Int& LinkedIntFragment = EntityManager->GetFragmentDataChecked<FTestFragment_Int>(LinkedEntity);
+		LinkedFloatFragment.Value = InirectFloatValue;
+		LinkedIntFragment.Value = LinkedIntValue;
+
+		FMassEntityLinkFragment LinkedEntityFragment;
+		LinkedEntityFragment.LinkedEntityHandle = LinkedEntity;
+		FConstSharedStruct SharedLinkedEntityFragment = FSharedStruct::Make(LinkedEntityFragment);
+
+		TArray<FMassEntityHandle> Entities;
+		EntityManager->BatchCreateEntities(IntsArchetype, 10, Entities);
+		for (FMassEntityHandle& Entity : Entities)
+		{
+			EntityManager->AddConstSharedFragmentToEntity(Entity, SharedLinkedEntityFragment);
+			FTestFragment_Int& DirectIntFragment = EntityManager->GetFragmentDataChecked<FTestFragment_Int>(Entity);
+			DirectIntFragment.Value = DirectIntValue;
+		}
+
+		Entities.Reset();
+		EntityManager->BatchCreateEntities(FloatsArchetype, 10, Entities);
+		for (FMassEntityHandle& Entity : Entities)
+		{
+			EntityManager->AddConstSharedFragmentToEntity(Entity, SharedLinkedEntityFragment);
+			FTestFragment_Float& DirectFloatFragment = EntityManager->GetFragmentDataChecked<FTestFragment_Float>(Entity);
+			DirectFloatFragment.Value = DirectFloatValue;
+		}
+
+		FMassEntityQuery Query(EntityManager);
+		Query.AddRequirement<FTestFragment_Int>(EMassFragmentAccess::ReadOnly);
+		Query.AddLinkedEntityRequirement<FTestFragment_Int, FMassEntityLinkFragment>(EMassFragmentAccess::ReadOnly);
+		Query.CacheArchetypes();
+
+		bool bAnyMatch = false;
+		bool bFoundShared = true;
+		FMassExecutionContext ExecContext(*EntityManager.Get());
+		Query.ForEachEntityChunk(ExecContext, [this, &bAnyMatch, &bFoundShared](FMassExecutionContext& Context)
+			{
+				const FTestFragment_Int* SharedInt = Context.GetLinkedEntityFragmentPtr<FTestFragment_Int>();
+				if (!SharedInt)
+				{
+					bFoundShared = false;
+					return;
+				}
+				const TConstArrayView<FTestFragment_Int> IntFragments = Context.GetFragmentView<FTestFragment_Int>();
+
+				for (FMassExecutionContext::FEntityIterator EntityIt = Context.CreateEntityIterator(); EntityIt; ++EntityIt)
+				{
+					const FTestFragment_Int& IntFragment = IntFragments[EntityIt];
+					bAnyMatch |= IntFragment.Value == SharedInt->Value;
+				}
+			});
+		AITEST_TRUE("linked fragment should be discoverable", bFoundShared);
+		AITEST_FALSE("owned fragment should be independent from linked fragment", bAnyMatch);
+		return true;
+	}
+};
+IMPLEMENT_AI_INSTANT_TEST(FQueryTest_LinkedFragmentAccess, "System.Mass.Query.LinkedFragmentAccess");
+
+struct FQueryTest_IndirectFragmentAccess : FEntityTestBase
+{
+	virtual bool InstantTest() override
+	{
+		CA_ASSUME(EntityManager);
+		const int32 IntValue = 1234;
+		const float FloatValue = 10.0f;
+		const int32 EntityCount = 100;
+
+		TArray<FMassEntityHandle> IntEntities;
+		EntityManager->BatchCreateEntities(IntsArchetype, EntityCount, IntEntities);
+		for (FMassEntityHandle& Entity : IntEntities)
+		{
+			FTestFragment_Int& IntFragment = EntityManager->GetFragmentDataChecked<FTestFragment_Int>(Entity);
+			IntFragment.Value = IntValue;
+		}
+
+		TArray<FMassEntityHandle> FloatEntities;
+		EntityManager->BatchCreateEntities(FloatsArchetype, EntityCount, FloatEntities);
+		int HandleIndex = 0;
+		for (FMassEntityHandle& Entity : FloatEntities)
+		{
+			EntityManager->AddFragmentToEntity(Entity, FTestFragment_EntityHandle::StaticStruct());
+			FTestFragment_Float& FloatFragment = EntityManager->GetFragmentDataChecked<FTestFragment_Float>(Entity);
+			FTestFragment_EntityHandle& EntityhandleFragment = EntityManager->GetFragmentDataChecked<FTestFragment_EntityHandle>(Entity);
+			FloatFragment.Value = FloatValue;
+			EntityhandleFragment.Value = IntEntities[HandleIndex++];
+		}
+
+		FMassEntityQuery Query(EntityManager);
+		Query.AddRequirement<FTestFragment_EntityHandle>(EMassFragmentAccess::ReadOnly);
+		Query.AddIndirectFragmentRequirement<FTestFragment_Float>(EMassFragmentAccess::ReadOnly);
+		Query.CacheArchetypes();
+
+		bool bAllMatch = true;
+		bool bFoundIndirect = true;
+		FMassExecutionContext ExecContext(*EntityManager.Get());
+		Query.ForEachEntityChunk(ExecContext, [this, &bAllMatch, &bFoundIndirect](FMassExecutionContext& Context)
+			{
+				const TConstArrayView<FTestFragment_EntityHandle> HandleFragments = Context.GetFragmentView<FTestFragment_EntityHandle>();
+
+				for (FMassExecutionContext::FEntityIterator EntityIt = Context.CreateEntityIterator(); EntityIt; ++EntityIt)
+				{
+					const FTestFragment_EntityHandle& HandleFragment = HandleFragments[EntityIt];
+					const FTestFragment_Int* IntFragmentPtr = Context.GetIndirectConstFragmentPtr<FTestFragment_Int>(HandleFragment.Value);
+					if (IntFragmentPtr)
+					{
+						bAllMatch |= IntFragmentPtr->Value == IntValue;
+					}
+					else
+					{
+						bFoundIndirect = false;
+					}
+				}
+			});
+		AITEST_TRUE("Indirect fragment should be discoverable", bFoundIndirect);
+		AITEST_TRUE("Indirect fragment should have correct value", bAllMatch);
+		return true;
+	}
+};
+IMPLEMENT_AI_INSTANT_TEST(FQueryTest_IndirectFragmentAccess, "System.Mass.Query.IndirectFragmentAccess");
+
 } // FMassQueryTest
 
 #undef LOCTEXT_NAMESPACE
diff --git a/Engine/Source/Developer/MassEntityTestSuite/Public/MassEntityTestTypes.h b/Engine/Source/Developer/MassEntityTestSuite/Public/MassEntityTestTypes.h
--- a/Engine/Source/Developer/MassEntityTestSuite/Public/MassEntityTestTypes.h
+++ b/Engine/Source/Developer/MassEntityTestSuite/Public/MassEntityTestTypes.h
@@ -28,6 +28,17 @@ namespace UE::Mass::Testing
 	struct FMassTestProcessingPhaseManager;
 }
 
+USTRUCT()
+struct FTestFragment_EntityHandle : public FMassFragment
+{
+	GENERATED_BODY()
+
+	UPROPERTY()
+	FMassEntityHandle Value;
+
+	FTestFragment_EntityHandle(const FMassEntityHandle InValue = FMassEntityHandle()) : Value(InValue) {}
+};
+
 USTRUCT()
 struct FTestFragment_Float : public FMassFragment
 {
@@ -41,6 +52,8 @@ USTRUCT()
 struct FTestFragment_Int : public FMassFragment
 {
 	GENERATED_BODY()
+
+	UPROPERTY()
 	int32 Value = 0;
 
 	FTestFragment_Int(const int32 InValue = 0) : Value(InValue) {}
@@ -61,6 +74,8 @@ USTRUCT()
 struct FTestFragment_Large : public FMassFragment
 {
 	GENERATED_BODY()
+
+	UPROPERTY()
 	uint8 Value[64];
 
 	FTestFragment_Large(uint8 Fill = 0)
@@ -73,6 +88,8 @@ USTRUCT()
 struct FTestFragment_Array : public FMassFragment
 {
 	GENERATED_BODY()
+
+	UPROPERTY()
 	TArray<int32> Value;
 
 	FTestFragment_Array(uint8 Num = 0)
@@ -94,6 +111,7 @@ USTRUCT()
 struct FFragmentWithSharedPtr : public FMassFragment
 {
 	GENERATED_BODY()
+
 	TSharedPtr<int32> Data;
 
 	FFragmentWithSharedPtr() = default;
@@ -115,6 +133,8 @@ USTRUCT()
 struct FTestChunkFragment_Int : public FMassChunkFragment
 {
 	GENERATED_BODY()
+
+	UPROPERTY()
 	int32 Value = 0;
 
 	FTestChunkFragment_Int(const int32 InValue = 0) : Value(InValue) {}
@@ -124,6 +144,8 @@ USTRUCT()
 struct FTestChunkFragment_Float : public FMassChunkFragment
 {
 	GENERATED_BODY()
+
+	UPROPERTY()
 	float Value = 0.f;
 
 	FTestChunkFragment_Float(const float InValue = 0.f) : Value(InValue) {}
diff --git a/Engine/Source/Runtime/MassEntity/Private/MassArchetypeData.cpp b/Engine/Source/Runtime/MassEntity/Private/MassArchetypeData.cpp
--- a/Engine/Source/Runtime/MassEntity/Private/MassArchetypeData.cpp
+++ b/Engine/Source/Runtime/MassEntity/Private/MassArchetypeData.cpp
@@ -525,6 +525,7 @@ void FMassArchetypeData::ExecuteFunction(FMassExecutionContext& RunContext, cons
 
 			RunContext.SetCurrentChunkSerialModificationNumber(Chunk.GetSerialModificationNumber());
 			BindChunkFragmentRequirements(RunContext, RequirementMapping.ChunkFragments, Chunk);
+			RunContext.BindLinkedEntity();
 
 			if (!ChunkCondition || ChunkCondition(RunContext))
 			{
@@ -576,6 +577,7 @@ void FMassArchetypeData::ExecuteFunction(FMassExecutionContext& RunContext, cons
 
 			RunContext.SetCurrentChunkSerialModificationNumber(Chunk.GetSerialModificationNumber());
 			BindChunkFragmentRequirements(RunContext, RequirementMapping.ChunkFragments, Chunk);
+			RunContext.BindLinkedEntity();
 
 			if (!ChunkCondition || ChunkCondition(RunContext))
 			{
@@ -612,6 +614,7 @@ void FMassArchetypeData::ExecutionFunctionForChunk(FMassExecutionContext& RunCon
 #endif // WITH_MASSENTITY_DEBUG
 
 		BindChunkFragmentRequirements(RunContext, RequirementMapping.ChunkFragments, Chunk);
+		RunContext.BindLinkedEntity();
 
 		if (!ChunkCondition || ChunkCondition(RunContext))
 		{
diff --git a/Engine/Source/Runtime/MassEntity/Private/MassExecutionContext.cpp b/Engine/Source/Runtime/MassEntity/Private/MassExecutionContext.cpp
--- a/Engine/Source/Runtime/MassEntity/Private/MassExecutionContext.cpp
+++ b/Engine/Source/Runtime/MassEntity/Private/MassExecutionContext.cpp
@@ -3,6 +3,7 @@
 #include "MassExecutionContext.h"
 #include "MassArchetypeData.h"
 #include "MassEntityManager.h"
+#include "MassEntityLinkFragments.h"
 #if WITH_MASSENTITY_DEBUG
 #include "MassDebugger.h"
 #endif // WITH_MASSENTITY_DEBUG
@@ -125,6 +126,43 @@ UWorld* FMassExecutionContext::GetWorld()
 	return EntityManager->GetWorld(); 
 }
 
+void FMassExecutionContext::BindLinkedEntity()
+{
+	check(QueriesStack.Num());
+	if (QueriesStack.Last().Query->HasLinkedEntityRequirements())
+	{
+		// find an indirect entity handle to bind to the query
+		for (FConstSharedFragmentView& SharedFragmentView : ConstSharedFragmentViews)
+		{
+			const UScriptStruct* ScriptStruct = SharedFragmentView.FragmentView.GetScriptStruct();
+			if (ScriptStruct && ScriptStruct->IsChildOf(FMassEntityLinkFragment::StaticStruct()))
+			{
+				const FMassEntityLinkFragment& IndirectEntityFragment =
+					UE::StructUtils::GetStructRef<FMassEntityLinkFragment>(ScriptStruct, SharedFragmentView.FragmentView.GetMemory());
+
+				BindLinkedEntity(IndirectEntityFragment);
+				break;
+			}
+		}
+		checkf(LinkedEntityView.IsSet(), TEXT("Linked Entity access requirement specified but no const shared LinkedEntityFragment required"));
+	}
+}
+
+void FMassExecutionContext::BindLinkedEntity(const FMassEntityLinkFragment& EntityLinkFragment)
+{
+	check(QueriesStack.Num());
+	checkfSlow(
+		SharedFragmentViews.ContainsByPredicate(
+			[&EntityLinkFragment](const FSharedFragmentView& Element)
+			{
+				const void* ElemMem = Element.FragmentView.GetMemory();
+				return ElemMem == static_cast<const void*>(&EntityLinkFragment);
+			})
+		, TEXT("EntityLinkFragment must be included in the const shared fragments of the current query"));
+
+	LinkedEntityView = FMassEntityView(EntityManager.Get(), EntityLinkFragment.LinkedEntityHandle);
+}
+
 void FMassExecutionContext::PushQuery(FMassEntityQuery& InQuery)
 {
 	FQueryTransientRuntime& RuntimeData = QueriesStack.Add_GetRef({&InQuery});
diff --git a/Engine/Source/Runtime/MassEntity/Private/MassProcessor.cpp b/Engine/Source/Runtime/MassEntity/Private/MassProcessor.cpp
--- a/Engine/Source/Runtime/MassEntity/Private/MassProcessor.cpp
+++ b/Engine/Source/Runtime/MassEntity/Private/MassProcessor.cpp
@@ -337,6 +337,8 @@ void UMassProcessor::SetProcessingPhase(EMassProcessingPhase Phase)
 
 void UMassProcessor::ExportRequirements(FMassExecutionRequirements& OutRequirements) const
 {
+	ProcessorRequirements.ExportRequirements(OutRequirements);
+
 	for (FMassEntityQuery* Query : OwnedQueries)
 	{
 		CA_ASSUME(Query);
diff --git a/Engine/Source/Runtime/MassEntity/Private/MassProcessorDependencySolver.cpp b/Engine/Source/Runtime/MassEntity/Private/MassProcessorDependencySolver.cpp
--- a/Engine/Source/Runtime/MassEntity/Private/MassProcessorDependencySolver.cpp
+++ b/Engine/Source/Runtime/MassEntity/Private/MassProcessorDependencySolver.cpp
@@ -137,6 +137,7 @@ void FMassExecutionRequirements::Append(const FMassExecutionRequirements& Other)
 	for (int i = 0; i < EMassAccessOperation::MAX; ++i)
 	{
 		Fragments[i] += Other.Fragments[i];
+		IndirectFragments[i] += Other.IndirectFragments[i];
 		ChunkFragments[i] += Other.ChunkFragments[i];
 		SharedFragments[i] += Other.SharedFragments[i];
 		RequiredSubsystems[i] += Other.RequiredSubsystems[i];
@@ -159,6 +160,7 @@ void FMassExecutionRequirements::CountResourcesUsed()
 	for (int i = 0; i < EMassAccessOperation::MAX; ++i)
 	{
 		ResourcesUsedCount += Fragments[i].CountStoredTypes();
+		ResourcesUsedCount += IndirectFragments[i].CountStoredTypes();
 		ResourcesUsedCount += ChunkFragments[i].CountStoredTypes();
 		ResourcesUsedCount += SharedFragments[i].CountStoredTypes();
 		ResourcesUsedCount += RequiredSubsystems[i].CountStoredTypes();
@@ -175,14 +177,14 @@ int32 FMassExecutionRequirements::GetTotalBitsUsedCount()
 
 bool FMassExecutionRequirements::IsEmpty() const
 {
-	return Fragments.IsEmpty() && ChunkFragments.IsEmpty() 
+	return Fragments.IsEmpty() && IndirectFragments.IsEmpty() && ChunkFragments.IsEmpty()
 		&& SharedFragments.IsEmpty() && ConstSharedFragments.IsEmpty() && RequiredSubsystems.IsEmpty()
 		&& RequiredAllTags.IsEmpty() && RequiredAnyTags.IsEmpty() && RequiredNoneTags.IsEmpty();
 }
 
 FMassArchetypeCompositionDescriptor FMassExecutionRequirements::AsCompositionDescriptor() const
 {
-	return FMassArchetypeCompositionDescriptor(Fragments.Read + Fragments.Write
+	return FMassArchetypeCompositionDescriptor(Fragments.Read + Fragments.Write + IndirectFragments.Read + IndirectFragments.Write
 		, RequiredAllTags + RequiredAnyTags
 		, ChunkFragments.Read + ChunkFragments.Write
 		, SharedFragments.Read + SharedFragments.Write
@@ -198,6 +200,7 @@ FMassProcessorDependencySolver::FResourceUsage::FResourceUsage(const TArray<FNod
 	for (int i = 0; i < EMassAccessOperation::MAX; ++i)
 	{
 		FragmentsAccess[i].Access.AddZeroed(FMassFragmentBitSet::GetMaxNum());
+		IndirectFragmentsAccess[i].Access.AddZeroed(FMassFragmentBitSet::GetMaxNum());
 		ChunkFragmentsAccess[i].Access.AddZeroed(FMassChunkFragmentBitSet::GetMaxNum());
 		SharedFragmentsAccess[i].Access.AddZeroed(FMassSharedFragmentBitSet::GetMaxNum());
 		RequiredSubsystemsAccess[i].Access.AddZeroed(FMassExternalSubsystemBitSet::GetMaxNum());
@@ -206,13 +209,14 @@ FMassProcessorDependencySolver::FResourceUsage::FResourceUsage(const TArray<FNod
 
 template<typename TBitSet>
 void FMassProcessorDependencySolver::FResourceUsage::HandleElementType(TMassExecutionAccess<FResourceAccess>& ElementAccess
-	, const TMassExecutionAccess<TBitSet>& TestedRequirements, FMassProcessorDependencySolver::FNode& InOutNode, const int32 NodeIndex)
+	, const TMassExecutionAccess<TBitSet>& TestedRequirements, FMassProcessorDependencySolver::FNode& InOutNode, const int32 NodeIndex, bool bSkipArchetypeOverlapTest)
 {
 	using UE::Mass::Private::DoArchetypeContainersOverlap;
 
 	// when considering subsystem access we don't care about archetypes, so we cache the information whether
 	// we're dealing with subsystems and use that to potentially short-circuit the checks below.
 	constexpr bool bSubsystems = std::is_same_v<TBitSet, FMassExternalSubsystemBitSet>;
+	bSkipArchetypeOverlapTest = bSkipArchetypeOverlapTest || bSubsystems;
 
 	// for every bit set in TestedRequirements we do the following:
 	// 1. For every read-only requirement we make InOutNode depend on the currently stored Writer of this resource
@@ -238,7 +242,7 @@ void FMassProcessorDependencySolver::FResourceUsage::HandleElementType(TMassExec
 	{
 		for (int32 UserIndex : ElementAccess.Write.Access[*It].Users)
 		{
-			if (bSubsystems || DoArchetypeContainersOverlap(AllNodesView[UserIndex].ValidArchetypes, InOutNode.ValidArchetypes))
+			if (UserIndex != NodeIndex && (bSkipArchetypeOverlapTest || DoArchetypeContainersOverlap(AllNodesView[UserIndex].ValidArchetypes, InOutNode.ValidArchetypes)))
 			{
 				InOutNode.OriginalDependencies.Add(UserIndex);
 			}
@@ -251,7 +255,7 @@ void FMassProcessorDependencySolver::FResourceUsage::HandleElementType(TMassExec
 		for (int32 i = ElementAccess.Read.Access[*It].Users.Num() - 1; i >= 0; --i)
 		{
 			const int32 UserIndex = ElementAccess.Read.Access[*It].Users[i];
-			if (bSubsystems || DoArchetypeContainersOverlap(AllNodesView[UserIndex].ValidArchetypes, InOutNode.ValidArchetypes))
+			if (UserIndex != NodeIndex && (bSkipArchetypeOverlapTest || DoArchetypeContainersOverlap(AllNodesView[UserIndex].ValidArchetypes, InOutNode.ValidArchetypes)))
 			{	
 				InOutNode.OriginalDependencies.Add(UserIndex);
 				ElementAccess.Read.Access[*It].Users.RemoveAtSwap(i);
@@ -261,7 +265,7 @@ void FMassProcessorDependencySolver::FResourceUsage::HandleElementType(TMassExec
 		for (int32 i = ElementAccess.Write.Access[*It].Users.Num() - 1; i >= 0; --i)
 		{
 			const int32 UserIndex = ElementAccess.Write.Access[*It].Users[i];
-			if (bSubsystems || DoArchetypeContainersOverlap(AllNodesView[UserIndex].ValidArchetypes, InOutNode.ValidArchetypes))
+			if (UserIndex != NodeIndex && (bSkipArchetypeOverlapTest || DoArchetypeContainersOverlap(AllNodesView[UserIndex].ValidArchetypes, InOutNode.ValidArchetypes)))
 			{
 				InOutNode.OriginalDependencies.Add(UserIndex);
 				ElementAccess.Write.Access[*It].Users.RemoveAtSwap(i);
@@ -325,23 +329,63 @@ bool FMassProcessorDependencySolver::FResourceUsage::HasArchetypeConflict(TMassE
 	return false;
 }
 
-bool FMassProcessorDependencySolver::FResourceUsage::CanAccessRequirements(const FMassExecutionRequirements& TestedRequirements, const TArray<FMassArchetypeHandle>& InArchetypes) const
+bool FMassProcessorDependencySolver::FResourceUsage::CanAccessRequirements(
+	const FMassExecutionRequirements& Tested, const TArray<FMassArchetypeHandle>& InArchetypes) const
 {
-	// note that on purpose we're not checking ConstSharedFragments - those are always only read, no danger of conflicting access
-	bool bCanAccess = (CanAccess<FMassFragmentBitSet>(Requirements.Fragments, TestedRequirements.Fragments) || !HasArchetypeConflict(FragmentsAccess, InArchetypes))
-		&& (CanAccess<FMassChunkFragmentBitSet>(Requirements.ChunkFragments, TestedRequirements.ChunkFragments) || !HasArchetypeConflict(ChunkFragmentsAccess, InArchetypes))
-		&& (CanAccess<FMassSharedFragmentBitSet>(Requirements.SharedFragments, TestedRequirements.SharedFragments) || !HasArchetypeConflict(SharedFragmentsAccess, InArchetypes))
-		&& CanAccess<FMassExternalSubsystemBitSet>(Requirements.RequiredSubsystems, TestedRequirements.RequiredSubsystems);
-
-	return bCanAccess;
+	// Direct fragments: allow archetype fallback.
+	const bool bDirectOK =
+		(CanAccess<FMassFragmentBitSet>(Requirements.Fragments, Tested.Fragments) ||
+			!HasArchetypeConflict(FragmentsAccess, InArchetypes));
+
+	// Indirect fragments: archetype-agnostic -> no fallback permitted.
+	const bool bIndirectOK =
+		CanAccess<FMassFragmentBitSet>(Requirements.IndirectFragments, Tested.IndirectFragments);
+
+	// Cross-type contention: whenever either side is indirect, treat as global (no archetype fallback).
+	const bool bIndirectVsDirectOK =
+		!(Tested.IndirectFragments.Write.HasAny(Requirements.Fragments.Write)
+			|| Tested.IndirectFragments.Write.HasAny(Requirements.Fragments.Read)
+			|| Tested.IndirectFragments.Read.HasAny(Requirements.Fragments.Write));
+
+	const bool bDirectVsIndirectOK =
+		!(Tested.Fragments.Write.HasAny(Requirements.IndirectFragments.Write)
+			|| Tested.Fragments.Write.HasAny(Requirements.IndirectFragments.Read)
+			|| Tested.Fragments.Read.HasAny(Requirements.IndirectFragments.Write));
+
+	const bool bChunkOK =
+		(CanAccess<FMassChunkFragmentBitSet>(Requirements.ChunkFragments, Tested.ChunkFragments) ||
+			!HasArchetypeConflict(ChunkFragmentsAccess, InArchetypes));
+
+	const bool bSharedOK =
+		(CanAccess<FMassSharedFragmentBitSet>(Requirements.SharedFragments, Tested.SharedFragments) ||
+			!HasArchetypeConflict(SharedFragmentsAccess, InArchetypes));
+
+	const bool bSubsystemsOK =
+		CanAccess<FMassExternalSubsystemBitSet>(Requirements.RequiredSubsystems, Tested.RequiredSubsystems);
+
+	return bDirectOK && bIndirectOK && bIndirectVsDirectOK && bDirectVsIndirectOK
+		&& bChunkOK && bSharedOK && bSubsystemsOK;
 }
 
+
 void FMassProcessorDependencySolver::FResourceUsage::SubmitNode(const int32 NodeIndex, FNode& InOutNode)
 {
+	// Populate Node data access dependencies:
+	// Fragments Direct-Direct:
 	HandleElementType<FMassFragmentBitSet>(FragmentsAccess, InOutNode.Requirements.Fragments, InOutNode, NodeIndex);
+	// Fragments Indirect-Indirect:
+	HandleElementType<FMassFragmentBitSet>(IndirectFragmentsAccess, InOutNode.Requirements.IndirectFragments, InOutNode, NodeIndex, /*bSkipArchetypeOverlapTest*/ true);
+	// Fragments Indirect-Direct:
+	HandleElementType<FMassFragmentBitSet>(IndirectFragmentsAccess, InOutNode.Requirements.Fragments, InOutNode, NodeIndex, /*bSkipArchetypeOverlapTest*/ true);
+	// Fragments Direct-Indirect:
+	HandleElementType<FMassFragmentBitSet>(FragmentsAccess, InOutNode.Requirements.IndirectFragments, InOutNode, NodeIndex, /*bSkipArchetypeOverlapTest*/ true);
+	// Chunk Fragments:
 	HandleElementType<FMassChunkFragmentBitSet>(ChunkFragmentsAccess, InOutNode.Requirements.ChunkFragments, InOutNode, NodeIndex);
+	// (Mutable) Shared Fragments:
 	HandleElementType<FMassSharedFragmentBitSet>(SharedFragmentsAccess, InOutNode.Requirements.SharedFragments, InOutNode, NodeIndex);
+	// Subsystems:
 	HandleElementType<FMassExternalSubsystemBitSet>(RequiredSubsystemsAccess, InOutNode.Requirements.RequiredSubsystems, InOutNode, NodeIndex);
+
 	// note that on purpose we're not pushing ConstSharedFragments - those are always only read, no danger of conflicting access
 	// so there's no point in tracking them
 
@@ -401,6 +445,7 @@ bool FMassProcessorDependencySolver::FNode::IncreaseWaitingNodesCountAndPriority
 
 	return true;
 }
+
 //----------------------------------------------------------------------//
 //  FProcessorDependencySolver
 //----------------------------------------------------------------------//
diff --git a/Engine/Source/Runtime/MassEntity/Private/MassRequirements.cpp b/Engine/Source/Runtime/MassEntity/Private/MassRequirements.cpp
--- a/Engine/Source/Runtime/MassEntity/Private/MassRequirements.cpp
+++ b/Engine/Source/Runtime/MassEntity/Private/MassRequirements.cpp
@@ -307,9 +307,13 @@ void FMassFragmentRequirements::ExportRequirements(FMassExecutionRequirements& O
 	ExportRequirements<FMassSharedFragmentBitSet>(SharedFragmentRequirements, OutRequirements.SharedFragments);
 	ExportRequirements<FMassConstSharedFragmentBitSet>(ConstSharedFragmentRequirements, OutRequirements.ConstSharedFragments);
 
-	OutRequirements.RequiredAllTags = RequiredAllTags;
-	OutRequirements.RequiredAnyTags = RequiredAnyTags;
-	OutRequirements.RequiredNoneTags = RequiredNoneTags;
+	//indirect fragment access doesn't support EMassFragmentPresence::None so we can just add them directly
+	OutRequirements.IndirectFragments.Read += IndirectReadOnlyFragments;
+	OutRequirements.IndirectFragments.Write += IndirectReadWriteFragments;
+
+	OutRequirements.RequiredAllTags += RequiredAllTags;
+	OutRequirements.RequiredAnyTags += RequiredAnyTags;
+	OutRequirements.RequiredNoneTags += RequiredNoneTags;
 	// not exporting optional tags by design
 }
 
diff --git a/Engine/Source/Runtime/MassEntity/Public/MassEntityLinkFragments.h b/Engine/Source/Runtime/MassEntity/Public/MassEntityLinkFragments.h
new file mode 100644
--- /dev/null
+++ b/Engine/Source/Runtime/MassEntity/Public/MassEntityLinkFragments.h
@@ -0,0 +1,21 @@
+// Copyright Epic Games, Inc. All Rights Reserved.
+
+#pragma once
+
+#include "MassEntityTypes.h"
+#include "MassEntityHandle.h"
+#include "MassEntityLinkFragments.generated.h"
+
+/**
+ * Add this to an entity to link it to another entity.
+ * Groups entities that access the same LinkedEntityHandle indirectly for better cache behavior.
+ * Use this to minimize random access when many entities access the same shared data.
+ */
+USTRUCT()
+struct FMassEntityLinkFragment : public FMassConstSharedFragment
+{
+	UPROPERTY()
+	FMassEntityHandle LinkedEntityHandle;
+
+	GENERATED_BODY()
+};
diff --git a/Engine/Source/Runtime/MassEntity/Public/MassExecutionContext.h b/Engine/Source/Runtime/MassEntity/Public/MassExecutionContext.h
--- a/Engine/Source/Runtime/MassEntity/Public/MassExecutionContext.h
+++ b/Engine/Source/Runtime/MassEntity/Public/MassExecutionContext.h
@@ -12,6 +12,7 @@
 #include "MassArchetypeTypes.h"
 #include "MassSubsystemAccess.h"
 #include "MassProcessor.h"
+#include "MassEntityView.h"
 
 
 #define CHECK_IF_VALID(View, Type) \
@@ -24,6 +25,7 @@
 		, View ? *GetNameSafe(View->Requirement.StructType) : TEXT("[Not found]"))
 
 struct FMassEntityQuery;
+struct FMassEntityLinkFragment;
 
 struct FMassExecutionContext
 {
@@ -52,6 +54,12 @@ private:
 	using FSharedFragmentView = TFragmentView<FStructView>;
 	TArray<FSharedFragmentView, TInlineAllocator<4>> SharedFragmentViews;
 
+	/** linked entity to allow access to fragments on an entity not in the query, must be bound before access */
+	FMassEntityView LinkedEntityView;
+
+	/** Used for indirect fragment access */
+	FMassEntityView CachedIndirectEntityView;
+
 	FMassSubsystemAccess SubsystemAccess;
 	
 	// mz@todo make this shared ptr thread-safe and never auto-flush in MT environment. 
@@ -78,6 +86,7 @@ private:
 		TNotNull<FMassEntityQuery*> Query;
 		FMassExternalSubsystemBitSet ConstSubsystemsBitSet;
 		FMassExternalSubsystemBitSet MutableSubsystemsBitSet;
+
 #if WITH_MASSENTITY_DEBUG
 		/** MaxBreakFragmentCount needs to be bigger than the greatest number of fragments a query has a write requirement for to that can have a breakpoint set */
 		static constexpr uint32 MaxFragmentBreakpointCount = 8;
@@ -272,6 +281,11 @@ public:
 	void SetExecutionType(EMassExecutionContextType InExecutionType);
 	EMassExecutionContextType GetExecutionType() const;
 
+	/** Binds a linked entity with a specific link type, use this if you already have the link fragment you wish to bind */
+	MASSENTITY_API void BindLinkedEntity(const FMassEntityLinkFragment& EntityLinkFragment);
+	/** Binds a linked entity with the link type specified in the current query */
+	MASSENTITY_API void BindLinkedEntity();
+
 	float GetDeltaTimeSeconds() const;
 
 	MASSENTITY_API UWorld* GetWorld();
@@ -477,6 +491,58 @@ public:
 		return TArrayView<TFragmentBase>(reinterpret_cast<TFragmentBase*>(View.GetData()), View.Num());;
 	}
 
+	template<typename T>
+	T* GetLinkedEntityFragmentPtr()
+	{
+		checkf(QueriesStack.Num(), TEXT("There must be an active query to access fragments on the linked entity"));
+		checkf(LinkedEntityView.IsValid(), TEXT("No entity has been linked. Add a FMassLinkedEntityFragment to entities in this query to access fragments on a linked entity that isn't in this query."));
+		checkfSlow(QueriesStack.Last().Query->HasMutableIndirectEntityRequirement<T>()
+			, TEXT("Attempting to access a fragment from a linked entity without first adding the LinkedEntityFragment requirement to the query"));
+
+		return LinkedEntityView.GetFragmentDataPtr<T>();
+	}
+
+	template<typename T>
+	const T* GetLinkedEntityConstFragmentPtr()
+	{
+		checkf(QueriesStack.Num(), TEXT("There must be an active query to access fragments on the linked entity"));
+		checkf(LinkedEntityView.IsValid(), TEXT("No entity has been linked. Add a FMassLinkedEntityFragment to entities in this query to access fragments on a linked entity that isn't in this query."));
+		checkfSlow(QueriesStack.Last().Query->HasConstIndirectEntityRequirement<T>()
+			, TEXT("Attempting to access a fragment from a linked entity without first adding the LinkedEntityFragment requirement to the query"));
+
+		return LinkedEntityView.GetFragmentDataPtr<T>();
+	}
+
+	template<typename T>
+	T* GetIndirectFragmentPtr(const FMassEntityHandle& FragmentOwner)
+	{
+		checkf(QueriesStack.Num(), TEXT("There must be an active query to access indirect fragments, otherwise just use direct access with a query"));
+		checkfSlow(QueriesStack.Last().Query->HasMutableIndirectEntityRequirement<T>()
+			, TEXT("Attempting indirect fragment access without first adding and IndirectFragment requirement to the query"));
+
+		if (CachedIndirectEntityView.GetEntity() != FragmentOwner)
+		{
+			CachedIndirectEntityView = FMassEntityView(EntityManager.Get(), FragmentOwner);
+		}
+
+		return CachedIndirectEntityView.GetFragmentDataPtr<T>();
+	}
+
+	template<typename T>
+	const T* GetIndirectConstFragmentPtr(const FMassEntityHandle& FragmentOwner)
+	{
+		checkf(QueriesStack.Num(), TEXT("There must be an active query to access indirect fragments, otherwise just use direct access with a query"));
+		checkfSlow(QueriesStack.Last().Query->HasMutableIndirectEntityRequirement<T>()
+			, TEXT("Attempting indirect fragment access without first adding and IndirectFragment requirement to the query"));
+
+		if (CachedIndirectEntityView.GetEntity() != FragmentOwner)
+		{
+			CachedIndirectEntityView = FMassEntityView(EntityManager.Get(), FragmentOwner);
+		}
+
+		return CachedIndirectEntityView.GetFragmentDataPtr<T>();
+	}
+
 	template<typename T, typename = typename TEnableIf<TIsDerivedFrom<T, USubsystem>::IsDerived>::Type>
 	T* GetMutableSubsystem()
 	{
diff --git a/Engine/Source/Runtime/MassEntity/Public/MassProcessorDependencySolver.h b/Engine/Source/Runtime/MassEntity/Public/MassProcessorDependencySolver.h
--- a/Engine/Source/Runtime/MassEntity/Public/MassProcessorDependencySolver.h
+++ b/Engine/Source/Runtime/MassEntity/Public/MassProcessorDependencySolver.h
@@ -73,6 +73,7 @@ struct FMassExecutionRequirements
 	UE_API FMassArchetypeCompositionDescriptor AsCompositionDescriptor() const;
 
 	TMassExecutionAccess<FMassFragmentBitSet> Fragments;
+	TMassExecutionAccess<FMassFragmentBitSet> IndirectFragments;
 	TMassExecutionAccess<FMassChunkFragmentBitSet> ChunkFragments;
 	TMassExecutionAccess<FMassSharedFragmentBitSet> SharedFragments;
 	TMassExecutionAccess<FMassConstSharedFragmentBitSet> ConstSharedFragments;
@@ -156,6 +157,7 @@ private:
 		
 		FMassExecutionRequirements Requirements;
 		TMassExecutionAccess<FResourceAccess> FragmentsAccess;
+		TMassExecutionAccess<FResourceAccess> IndirectFragmentsAccess;
 		TMassExecutionAccess<FResourceAccess> ChunkFragmentsAccess;
 		TMassExecutionAccess<FResourceAccess> SharedFragmentsAccess;
 		TMassExecutionAccess<FResourceAccess> RequiredSubsystemsAccess;
@@ -163,7 +165,7 @@ private:
 
 		template<typename TBitSet>
 		void HandleElementType(TMassExecutionAccess<FResourceAccess>& ElementAccess
-			, const TMassExecutionAccess<TBitSet>& TestedRequirements, FMassProcessorDependencySolver::FNode& InOutNode, const int32 NodeIndex);
+			, const TMassExecutionAccess<TBitSet>& TestedRequirements, FMassProcessorDependencySolver::FNode& InOutNode, const int32 NodeIndex, bool bSkipArchetypeOverlapTest = false);
 
 		template<typename TBitSet>
 		static bool CanAccess(const TMassExecutionAccess<TBitSet>& StoredElements, const TMassExecutionAccess<TBitSet>& TestedElements);
diff --git a/Engine/Source/Runtime/MassEntity/Public/MassRequirements.h b/Engine/Source/Runtime/MassEntity/Public/MassRequirements.h
--- a/Engine/Source/Runtime/MassEntity/Public/MassRequirements.h
+++ b/Engine/Source/Runtime/MassEntity/Public/MassRequirements.h
@@ -4,6 +4,7 @@
 
 #include "MassEntityTypes.h"
 #include "MassExternalSubsystemTraits.h"
+#include "MassEntityLinkFragments.h"
 #include "Templates/SubclassOf.h"
 #include "MassRequirements.generated.h"
 
@@ -39,7 +40,7 @@ enum class EMassFragmentAccess : uint8
 	/** no binding required */
 	None, 
 
-	/** We want to read the data for the fragment */
+	/** We want to read the data from the fragment */
 	ReadOnly,
 
 	/** We want to read and write the data for the fragment */
@@ -199,9 +200,9 @@ struct FMassFragmentRequirements
 	FMassFragmentRequirements& AddRequirement(const UScriptStruct* FragmentType, const EMassFragmentAccess AccessMode, const EMassFragmentPresence Presence = EMassFragmentPresence::All)
 	{
 		checkf(bInitialized, TEXT("Modifying requirements before initialization is not supported."));
-		checkf(FragmentRequirements.FindByPredicate([FragmentType](const FMassFragmentRequirementDescription& Item){ return Item.StructType == FragmentType; }) == nullptr
+		checkf(FragmentRequirements.FindByPredicate([FragmentType](const FMassFragmentRequirementDescription& Item) { return Item.StructType == FragmentType; }) == nullptr
 			, TEXT("Duplicated requirements are not supported. %s already present"), *GetNameSafe(FragmentType));
-		
+
 		if (Presence != EMassFragmentPresence::None)
 		{
 			FragmentRequirements.Emplace(FragmentType, AccessMode, Presence);
@@ -236,12 +237,12 @@ struct FMassFragmentRequirements
 			, TEXT("Duplicated requirements are not supported. %s already present"), *T::StaticStruct()->GetName());
 
 		static_assert(UE::Mass::CFragment<T>, MASS_INVALID_FRAGMENT_MSG);
-		
+
 		if (Presence != EMassFragmentPresence::None)
 		{
 			FragmentRequirements.Emplace(T::StaticStruct(), AccessMode, Presence);
 		}
-		
+
 		switch (Presence)
 		{
 		case EMassFragmentPresence::All:
@@ -298,25 +299,25 @@ struct FMassFragmentRequirements
 		static_assert(UE::Mass::CTag<T>, "Given struct doesn't represent a valid tag type. Make sure to inherit from FMassFragment or one of its child-types.");
 		switch (Presence)
 		{
-			case EMassFragmentPresence::All:
-				RequiredAllTags.Add<T>();
-				break;
-			case EMassFragmentPresence::Any:
-				RequiredAnyTags.Add<T>();
-				break;
-			case EMassFragmentPresence::None:
-				RequiredNoneTags.Add<T>();
-				break;
-			case EMassFragmentPresence::Optional:
-				RequiredOptionalTags.Add<T>();
-				break;
+		case EMassFragmentPresence::All:
+			RequiredAllTags.Add<T>();
+			break;
+		case EMassFragmentPresence::Any:
+			RequiredAnyTags.Add<T>();
+			break;
+		case EMassFragmentPresence::None:
+			RequiredNoneTags.Add<T>();
+			break;
+		case EMassFragmentPresence::Optional:
+			RequiredOptionalTags.Add<T>();
+			break;
 		}
 		IncrementChangeCounter();
 		return *this;
 	}
 
 	/** actual implementation in specializations */
-	template<EMassFragmentPresence Presence> 
+	template<EMassFragmentPresence Presence>
 	FMassFragmentRequirements& AddTagRequirements(const FMassTagBitSet& TagBitSet)
 	{
 		static_assert(Presence == EMassFragmentPresence::None || Presence == EMassFragmentPresence::All || Presence == EMassFragmentPresence::Any
@@ -338,17 +339,17 @@ struct FMassFragmentRequirements
 
 		switch (Presence)
 		{
-			case EMassFragmentPresence::All:
-				RequiredAllChunkFragments.Add<T>();
-				ChunkFragmentRequirements.Emplace(T::StaticStruct(), AccessMode, Presence);
-				break;
-			case EMassFragmentPresence::Optional:
-				RequiredOptionalChunkFragments.Add<T>();
-				ChunkFragmentRequirements.Emplace(T::StaticStruct(), AccessMode, Presence);
-				break;
-			case EMassFragmentPresence::None:
-				RequiredNoneChunkFragments.Add<T>();
-				break;
+		case EMassFragmentPresence::All:
+			RequiredAllChunkFragments.Add<T>();
+			ChunkFragmentRequirements.Emplace(T::StaticStruct(), AccessMode, Presence);
+			break;
+		case EMassFragmentPresence::Optional:
+			RequiredOptionalChunkFragments.Add<T>();
+			ChunkFragmentRequirements.Emplace(T::StaticStruct(), AccessMode, Presence);
+			break;
+		case EMassFragmentPresence::None:
+			RequiredNoneChunkFragments.Add<T>();
+			break;
 		}
 		IncrementChangeCounter();
 		return *this;
@@ -363,17 +364,17 @@ struct FMassFragmentRequirements
 
 		switch (Presence)
 		{
-			case EMassFragmentPresence::All:
-				RequiredAllChunkFragments.Add(*ChunkFragmentType);
-				ChunkFragmentRequirements.Emplace(ChunkFragmentType, AccessMode, Presence);
-				break;
-			case EMassFragmentPresence::Optional:
-				RequiredOptionalChunkFragments.Add(*ChunkFragmentType);
-				ChunkFragmentRequirements.Emplace(ChunkFragmentType, AccessMode, Presence);
-				break;
-			case EMassFragmentPresence::None:
-				RequiredNoneChunkFragments.Add(*ChunkFragmentType);
-				break;
+		case EMassFragmentPresence::All:
+			RequiredAllChunkFragments.Add(*ChunkFragmentType);
+			ChunkFragmentRequirements.Emplace(ChunkFragmentType, AccessMode, Presence);
+			break;
+		case EMassFragmentPresence::Optional:
+			RequiredOptionalChunkFragments.Add(*ChunkFragmentType);
+			ChunkFragmentRequirements.Emplace(ChunkFragmentType, AccessMode, Presence);
+			break;
+		case EMassFragmentPresence::None:
+			RequiredNoneChunkFragments.Add(*ChunkFragmentType);
+			break;
 		}
 		IncrementChangeCounter();
 		return *this;
@@ -474,7 +475,7 @@ struct FMassFragmentRequirements
 
 	FMassFragmentRequirements& AddSharedRequirement(TNotNull<const UScriptStruct*> SharedFragmentType, const EMassFragmentAccess AccessMode, const EMassFragmentPresence Presence = EMassFragmentPresence::All)
 	{
-		checkf(UE::Mass::IsA<FMassSharedFragment>(SharedFragmentType), TEXT("Given struct doesn't represent a valid shared fragment type. Make sure to inherit from FMassSharedFragment or one of its child-types.")); 
+		checkf(UE::Mass::IsA<FMassSharedFragment>(SharedFragmentType), TEXT("Given struct doesn't represent a valid shared fragment type. Make sure to inherit from FMassSharedFragment or one of its child-types."));
 		checkf(bInitialized, TEXT("Modifying requirements before initialization is not supported."));
 		checkf(SharedFragmentRequirements.FindByPredicate([&SharedFragmentType](const FMassFragmentRequirementDescription& Item) { return Item.StructType == SharedFragmentType; }) == nullptr
 			, TEXT("Duplicated requirements are not supported. %s already present"), *SharedFragmentType->GetName());
@@ -506,6 +507,119 @@ struct FMassFragmentRequirements
 		return *this;
 	}
 
+	/**
+	 * Adds an access requirement for a fragment on a linked entity.
+	 * This is a specialization of Indirect fragment access that requires the target entity be "Linked"
+	 * To link an entity, add a FMassEntityLinkFragment to the entities you want to link from and set the LinkedEntityHandle to the target entity.
+	 * FMassEntityLinkFragment can be inherited to support specific link types and multiple linked entities.
+	 * Linking entities this way will group all entities that link to the target entity in memory for better cache behavior when iterating through query results.
+	 */
+	FMassFragmentRequirements& AddLinkedEntityRequirement(TNotNull<const UScriptStruct*> FragmentType, const EMassFragmentAccess AccessMode, TNotNull<const UScriptStruct*> LinkType = FMassEntityLinkFragment::StaticStruct())
+	{
+		checkf(bInitialized, TEXT("Modifying requirements before initialization is not supported."));
+		checkf(UE::Mass::IsA<FMassFragment>(FragmentType), TEXT("Fragment must inherit FMassFragment."));
+		checkfSlow(LinkedEntityFragmentRequirements.FindByPredicate([FragmentType](const FMassFragmentRequirementDescription& Item) { return Item.StructType == FragmentType; }) == nullptr
+			, TEXT("Duplicated linked entity requirements are not supported. %s already present"), *GetNameSafe(FragmentType));
+
+		const bool bHasLinkType = 
+			ConstSharedFragmentRequirements.ContainsByPredicate([LinkType](const FMassFragmentRequirementDescription& Item)
+			{
+				return Item.StructType == LinkType;
+			});
+
+			if (bHasLinkType)
+		{
+			checkfSlow(ConstSharedFragmentRequirements.ContainsByPredicate([](const FMassFragmentRequirementDescription& Item)
+				{
+					return Item.StructType->IsChildOf(FMassEntityLinkFragment::StaticStruct());
+				}) == false
+				, TEXT("A query can not use more than one Entity Link type. Use Indirect Fragment access instead, or separate queries."));
+
+			AddConstSharedRequirement(LinkType);
+		}
+
+
+		LinkedEntityFragmentRequirements.Emplace(FragmentType, AccessMode, EMassFragmentPresence::Any);
+
+		switch (AccessMode)
+		{
+		case EMassFragmentAccess::ReadOnly:
+			IndirectReadOnlyFragments.Add(*FragmentType);
+			break;
+		case EMassFragmentAccess::ReadWrite:
+			IndirectReadWriteFragments.Add(*FragmentType);
+			break;
+		}
+
+		// force recaching the next time this query is used or the following CacheArchetypes call.
+		IncrementChangeCounter();
+		return *this;
+	}
+
+	/**
+	 * Adds an access requirement for a fragment on a linked entity.
+	 * This is a specialization of Indirect fragment access that requires the target entity be "Linked"
+	 * To link an entity, add a FMassEntityLinkFragment to the entities you want to link from and set the LinkedEntityHandle to the target entity.
+	 * FMassEntityLinkFragment can be inherited to support specific link types and multiple linked entities.
+	 * Linking entities this way will group all entities that link to the target entity in memory for better cache behavior when iterating through query results.
+	 */
+	template<UE::Mass::CFragment T, typename TLinkFragment = FMassEntityLinkFragment>
+	FMassFragmentRequirements& AddLinkedEntityRequirement(const EMassFragmentAccess AccessMode)
+	{
+		static_assert(UE::Mass::CFragment<T>, MASS_INVALID_FRAGMENT_MSG);
+		return AddLinkedEntityRequirement(T::StaticStruct(), AccessMode, TLinkFragment::StaticStruct());
+	}
+
+	template<UE::Mass::CFragment T>
+	bool HasConstIndirectEntityRequirement()
+	{
+		const int32 TypeIndex = FMassFragmentBitSet::GetTypeIndex<T>();
+		return IndirectReadOnlyFragments.IsBitSet(TypeIndex) || IndirectReadWriteFragments.IsBitSet(TypeIndex);
+	}
+
+	template<UE::Mass::CFragment T>
+	bool HasMutableIndirectEntityRequirement()
+	{
+		return IndirectReadWriteFragments.IsBitSet(FMassFragmentBitSet::GetTypeIndex<T>());
+	}
+
+	/**
+	 * Adds an access requirement for a fragment on another entity.
+	 * This is used to enforce access requirements in the dependency solver, not to discover entities/archetypes.
+	 * Use FMassExecutionContext::GetIndirectFragmentPtr to access fragments indirectly in your processor.
+	 */
+	FMassFragmentRequirements& AddIndirectFragmentRequirement(TNotNull<const UScriptStruct*> FragmentType, const EMassFragmentAccess AccessMode)
+	{
+		checkf(bInitialized, TEXT("Modifying requirements before initialization is not supported."));
+		checkfSlow(UE::Mass::IsA<FMassFragment>(FragmentType), TEXT("Fragment must inherit FMassFragment."));
+
+		switch (AccessMode)
+		{
+		case EMassFragmentAccess::ReadOnly:
+			IndirectReadOnlyFragments.Add(*FragmentType);
+			break;
+		case EMassFragmentAccess::ReadWrite:
+			IndirectReadWriteFragments.Add(*FragmentType);
+			break;
+		}
+
+		// force recaching the next time this query is used or the following CacheArchetypes call.
+		IncrementChangeCounter();
+		return *this;
+	}
+
+	/**
+	 * Adds an access requirement for a fragment on another entity.
+	 * This is used to enforce access requirements in the dependency solver, not to discover entities/archetypes.
+	 * Use FMassExecutionContext::GetIndirectFragmentPtr to access fragments indirectly in your processor.
+	 */
+	template<UE::Mass::CFragment T>
+	FMassFragmentRequirements& AddIndirectFragmentRequirement(const EMassFragmentAccess AccessMode)
+	{
+		static_assert(UE::Mass::CFragment<T>, MASS_INVALID_FRAGMENT_MSG);
+		return AddIndirectFragmentRequirement(T::StaticStruct(), AccessMode);
+	}
+
 	MASSENTITY_API void Reset();
 
 	/** 
@@ -543,6 +657,7 @@ struct FMassFragmentRequirements
 	bool HasPositiveRequirements() const;
 	bool HasNegativeRequirements() const;
 	bool HasOptionalRequirements() const;
+	bool HasLinkedEntityRequirements() const;
 
 	MASSENTITY_API bool DoesArchetypeMatchRequirements(const FMassArchetypeHandle& ArchetypeHandle) const;
 	MASSENTITY_API bool DoesArchetypeMatchRequirements(const FMassArchetypeCompositionDescriptor& ArchetypeComposition) const;
@@ -595,6 +710,21 @@ protected:
 	FMassConstSharedFragmentBitSet RequiredOptionalConstSharedFragments;
 	FMassConstSharedFragmentBitSet RequiredNoneConstSharedFragments;
 
+	/**
+	 * Fragments that will be accessed on other entities along with this query.
+	 * These do not affect which entities are discovered by this query, but 
+	 * inform the dependency solver to avoid write contention on these fragments.
+	 * Used by Linked entity access and Indirect entity access.
+	 */
+	FMassFragmentBitSet IndirectReadOnlyFragments;
+	FMassFragmentBitSet IndirectReadWriteFragments;
+
+	/**
+	 * Requirements for fragments on Linked entities.
+	 * Queries will only discover entities with a linked entity that meets these requirements.
+	 */
+	TArray<FMassFragmentRequirementDescription> LinkedEntityFragmentRequirements;
+
 	TSharedPtr<FMassEntityManager> CachedEntityManager;
 
 private:
@@ -608,7 +738,7 @@ private:
 	 * still need to be satisfied.
 	 */
 	mutable uint16 bHasOptionalRequirements : 1 = false;
-
+	
 	uint16 bInitialized : 1 = false;
 	uint16 IncrementalChangesCount = 0;
 
@@ -839,6 +969,11 @@ inline bool FMassFragmentRequirements::HasOptionalRequirements() const
 	return bHasOptionalRequirements; 
 }
 
+inline bool FMassFragmentRequirements::HasLinkedEntityRequirements() const
+{ 
+	return !LinkedEntityFragmentRequirements.IsEmpty(); 
+}
+
 inline bool FMassFragmentRequirements::DoesRequireGameThreadExecution() const
 {
 	return bRequiresGameThreadExecution;
