[Mass] Intoducing FMassElementBitset to replace all other fragment/tag bitsets

diff --git a/Engine/Plugins/AI/MassAI/Source/MassAIDebug/Private/GameplayDebuggerCategory_Mass.cpp b/Engine/Plugins/AI/MassAI/Source/MassAIDebug/Private/GameplayDebuggerCategory_Mass.cpp
--- a/Engine/Plugins/AI/MassAI/Source/MassAIDebug/Private/GameplayDebuggerCategory_Mass.cpp
+++ b/Engine/Plugins/AI/MassAI/Source/MassAIDebug/Private/GameplayDebuggerCategory_Mass.cpp
@@ -473,23 +473,23 @@ void FGameplayDebuggerCategory_Mass::CollectData(APlayerController* OwnerPC, AAc
 					};
 
 					TArray<FName> ItemNames;
-					Composition.GetTags().DebugGetIndividualNames(ItemNames);
+					Composition.Get<FMassTagBitSet>().DebugGetIndividualNames(ItemNames);
 					AddTextLine(FString::Printf(TEXT("{Green}Tags:{White}%s"), *DescriptionBuilder(ItemNames)));
 					
 					ItemNames.Reset();
-					Composition.GetFragments().DebugGetIndividualNames(ItemNames);
+					Composition.Get<FMassFragmentBitSet>().DebugGetIndividualNames(ItemNames);
 					AddTextLine(FString::Printf(TEXT("{Green}Fragments:{White}%s"), *DescriptionBuilder(ItemNames)));
 					
 					ItemNames.Reset();
-					Composition.GetChunkFragments().DebugGetIndividualNames(ItemNames);
+					Composition.Get<FMassChunkFragmentBitSet>().DebugGetIndividualNames(ItemNames);
 					AddTextLine(FString::Printf(TEXT("{Green}Chunk Fragments:{White}%s"), *DescriptionBuilder(ItemNames)));
 
 					ItemNames.Reset();
-					Composition.GetSharedFragments().DebugGetIndividualNames(ItemNames);
+					Composition.Get<FMassSharedFragmentBitSet>().DebugGetIndividualNames(ItemNames);
 					AddTextLine(FString::Printf(TEXT("{Green}Shared Fragments:{White}%s"), *DescriptionBuilder(ItemNames)));
 
 					ItemNames.Reset();
-					Composition.GetConstSharedFragments().DebugGetIndividualNames(ItemNames);
+					Composition.Get<FMassConstSharedFragmentBitSet>().DebugGetIndividualNames(ItemNames);
 					AddTextLine(FString::Printf(TEXT("{Green}Const Shared Fragments:{White}%s"), *DescriptionBuilder(ItemNames)));
 				}
 
diff --git a/Engine/Plugins/Experimental/EditorDataStorage/Source/TedsCore/Private/Processors/TypedElementProcessorAdaptors.cpp b/Engine/Plugins/Experimental/EditorDataStorage/Source/TedsCore/Private/Processors/TypedElementProcessorAdaptors.cpp
--- a/Engine/Plugins/Experimental/EditorDataStorage/Source/TedsCore/Private/Processors/TypedElementProcessorAdaptors.cpp
+++ b/Engine/Plugins/Experimental/EditorDataStorage/Source/TedsCore/Private/Processors/TypedElementProcessorAdaptors.cpp
@@ -104,13 +104,9 @@ namespace UE::Editor::DataStorage
 				FMassArchetypeHandle Archetype = Manager.GetArchetypeForEntity(Entity);
 				const FMassArchetypeCompositionDescriptor& Composition = Manager.GetArchetypeComposition(Archetype);
 
-				if (UE::Mass::IsA<FMassTag>(ColumnType))
-				{
-					return Composition.GetTags().Contains(*ColumnType);
-				}
-				if (UE::Mass::IsA<FMassFragment>(ColumnType))
+				if (UE::Mass::IsA<FMassTag>(ColumnType) || UE::Mass::IsA<FMassFragment>(ColumnType))
 				{
-					return Composition.GetFragments().Contains(*ColumnType);
+					return Composition.GetElementsBitSet().Contains(*ColumnType);
 				}
 				const bool bIsTagOrFragment = false;
 				checkf(bIsTagOrFragment, TEXT("Attempting to check for a column type that is not a column or tag."));
@@ -213,7 +209,7 @@ namespace UE::Editor::DataStorage
 					{
 						if (this->Context.DoesArchetypeHaveTag(*ColumnType))
 						{
-							Descriptor.GetTags().Add(*ColumnType);
+							Descriptor.Add(ColumnType);
 						}
 					}
 					else
@@ -222,7 +218,7 @@ namespace UE::Editor::DataStorage
 							TEXT("Given struct type is not a valid fragment or tag type."));
 						if (this->Context.DoesArchetypeHaveFragment(*ColumnType))
 						{
-							Descriptor.GetFragments().Add(*ColumnType);
+							Descriptor.Add(ColumnType);
 						}
 					}
 				}
@@ -236,13 +232,13 @@ namespace UE::Editor::DataStorage
 				{
 					if (UE::Mass::IsA<FMassTag>(ColumnType))
 					{
-						Descriptor.GetTags().Add(*ColumnType);
+						Descriptor.Add(ColumnType);
 					}
 					else
 					{
 						checkf(UE::Mass::IsA<FMassFragment>(ColumnType),
 							TEXT("Given struct type is not a valid fragment or tag type."));
-						Descriptor.GetFragments().Add(*ColumnType);
+						Descriptor.Add(ColumnType);
 
 					}
 				}
@@ -284,7 +280,7 @@ namespace UE::Editor::DataStorage
 							UniqueAdd;
 						for (const FAddedColumns* const AddColumns : AddedColumnsQueue)
 						{
-							UniqueAdd.Add(TPair<FMassTagBitSet, FMassFragmentBitSet>(AddColumns->AddDescriptor.GetTags(), AddColumns->AddDescriptor.GetFragments()));
+							UniqueAdd.Add(TPair<FMassTagBitSet, FMassFragmentBitSet>(AddColumns->AddDescriptor.Get<FMassTagBitSet>(), AddColumns->AddDescriptor.Get<FMassFragmentBitSet>()));
 						}
 						const int32 AddedColumnsRowGuess = UniqueAdd.IsEmpty() ? 0 : AddedColumnsQueue.Num() / UniqueAdd.Num();
 						UniqueAdd.Empty();
@@ -343,8 +339,8 @@ namespace UE::Editor::DataStorage
 						for (FAddedColumns* const AddColumns : AddedColumnsQueue)
 						{
 							FRowOperations& Operations = RowOperations.FindOrAdd(AddColumns->Entity);
-							Operations.ColumnsChange.AddedTags = AddColumns->AddDescriptor.GetTags();
-							Operations.ColumnsChange.AddedFragments = AddColumns->AddDescriptor.GetFragments();
+							Operations.ColumnsChange.AddedTags = AddColumns->AddDescriptor.Get<FMassTagBitSet>();
+							Operations.ColumnsChange.AddedFragments = AddColumns->AddDescriptor.Get<FMassFragmentBitSet>();
 						}
 						AddedColumnsQueue.Empty();
 
@@ -353,8 +349,8 @@ namespace UE::Editor::DataStorage
 							for (const FMassEntityHandle& Row : AddColumns->Entities)
 							{
 								FRowOperations& Operations = RowOperations.FindOrAdd(Row);
-								Operations.ColumnsChange.AddedTags += AddColumns->AddDescriptor.GetTags();
-								Operations.ColumnsChange.AddedFragments += AddColumns->AddDescriptor.GetFragments();
+								Operations.ColumnsChange.AddedTags += AddColumns->AddDescriptor.Get<FMassTagBitSet>();
+								Operations.ColumnsChange.AddedFragments += AddColumns->AddDescriptor.Get<FMassFragmentBitSet>();
 							}
 						}
 						AddedColumnsBatchQueue.Empty();
@@ -397,8 +393,8 @@ namespace UE::Editor::DataStorage
 						for (FRemovedColumns* const RemovedColumns : RemovedColumnsQueue)
 						{
 							FRowOperations& Operations = RowOperations.FindOrAdd(RemovedColumns->Entity);
-							Operations.ColumnsChange.RemovedTags += RemovedColumns->RemoveDescriptor.GetTags();
-							Operations.ColumnsChange.RemovedFragments += RemovedColumns->RemoveDescriptor.GetFragments();
+							Operations.ColumnsChange.RemovedTags += RemovedColumns->RemoveDescriptor.Get<FMassTagBitSet>();
+							Operations.ColumnsChange.RemovedFragments += RemovedColumns->RemoveDescriptor.Get<FMassFragmentBitSet>();
 						}
 						RemovedColumnsQueue.Empty();
 
@@ -407,8 +403,8 @@ namespace UE::Editor::DataStorage
 							for (const FMassEntityHandle& Row : RemovedColumns->Entities)
 							{
 								FRowOperations& Operations = RowOperations.FindOrAdd(Row);
-								Operations.ColumnsChange.RemovedTags += RemovedColumns->RemoveDescriptor.GetTags();
-								Operations.ColumnsChange.RemovedFragments += RemovedColumns->RemoveDescriptor.GetFragments();
+								Operations.ColumnsChange.RemovedTags += RemovedColumns->RemoveDescriptor.Get<FMassTagBitSet>();
+								Operations.ColumnsChange.RemovedFragments += RemovedColumns->RemoveDescriptor.Get<FMassFragmentBitSet>();
 							}
 						}
 						RemovedColumnsBatchQueue.Empty();
@@ -866,7 +862,7 @@ namespace UE::Editor::DataStorage
 
 
 							// This could be improved by adding an operation that would both combine the Fragments and Tags change in one bath operation.
-							if (!RemovedColumns->RemoveDescriptor.GetFragments().IsEmpty())
+							if (!RemovedColumns->RemoveDescriptor.Get<FMassFragmentBitSet>().IsEmpty())
 							{
 								for (FMassEntityHandle Entity : Entities)
 								{
@@ -886,9 +882,9 @@ namespace UE::Editor::DataStorage
 									EntityCollections.Emplace(It.Key(), It.Value(), FMassArchetypeEntityCollection::EDuplicatesHandling::FoldDuplicates);
 								}
 
-								System.BatchChangeFragmentCompositionForEntities(EntityCollections, FMassFragmentBitSet(), RemovedColumns->RemoveDescriptor.GetFragments());
+								System.BatchChangeFragmentCompositionForEntities(EntityCollections, FMassFragmentBitSet(), RemovedColumns->RemoveDescriptor.Get<FMassFragmentBitSet>());
 							}
-							if (!RemovedColumns->RemoveDescriptor.GetTags().IsEmpty())
+							if (!RemovedColumns->RemoveDescriptor.Get<FMassTagBitSet>().IsEmpty())
 							{
 								LookupTable.Reset();
 								EntityCollections.Reset();
@@ -911,7 +907,7 @@ namespace UE::Editor::DataStorage
 									EntityCollections.Emplace(It.Key(), It.Value(), FMassArchetypeEntityCollection::EDuplicatesHandling::FoldDuplicates);
 								}
 
-								System.BatchChangeTagsForEntities(EntityCollections, FMassTagBitSet(), RemovedColumns->RemoveDescriptor.GetTags());
+								System.BatchChangeTagsForEntities(EntityCollections, FMassTagBitSet(), RemovedColumns->RemoveDescriptor.Get<FMassTagBitSet>());
 							}
 						});
 					}
diff --git a/Engine/Plugins/Experimental/EditorDataStorage/Source/TedsCore/Private/Queries/TypedElementExtendedQueryStore.cpp b/Engine/Plugins/Experimental/EditorDataStorage/Source/TedsCore/Private/Queries/TypedElementExtendedQueryStore.cpp
--- a/Engine/Plugins/Experimental/EditorDataStorage/Source/TedsCore/Private/Queries/TypedElementExtendedQueryStore.cpp
+++ b/Engine/Plugins/Experimental/EditorDataStorage/Source/TedsCore/Private/Queries/TypedElementExtendedQueryStore.cpp
@@ -54,10 +54,11 @@ namespace UE::Editor::DataStorage
 		StringBuilder.Append( TEXT("Composition Descriptor:\nColumns: {"));
 		{
 			int32 FragmentsProcessed = 0;
-			const int32 FragmentCount = CompositionDescriptor.GetFragments().CountStoredTypes();
-			for(auto IndexIterator = CompositionDescriptor.GetFragments().GetIndexIterator(); IndexIterator; ++IndexIterator)
+			const FMassFragmentBitSet FragmentBitSet = CompositionDescriptor.Get<FMassFragmentBitSet>();
+			const int32 FragmentCount = FragmentBitSet.CountStoredTypes();
+			for(auto IndexIterator = FragmentBitSet.GetIndexIterator(); IndexIterator; ++IndexIterator)
 			{
-				StringBuilder.Append(CompositionDescriptor.GetFragments().GetTypeAtIndex(*IndexIterator)->GetName());
+				StringBuilder.Append(FragmentBitSet.GetTypeAtIndex(*IndexIterator)->GetName());
 				if (FragmentsProcessed != FragmentCount - 1)
 				{
 					StringBuilder.Append(TEXT(", "));
@@ -68,11 +69,12 @@ namespace UE::Editor::DataStorage
 		}
 		{
 			int32 TagsProcessed = 0;
-			const int32 TagsCount = CompositionDescriptor.GetTags().CountStoredTypes();
+			const FMassTagBitSet TagBitSet = CompositionDescriptor.Get<FMassTagBitSet>();
+			const int32 TagsCount = TagBitSet.CountStoredTypes();
 			StringBuilder.Append(TEXT("Tags: {"));
-			for(auto IndexIterator = CompositionDescriptor.GetTags().GetIndexIterator(); IndexIterator; ++IndexIterator)
+			for(auto IndexIterator = TagBitSet.GetIndexIterator(); IndexIterator; ++IndexIterator)
 			{
-				StringBuilder.Append(CompositionDescriptor.GetTags().GetTypeAtIndex(*IndexIterator)->GetName());
+				StringBuilder.Append(TagBitSet.GetTypeAtIndex(*IndexIterator)->GetName());
 				if (TagsProcessed != TagsCount - 1)
 				{
 					StringBuilder.Append(TEXT(", "));
@@ -96,10 +98,11 @@ namespace UE::Editor::DataStorage
 		if (Query->Description.Conditions)
 		{
 			TArray<TWeakObjectPtr<const UScriptStruct>, TInlineAllocator<64>> ColumnTypes;
-			auto FragmentIterator = CompositionDescriptor.GetFragments().GetIndexIterator(true);
+			FMassFragmentBitSet FragmentBitSet = CompositionDescriptor.Get<FMassFragmentBitSet>();
+			auto FragmentIterator = FragmentBitSet.GetIndexIterator(true);
 			while (FragmentIterator)
 			{
-				const UScriptStruct* FragmentStruct = CompositionDescriptor.GetFragments().GetTypeAtIndex(*FragmentIterator);
+				const UScriptStruct* FragmentStruct = FragmentBitSet.GetTypeAtIndex(*FragmentIterator);
 				ColumnTypes.Add(FragmentStruct);
 				++FragmentIterator;
 			}
@@ -118,10 +121,11 @@ namespace UE::Editor::DataStorage
 			}
 
 			// Then ensure it also matches the query conditions
-			auto TagIterator = CompositionDescriptor.GetTags().GetIndexIterator(true);			
+			const FMassTagBitSet TagBitSet = CompositionDescriptor.Get<FMassTagBitSet>();
+			auto TagIterator = TagBitSet.GetIndexIterator(true);
 			while (TagIterator)
 			{
-				const UScriptStruct* TagStruct = CompositionDescriptor.GetTags().GetTypeAtIndex(*TagIterator);
+				const UScriptStruct* TagStruct = TagBitSet.GetTypeAtIndex(*TagIterator);
 				ColumnTypes.Add(TagStruct);
 				++TagIterator;
 			}
diff --git a/Engine/Plugins/Experimental/EditorDataStorage/Source/TedsCore/Private/TypedElementDatabase.cpp b/Engine/Plugins/Experimental/EditorDataStorage/Source/TedsCore/Private/TypedElementDatabase.cpp
--- a/Engine/Plugins/Experimental/EditorDataStorage/Source/TedsCore/Private/TypedElementDatabase.cpp
+++ b/Engine/Plugins/Experimental/EditorDataStorage/Source/TedsCore/Private/TypedElementDatabase.cpp
@@ -911,13 +911,9 @@ bool UEditorDataStorage::HasColumns(UE::Editor::DataStorage::RowHandle Row, TCon
 			const UScriptStruct* const* ColumnTypesEnd = ColumnTypes.end();
 			for (const UScriptStruct* const* ColumnType = ColumnTypes.begin(); ColumnType != ColumnTypesEnd && bHasAllColumns; ++ColumnType)
 			{
-				if (UE::Mass::IsA<FMassFragment>(*ColumnType))
+				if (UE::Mass::IsA<FMassFragment>(*ColumnType) || UE::Mass::IsA<FMassTag>(*ColumnType))
 				{
-					bHasAllColumns = Composition.GetFragments().Contains(**ColumnType);
-				}
-				else if (UE::Mass::IsA<FMassTag>(*ColumnType))
-				{
-					bHasAllColumns = Composition.GetTags().Contains(**ColumnType);
+					bHasAllColumns = Composition.GetElementsBitSet().Contains(**ColumnType);
 				}
 				else
 				{
@@ -957,14 +953,9 @@ bool UEditorDataStorage::HasColumns(UE::Editor::DataStorage::RowHandle Row, TCon
 			{
 				if (ColumnType->IsValid())
 				{
-					if (UE::Mass::IsA<FMassFragment>(ColumnType->Get()))
-					{
-						bHasAllColumns = Composition.GetFragments().Contains(**ColumnType);
-						continue;
-					}
-					else if (UE::Mass::IsA<FMassTag>(ColumnType->Get()))
+					if (UE::Mass::IsA<FMassFragment>(ColumnType->Get()) || UE::Mass::IsA<FMassTag>(ColumnType->Get()))
 					{
-						bHasAllColumns = Composition.GetTags().Contains(**ColumnType);
+						bHasAllColumns = Composition.GetElementsBitSet().Contains(**ColumnType);
 						continue;
 					}
 				}
@@ -1004,8 +995,7 @@ void UEditorDataStorage::ListColumns(UE::Editor::DataStorage::RowHandle Row, UE:
 					}
 					return true;
 				};
-			Composition.GetFragments().ExportTypes(CallbackWrapper);
-			Composition.GetTags().ExportTypes(CallbackWrapper);
+			Composition.GetElementsBitSet().ExportTypes(CallbackWrapper);
 		}
 		else
 		{
@@ -1024,7 +1014,8 @@ void UEditorDataStorage::ListColumns(UE::Editor::DataStorage::RowHandle Row, UE:
 			FMassArchetypeHandle Archetype = ActiveEditorEntityManager->GetArchetypeForEntity(Entity);
 			const FMassArchetypeCompositionDescriptor& Composition = ActiveEditorEntityManager->GetArchetypeComposition(Archetype);
 
-			Composition.GetFragments().ExportTypes(
+			// @todo consider a unified path for FMassElementBitSet
+			Composition.Get<FMassFragmentBitSet>().ExportTypes(
 				[this, &Callback, Entity](const UScriptStruct* ColumnType)
 				{
 					if (ColumnType)
@@ -1033,7 +1024,7 @@ void UEditorDataStorage::ListColumns(UE::Editor::DataStorage::RowHandle Row, UE:
 					}
 					return true;
 				});
-			Composition.GetTags().ExportTypes(
+			Composition.Get<FMassTagBitSet>().ExportTypes(
 				[&Callback](const UScriptStruct* ColumnType)
 				{
 					if (ColumnType)
diff --git a/Engine/Plugins/Runtime/MassGameplay/Source/MassGameplayEditor/Private/MassTraitRepository.cpp b/Engine/Plugins/Runtime/MassGameplay/Source/MassGameplayEditor/Private/MassTraitRepository.cpp
--- a/Engine/Plugins/Runtime/MassGameplay/Source/MassGameplayEditor/Private/MassTraitRepository.cpp
+++ b/Engine/Plugins/Runtime/MassGameplay/Source/MassGameplayEditor/Private/MassTraitRepository.cpp
@@ -365,33 +365,9 @@ void UMassTraitRepository::OnNewTraitType(UMassEntityTraitBase& Trait)
 
 	const FMassArchetypeCompositionDescriptor& Composition = InvestigationContext.EntityTemplate.GetCompositionDescriptor();
 
-	for (FMassFragmentBitSet::FIndexIterator It = Composition.GetFragments().GetIndexIterator(); It; ++It)
+	for (FMassFragmentBitSet::FIndexIterator It = Composition.GetElementsBitSet().GetIndexIterator(); It; ++It)
 	{
-		const FName StructName = Composition.GetFragments().DebugGetStructTypeName(*It);
-		TraitData.ElementNames.Add(StructName);
-	}
-
-	for (FMassTagBitSet::FIndexIterator It = Composition.GetTags().GetIndexIterator(); It; ++It)
-	{
-		const FName StructName = Composition.GetTags().DebugGetStructTypeName(*It);
-		TraitData.ElementNames.Add(StructName);
-	}
-
-	for (FMassChunkFragmentBitSet::FIndexIterator It = Composition.GetChunkFragments().GetIndexIterator(); It; ++It)
-	{
-		const FName StructName = Composition.GetChunkFragments().DebugGetStructTypeName(*It);
-		TraitData.ElementNames.Add(StructName);
-	}
-
-	for (FMassSharedFragmentBitSet::FIndexIterator It = Composition.GetSharedFragments().GetIndexIterator(); It; ++It)
-	{
-		const FName StructName = Composition.GetSharedFragments().DebugGetStructTypeName(*It);
-		TraitData.ElementNames.Add(StructName);
-	}
-
-	for (FMassConstSharedFragmentBitSet::FIndexIterator It = Composition.GetConstSharedFragments().GetIndexIterator(); It; ++It)
-	{
-		const FName StructName = Composition.GetConstSharedFragments().DebugGetStructTypeName(*It);
+		const FName StructName = Composition.GetElementsBitSet().DebugGetStructTypeName(*It);
 		TraitData.ElementNames.Add(StructName);
 	}
 
diff --git a/Engine/Plugins/Runtime/MassGameplay/Source/MassGameplayTestSuite/Private/MassEntityTemplateTest.cpp b/Engine/Plugins/Runtime/MassGameplay/Source/MassGameplayTestSuite/Private/MassEntityTemplateTest.cpp
--- a/Engine/Plugins/Runtime/MassGameplay/Source/MassGameplayTestSuite/Private/MassEntityTemplateTest.cpp
+++ b/Engine/Plugins/Runtime/MassGameplay/Source/MassGameplayTestSuite/Private/MassEntityTemplateTest.cpp
@@ -75,11 +75,11 @@ struct FEntityTemplate_Composition : FEntityTemplateBase
 		TemplateData.Sort();
 
 		FMassArchetypeCompositionDescriptor ExpectedComposition;
-		ExpectedComposition.GetFragments().Add<FTestFragment_Int>();
-		ExpectedComposition.GetFragments().Add<FTestFragment_Float>();
-		ExpectedComposition.GetChunkFragments().Add<FTestChunkFragment_Int>();
-		ExpectedComposition.GetSharedFragments().Add<FTestSharedFragment_Int>();
-		ExpectedComposition.GetSharedFragments().Add<FTestSharedFragment_Float>();
+		ExpectedComposition.Add<FTestFragment_Int>();
+		ExpectedComposition.Add<FTestFragment_Float>();
+		ExpectedComposition.Add<FTestChunkFragment_Int>();
+		ExpectedComposition.Add<FTestSharedFragment_Int>();
+		ExpectedComposition.Add<FTestSharedFragment_Float>();
 
 		AITEST_TRUE("The composition should end up being the same regardless of whether the data is added via a template data or a composition descriptor"
 			, ExpectedComposition.IsEquivalent(TemplateData.GetCompositionDescriptor()));
diff --git a/Engine/Plugins/Runtime/MassGameplay/Source/MassSpawner/Public/MassEntityTemplate.h b/Engine/Plugins/Runtime/MassGameplay/Source/MassSpawner/Public/MassEntityTemplate.h
--- a/Engine/Plugins/Runtime/MassGameplay/Source/MassSpawner/Public/MassEntityTemplate.h
+++ b/Engine/Plugins/Runtime/MassGameplay/Source/MassSpawner/Public/MassEntityTemplate.h
@@ -130,13 +130,13 @@ struct FMassEntityTemplateData
 	void AddFragment()
 	{
 		static_assert(UE::Mass::CFragment<T>, MASS_INVALID_FRAGMENT_MSG);
-		Composition.GetFragments().Add<T>();
+		Composition.GetElementsBitSet().Add<T>();
 	}
 
 	void AddFragment(const UScriptStruct& FragmentType)
 	{
 		checkf(UE::Mass::IsA<FMassFragment>(&FragmentType), TEXT(MASS_INVALID_FRAGMENT_MSG_F), *FragmentType.GetName());
-		Composition.GetFragments().Add(FragmentType);
+		Composition.GetElementsBitSet().Add(FragmentType);
 	}
 
 	// @todo this function is doing nothing if a given fragment's initial value has already been created. This seems inconsistent with the other AddFragment functions (especially AddFragment_GetRef).
@@ -145,9 +145,9 @@ struct FMassEntityTemplateData
 		const UScriptStruct* FragmentType = Fragment.GetScriptStruct();
 		checkf(UE::Mass::IsA<FMassFragment>(FragmentType), TEXT(MASS_INVALID_FRAGMENT_MSG_F), *GetNameSafe(FragmentType));
 
-		if (!Composition.GetFragments().Contains(*FragmentType))
+		if (!Composition.GetElementsBitSet().Contains(*FragmentType))
 		{
-			Composition.GetFragments().Add(*FragmentType);
+			Composition.GetElementsBitSet().Add(*FragmentType);
 			InitialFragmentValues.Emplace(Fragment);
 		}
 		else if (!InitialFragmentValues.ContainsByPredicate(FStructTypeEqualOperator(FragmentType)))
@@ -160,9 +160,9 @@ struct FMassEntityTemplateData
 	T& AddFragment_GetRef()
 	{
 		static_assert(UE::Mass::CFragment<T>, MASS_INVALID_FRAGMENT_MSG);
-		if (!Composition.GetFragments().Contains<T>())
+		if (!Composition.GetElementsBitSet().Contains<T>())
 		{
-			Composition.GetFragments().Add<T>();
+			Composition.GetElementsBitSet().Add<T>();
 		}
 		else if (FInstancedStruct* Fragment = InitialFragmentValues.FindByPredicate(FStructTypeEqualOperator(T::StaticStruct())))
 		{
@@ -185,36 +185,36 @@ struct FMassEntityTemplateData
 	void AddTag()
 	{
 		static_assert(UE::Mass::CTag<T>, "Given struct doesn't represent a valid mass tag type. Make sure to inherit from FMassTag or one of its child-types.");
-		Composition.GetTags().Add<T>();
+		Composition.GetElementsBitSet().Add<T>();
 	}
 	
 	void AddTag(const UScriptStruct& TagType)
 	{
 		checkf(UE::Mass::IsA<FMassTag>(&TagType), TEXT("Given struct doesn't represent a valid mass tag type. Make sure to inherit from FMassTag or one of its child-types."));
-		Composition.GetTags().Add(TagType);
+		Composition.GetElementsBitSet().Add(TagType);
 	}
 
 	template<typename T>
 	void RemoveTag()
 	{
 		static_assert(UE::Mass::CTag<T>, "Given struct doesn't represent a valid mass tag type. Make sure to inherit from FMassTag or one of its child-types.");
-		Composition.GetTags().Remove<T>();
+		Composition.GetElementsBitSet().Remove<T>();
 	}
 
 	void RemoveTag(const UScriptStruct& TagType)
 	{
 		checkf(UE::Mass::IsA<FMassTag>(&TagType), TEXT("Given struct doesn't represent a valid mass tag type. Make sure to inherit from FMassTag or one of its child-types."));
-		Composition.GetTags().Remove(TagType);
+		Composition.GetElementsBitSet().Remove(TagType);
 	}
 
-	const FMassTagBitSet& GetTags() const;
-	FMassTagBitSet& GetMutableTags();
+	FMassTagBitSet GetTags() const;
+	FMassElementBitSet& GetMutableTags();
 
 	template<typename T>
 	void AddChunkFragment()
 	{
 		static_assert(UE::Mass::CChunkFragment<T>, "Given struct doesn't represent a valid chunk fragment type. Make sure to inherit from FMassChunkFragment or one of its child-types.");
-		Composition.GetChunkFragments().Add<T>();
+		Composition.GetElementsBitSet().Add<T>();
 	}
 
 	void AddConstSharedFragment(const FConstSharedStruct& SharedFragment)
@@ -222,9 +222,9 @@ struct FMassEntityTemplateData
 		const UScriptStruct* FragmentType = SharedFragment.GetScriptStruct();
 		if(ensureMsgf(UE::Mass::IsA<FMassConstSharedFragment>(FragmentType), TEXT("Given struct doesn't represent a valid const shared fragment type. Make sure to inherit from FMassConstSharedFragment or one of its child-types.")))
 		{
-			if (!Composition.GetConstSharedFragments().Contains(*FragmentType))
+			if (!Composition.GetElementsBitSet().Contains(*FragmentType))
 			{
-				Composition.GetConstSharedFragments().Add(*FragmentType);
+				Composition.GetElementsBitSet().Add(*FragmentType);
 				SharedFragmentValues.Add(SharedFragment);
 			}
 #if DO_ENSURE
@@ -243,9 +243,9 @@ struct FMassEntityTemplateData
 		const UScriptStruct* FragmentType = SharedFragment.GetScriptStruct();
 		if(ensureMsgf(UE::Mass::IsA<FMassSharedFragment>(FragmentType), TEXT("Given struct doesn't represent a valid shared fragment type. Make sure to inherit from FMassSharedFragment or one of its child-types.")))
 		{
-			if (!Composition.GetSharedFragments().Contains(*FragmentType))
+			if (!Composition.GetElementsBitSet().Contains(*FragmentType))
 			{
-				Composition.GetSharedFragments().Add(*FragmentType);
+				Composition.GetElementsBitSet().Add(*FragmentType);
 				SharedFragmentValues.Add(SharedFragment);
 			}
 	#if DO_ENSURE
@@ -259,49 +259,52 @@ struct FMassEntityTemplateData
 		}
 	}
 
-	template<typename T>
+	template<UE::Mass::CFragment T>
 	bool HasFragment() const
 	{
-		return Composition.GetFragments().Contains<T>();
+		return Composition.GetElementsBitSet().Contains<T>();
 	}
 	
 	bool HasFragment(const UScriptStruct& ScriptStruct) const
 	{
-		return Composition.GetFragments().Contains(ScriptStruct);
+		check(UE::Mass::IsA<FMassFragment>(&ScriptStruct))
+		return Composition.GetElementsBitSet().Contains(ScriptStruct);
 	}
 
-	template<typename T>
+	template<UE::Mass::CTag T>
 	bool HasTag() const
 	{
-		return Composition.GetTags().Contains<T>();
+		return Composition.GetElementsBitSet().Contains<T>();
 	}
 
-	template<typename T>
+	template<UE::Mass::CChunkFragment T>
 	bool HasChunkFragment() const
 	{
-		return Composition.GetChunkFragments().Contains<T>();
+		return Composition.GetElementsBitSet().Contains<T>();
 	}
 
-	template<typename T>
+	template<UE::Mass::CSharedFragment T>
 	bool HasSharedFragment() const
 	{
-		return Composition.GetSharedFragments().Contains<T>();
+		return Composition.GetElementsBitSet().Contains<T>();
 	}
 
 	bool HasSharedFragment(const UScriptStruct& ScriptStruct) const
 	{
-		return Composition.GetSharedFragments().Contains(ScriptStruct);
+		check(UE::Mass::IsA<FMassSharedFragment>(&ScriptStruct))
+		return Composition.GetElementsBitSet().Contains(ScriptStruct);
 	}
 
-	template<typename T>
+	template<UE::Mass::CConstSharedFragment T>
 	bool HasConstSharedFragment() const
 	{
-		return Composition.GetConstSharedFragments().Contains<T>();
+		return Composition.GetElementsBitSet().Contains<T>();
 	}
 
 	bool HasConstSharedFragment(const UScriptStruct& ScriptStruct) const
 	{
-		return Composition.GetConstSharedFragments().Contains(ScriptStruct);
+		check(UE::Mass::IsA<FMassConstSharedFragment>(&ScriptStruct))
+		return Composition.GetElementsBitSet().Contains(ScriptStruct);
 	}
 
 	void Sort()
@@ -429,14 +432,14 @@ inline void FMassEntityTemplateData::SetTemplateName(const FString& Name)
 	TemplateName = Name; 
 }
 
-inline const FMassTagBitSet& FMassEntityTemplateData::GetTags() const
+inline FMassTagBitSet FMassEntityTemplateData::GetTags() const
 {
-	return Composition.GetTags();
+	return Composition.GetElementsBitSet().Get<FMassTagBitSet>();
 }
 
-inline FMassTagBitSet& FMassEntityTemplateData::GetMutableTags()
+inline FMassElementBitSet& FMassEntityTemplateData::GetMutableTags()
 {
-	return Composition.GetTags();
+	return Composition.GetElementsBitSet();
 }
 
 inline FMassArchetypeCreationParams& FMassEntityTemplateData::GetArchetypeCreationParams()
diff --git a/Engine/Plugins/Runtime/MassGameplay/Source/MassSpawner/Public/MassTranslator.h b/Engine/Plugins/Runtime/MassGameplay/Source/MassSpawner/Public/MassTranslator.h
--- a/Engine/Plugins/Runtime/MassGameplay/Source/MassSpawner/Public/MassTranslator.h
+++ b/Engine/Plugins/Runtime/MassGameplay/Source/MassSpawner/Public/MassTranslator.h
@@ -38,8 +38,16 @@ public:
 	 * 
 	 *  @todo might want this function on the MassProcessor level. TBD
 	 * 
-	 *  @param OutTagTypes tag types will be appended to this array. No uniqueness checks are performed. */
-	void AppendRequiredTags(FMassTagBitSet& InOutTags) const { InOutTags += RequiredTags; }
+	 *  @param InOutTags tag types will be appended to this array. No uniqueness checks are performed. */
+	void AppendRequiredTags(FMassTagBitSet& InOutTags) const 
+	{
+		InOutTags += RequiredTags; 
+	}
+
+	void AppendRequiredTags(FMassElementBitSet& InOutTags) const 
+	{
+		InOutTags += RequiredTags; 
+	}
 
 protected:
 	UE_API void AddRequiredTagsToQuery(FMassEntityQuery& EntityQuery);
diff --git a/Engine/Source/Developer/MassEntityTestSuite/Private/MassQueryTest.cpp b/Engine/Source/Developer/MassEntityTestSuite/Private/MassQueryTest.cpp
--- a/Engine/Source/Developer/MassEntityTestSuite/Private/MassQueryTest.cpp
+++ b/Engine/Source/Developer/MassEntityTestSuite/Private/MassQueryTest.cpp
@@ -609,13 +609,13 @@ struct FQueryTest_AllOptional : FEntityTestBase
 
 		{
 			FMassArchetypeCompositionDescriptor Descriptor(EntityManager->GetArchetypeComposition(IntsArchetype));
-			Descriptor.GetChunkFragments().Add<FTestChunkFragment_Int>();
+			Descriptor.Add<FTestChunkFragment_Int>();
 			EntityManager->CreateArchetype(Descriptor);
 			++ExpectedNumOfArchetypes;
 		}
 		{
 			FMassArchetypeCompositionDescriptor Descriptor(EntityManager->GetArchetypeComposition(IntsArchetype));
-			Descriptor.GetChunkFragments().Add<FTestChunkFragment_Float>();
+			Descriptor.Add<FTestChunkFragment_Float>();
 			EntityManager->CreateArchetype(Descriptor);
 		}
 		Query.CacheArchetypes();
@@ -678,13 +678,13 @@ struct FQueryTest_JustATag : FEntityTestBase
 
 		{
 			FMassArchetypeCompositionDescriptor Descriptor(EntityManager->GetArchetypeComposition(IntsArchetype));
-			Descriptor.GetTags().Add<FTestTag_A>();
+			Descriptor.Add<FTestTag_A>();
 			EntityManager->CreateArchetype(Descriptor);
 			++ExpectedNumOfArchetypes;
 		}
 		{
 			FMassArchetypeCompositionDescriptor Descriptor(EntityManager->GetArchetypeComposition(IntsArchetype));
-			Descriptor.GetTags().Add<FTestTag_B>();
+			Descriptor.Add<FTestTag_B>();
 			EntityManager->CreateArchetype(Descriptor);
 		}
 		Query.CacheArchetypes();
@@ -692,17 +692,17 @@ struct FQueryTest_JustATag : FEntityTestBase
 
 		{
 			FMassArchetypeCompositionDescriptor Descriptor(EntityManager->GetArchetypeComposition(IntsArchetype));
-			Descriptor.GetTags().Add<FTestTag_A>();
-			Descriptor.GetTags().Add<FTestTag_C>();
-			Descriptor.GetTags().Add<FTestTag_D>();
+			Descriptor.Add<FTestTag_A>();
+			Descriptor.Add<FTestTag_C>();
+			Descriptor.Add<FTestTag_D>();
 			EntityManager->CreateArchetype(Descriptor);
 			++ExpectedNumOfArchetypes;
 		}
 		{
 			FMassArchetypeCompositionDescriptor Descriptor(EntityManager->GetArchetypeComposition(IntsArchetype));
-			Descriptor.GetTags().Add<FTestTag_B>();
-			Descriptor.GetTags().Add<FTestTag_C>();
-			Descriptor.GetTags().Add<FTestTag_D>();
+			Descriptor.Add<FTestTag_B>();
+			Descriptor.Add<FTestTag_C>();
+			Descriptor.Add<FTestTag_D>();
 			EntityManager->CreateArchetype(Descriptor);
 		}
 		Query.CacheArchetypes();
@@ -731,13 +731,13 @@ struct FQueryTest_JustAChunkFragment : FEntityTestBase
 
 		{
 			FMassArchetypeCompositionDescriptor Descriptor(EntityManager->GetArchetypeComposition(IntsArchetype));
-			Descriptor.GetChunkFragments().Add<FTestChunkFragment_Int>();
+			Descriptor.Add<FTestChunkFragment_Int>();
 			TargetArchetype = EntityManager->CreateArchetype(Descriptor);
 			++ExpectedNumOfArchetypes;
 		}
 		{
 			FMassArchetypeCompositionDescriptor Descriptor(EntityManager->GetArchetypeComposition(IntsArchetype));
-			Descriptor.GetChunkFragments().Add<FTestChunkFragment_Float>();
+			Descriptor.Add<FTestChunkFragment_Float>();
 			EntityManager->CreateArchetype(Descriptor);
 		}
 		Query.CacheArchetypes();
diff --git a/Engine/Source/Developer/MassEntityTestSuite/Private/MassSharedFragmentsTest.cpp b/Engine/Source/Developer/MassEntityTestSuite/Private/MassSharedFragmentsTest.cpp
--- a/Engine/Source/Developer/MassEntityTestSuite/Private/MassSharedFragmentsTest.cpp
+++ b/Engine/Source/Developer/MassEntityTestSuite/Private/MassSharedFragmentsTest.cpp
@@ -32,6 +32,16 @@ TMassSharedFragmentType* GetMutableSharedFragmentPtr(FMassArchetypeSharedFragmen
 	return SharedStruct.IsValid() ? SharedStruct.GetPtr<TMassSharedFragmentType>() : nullptr;
 }
 
+bool HasExactSharedFragmentTypesMatch(const FMassArchetypeSharedFragmentValues& Values, const FMassSharedFragmentBitSet SharedFragmentBitSet)
+{
+	return Values.GetBitSet().Get<FMassSharedFragmentBitSet>() == SharedFragmentBitSet;
+}
+
+bool HasExactConstSharedFragmentTypesMatch(const FMassArchetypeSharedFragmentValues& Values, const FMassConstSharedFragmentBitSet SharedFragmentBitSet)
+{
+	return Values.GetBitSet().Get<FMassConstSharedFragmentBitSet>() == SharedFragmentBitSet;
+}
+
 //-----------------------------------------------------------------------------
 // FMassArchetypeSharedFragmentValues Tests
 //-----------------------------------------------------------------------------
@@ -201,13 +211,13 @@ struct FSharedFragmentValues_Remove : FExecutionTestBase
 				AITEST_TRUE("Assignment should result in same values", Values.HasSameValues(ValuestNonConstIntFloat));
 
 				// removing just the Int shared fragment
-				Values.Remove(ValuesNonConstInt.GetSharedFragmentBitSet());
+				Values.Remove(ValuesNonConstInt);
 				AITEST_TRUE("#1 Removal results should match expectations", Values.HasSameValues(ValuesNonConstFloat));
 			}
 			{
 				FMassArchetypeSharedFragmentValues Values = ValuestNonConstIntFloat;
 				// removing just the Float shared fragment
-				Values.Remove(ValuesNonConstFloat.GetSharedFragmentBitSet());
+				Values.Remove(ValuesNonConstFloat);
 				AITEST_TRUE("#2 Removal results should match expectations", Values.HasSameValues(ValuesNonConstInt));
 			}
 		}
@@ -225,13 +235,13 @@ struct FSharedFragmentValues_Remove : FExecutionTestBase
 				AITEST_TRUE("Assignment should result in same values", Values.HasSameValues(ValuestConstIntFloat));
 
 				// removing just the Int shared fragment
-				Values.Remove(ValuesConstInt.GetConstSharedFragmentBitSet());
+				Values.Remove(ValuesConstInt);
 				AITEST_TRUE("#3 Removal results should match expectations", Values.HasSameValues(ValuesConstFloat));
 			}
 			{
 				FMassArchetypeSharedFragmentValues Values = ValuestConstIntFloat;
 				// removing just the Float shared fragment
-				Values.Remove(ValuesConstFloat.GetConstSharedFragmentBitSet());
+				Values.Remove(ValuesConstFloat);
 				AITEST_TRUE("#4 Removal results should match expectations", Values.HasSameValues(ValuesConstInt));
 			}
 		}
@@ -771,29 +781,29 @@ struct FSharedFragmentValues_TypeEquivalency : FEntityTestBase
 		const FMassSharedFragmentBitSet EmptySharedFragmentBitSet;
 		const FMassConstSharedFragmentBitSet EmptyConstSharedFragmentBitSet;
 
-		AITEST_TRUE("Empty shared values match type with empty bitset", Values.HasExactSharedFragmentTypesMatch(EmptySharedFragmentBitSet));
-		AITEST_TRUE("Empty const shared values match type with empty const bitset", Values.HasExactConstSharedFragmentTypesMatch(EmptyConstSharedFragmentBitSet));
+		AITEST_TRUE("Empty shared values match type with empty bitset", HasExactSharedFragmentTypesMatch(Values, EmptySharedFragmentBitSet));
+		AITEST_TRUE("Empty const shared values match type with empty const bitset", HasExactConstSharedFragmentTypesMatch(Values, EmptyConstSharedFragmentBitSet));
 
 		const FMassSharedFragmentBitSet IntSharedFragmentBitSet = FMassSharedFragmentBitSet::GetTypeBitSet<FTestSharedFragment_Int>();
 		FMassSharedFragmentBitSet IntFloatSharedFragmentBitSet = IntSharedFragmentBitSet;
 		IntFloatSharedFragmentBitSet.Add<FTestSharedFragment_Float>();
 		Values.Add(FSharedStruct::Make<FTestSharedFragment_Int>(TestIntValue));
-		AITEST_TRUE("Single shared value type matches expected bitset", Values.HasExactSharedFragmentTypesMatch(IntSharedFragmentBitSet));
-		AITEST_FALSE("Single shared value type doesn't match two-type bitset", Values.HasExactSharedFragmentTypesMatch(IntFloatSharedFragmentBitSet));
-		AITEST_FALSE("Single shared value type doesn't match empty", Values.HasExactSharedFragmentTypesMatch(EmptySharedFragmentBitSet));
+		AITEST_TRUE("Single shared value type matches expected bitset", HasExactSharedFragmentTypesMatch(Values, IntSharedFragmentBitSet));
+		AITEST_FALSE("Single shared value type doesn't match two-type bitset", HasExactSharedFragmentTypesMatch(Values, IntFloatSharedFragmentBitSet));
+		AITEST_FALSE("Single shared value type doesn't match empty", HasExactSharedFragmentTypesMatch(Values, EmptySharedFragmentBitSet));
 
 		const FMassConstSharedFragmentBitSet IntConstSharedFragmentBitSet = FMassConstSharedFragmentBitSet::GetTypeBitSet<FTestConstSharedFragment_Int>();
 		FMassConstSharedFragmentBitSet IntFloatConstSharedFragmentBitSet = IntConstSharedFragmentBitSet;
 		IntFloatConstSharedFragmentBitSet.Add<FTestConstSharedFragment_Float>();
 		Values.Add(FConstSharedStruct::Make<FTestConstSharedFragment_Int>(TestIntValue));
-		AITEST_TRUE("Single const shared value type matches expected bitset", Values.HasExactConstSharedFragmentTypesMatch(IntConstSharedFragmentBitSet));
-		AITEST_FALSE("Single const shared value type doesn't match two-type bitset", Values.HasExactConstSharedFragmentTypesMatch(IntFloatConstSharedFragmentBitSet));
-		AITEST_FALSE("Single const shared value type doesn't match empty", Values.HasExactConstSharedFragmentTypesMatch(EmptyConstSharedFragmentBitSet));
+		AITEST_TRUE("Single const shared value type matches expected bitset", HasExactConstSharedFragmentTypesMatch(Values, IntConstSharedFragmentBitSet));
+		AITEST_FALSE("Single const shared value type doesn't match two-type bitset", HasExactConstSharedFragmentTypesMatch(Values, IntFloatConstSharedFragmentBitSet));
+		AITEST_FALSE("Single const shared value type doesn't match empty", HasExactConstSharedFragmentTypesMatch(Values, EmptyConstSharedFragmentBitSet));
 
 		Values.Remove(IntSharedFragmentBitSet);
-		AITEST_TRUE("Emptied shared values match type with empty bitset", Values.HasExactSharedFragmentTypesMatch(EmptySharedFragmentBitSet));
+		AITEST_TRUE("Emptied shared values match type with empty bitset", HasExactSharedFragmentTypesMatch(Values, EmptySharedFragmentBitSet));
 		Values.Remove(IntConstSharedFragmentBitSet);
-		AITEST_TRUE("Emptied const shared values match type with empty const bitset", Values.HasExactConstSharedFragmentTypesMatch(EmptyConstSharedFragmentBitSet));
+		AITEST_TRUE("Emptied const shared values match type with empty const bitset", HasExactConstSharedFragmentTypesMatch(Values, EmptyConstSharedFragmentBitSet));
 
 		return true;
 	}
diff --git a/Engine/Source/Editor/MassEntityDebugger/Private/MassDebuggerModel.cpp b/Engine/Source/Editor/MassEntityDebugger/Private/MassDebuggerModel.cpp
--- a/Engine/Source/Editor/MassEntityDebugger/Private/MassDebuggerModel.cpp
+++ b/Engine/Source/Editor/MassEntityDebugger/Private/MassDebuggerModel.cpp
@@ -38,10 +38,7 @@ namespace UE::Mass::Debugger::Private
 
 		check(TotalLength > 0);
 
-		return float(BitSetDistance(A.Composition.GetFragments(), B.Composition.GetFragments())
-			+ BitSetDistance(A.Composition.GetTags(), B.Composition.GetTags())
-			+ BitSetDistance(A.Composition.GetChunkFragments(), B.Composition.GetChunkFragments())
-			+ BitSetDistance(A.Composition.GetSharedFragments(), B.Composition.GetSharedFragments())) 
+		return float(BitSetDistance(A.Composition.GetElementsBitSet(), B.Composition.GetElementsBitSet())) 
 			/ TotalLength;
 	}
 
@@ -275,9 +272,9 @@ FMassDebuggerArchetypeData::FMassDebuggerArchetypeData(const FMassArchetypeHandl
 		PrimaryDebugName = HashAsString;
 
 		// Use first fragment as name
-		if (FMassFragmentBitSet::FIndexIterator It = Composition.GetFragments().GetIndexIterator())
+		if (FMassFragmentBitSet::FIndexIterator It = Composition.DebugGetFragments().GetIndexIterator())
 		{
-			const FName FirstStructName = Composition.GetFragments().DebugGetStructTypeName(*It);
+			const FName FirstStructName = Composition.DebugGetFragments().DebugGetStructTypeName(*It);
 			TStringBuilder<256> StringBuilder;
 			StringBuilder.Append(FirstStructName.ToString());
 			StringBuilder.Append(TEXT("..."));
diff --git a/Engine/Source/Editor/MassEntityDebugger/Private/SMassArchetype.cpp b/Engine/Source/Editor/MassEntityDebugger/Private/SMassArchetype.cpp
--- a/Engine/Source/Editor/MassEntityDebugger/Private/SMassArchetype.cpp
+++ b/Engine/Source/Editor/MassEntityDebugger/Private/SMassArchetype.cpp
@@ -105,19 +105,19 @@ void SMassArchetype::Construct(const FArguments& InArgs, TSharedPtr<FMassDebugge
 	{
 		const FMassArchetypeCompositionDescriptor& ParentComposition = BaseArchetypeDebugData->Composition;
 
-		UE::Mass::Debugger::UI::AddBitSetDiff(Box, ParentComposition.GetFragments(), Composition.GetFragments(), TEXT("Fragments"), Brush, Prune, InDebuggerModel);
-		UE::Mass::Debugger::UI::AddBitSetDiff(Box, ParentComposition.GetTags(), Composition.GetTags(), TEXT("Tags"), Brush, Prune, InDebuggerModel);
-		UE::Mass::Debugger::UI::AddBitSetDiff(Box, ParentComposition.GetChunkFragments(), Composition.GetChunkFragments(), TEXT("Chunk Fragments"), Brush, Prune, InDebuggerModel);
-		UE::Mass::Debugger::UI::AddBitSetDiff(Box, ParentComposition.GetSharedFragments(), Composition.GetSharedFragments(), TEXT("Shared Fragments"), Brush, Prune, InDebuggerModel);
-		UE::Mass::Debugger::UI::AddBitSetDiff(Box, ParentComposition.GetConstSharedFragments(), Composition.GetConstSharedFragments(), TEXT("Const Shared Fragments"), Brush, Prune, InDebuggerModel);
+		UE::Mass::Debugger::UI::AddBitSetDiff(Box, ParentComposition.DebugGetFragments(), Composition.DebugGetFragments(), TEXT("Fragments"), Brush, Prune, InDebuggerModel);
+		UE::Mass::Debugger::UI::AddBitSetDiff(Box, ParentComposition.DebugGetTags(), Composition.DebugGetTags(), TEXT("Tags"), Brush, Prune, InDebuggerModel);
+		UE::Mass::Debugger::UI::AddBitSetDiff(Box, ParentComposition.DebugGetChunkFragments(), Composition.DebugGetChunkFragments(), TEXT("Chunk Fragments"), Brush, Prune, InDebuggerModel);
+		UE::Mass::Debugger::UI::AddBitSetDiff(Box, ParentComposition.DebugGetSharedFragments(), Composition.DebugGetSharedFragments(), TEXT("Shared Fragments"), Brush, Prune, InDebuggerModel);
+		UE::Mass::Debugger::UI::AddBitSetDiff(Box, ParentComposition.DebugGetConstSharedFragments(), Composition.DebugGetConstSharedFragments(), TEXT("Const Shared Fragments"), Brush, Prune, InDebuggerModel);
 	}
 	else
 	{
-		UE::Mass::Debugger::UI::AddBitSet(Box, Composition.GetFragments(), TEXT("Fragments"), Brush, InDebuggerModel);
-		UE::Mass::Debugger::UI::AddBitSet(Box, Composition.GetTags(), TEXT("Tags"), Brush, InDebuggerModel);
-		UE::Mass::Debugger::UI::AddBitSet(Box, Composition.GetChunkFragments(), TEXT("Chunk Fragments"), Brush, InDebuggerModel);
-		UE::Mass::Debugger::UI::AddBitSet(Box, Composition.GetSharedFragments(), TEXT("Shared Fragments"), Brush, InDebuggerModel);
-		UE::Mass::Debugger::UI::AddBitSet(Box, Composition.GetConstSharedFragments(), TEXT("Const Shared Fragments"), Brush, InDebuggerModel);
+		UE::Mass::Debugger::UI::AddBitSet(Box, Composition.DebugGetFragments(), TEXT("Fragments"), Brush, InDebuggerModel);
+		UE::Mass::Debugger::UI::AddBitSet(Box, Composition.DebugGetTags(), TEXT("Tags"), Brush, InDebuggerModel);
+		UE::Mass::Debugger::UI::AddBitSet(Box, Composition.DebugGetChunkFragments(), TEXT("Chunk Fragments"), Brush, InDebuggerModel);
+		UE::Mass::Debugger::UI::AddBitSet(Box, Composition.DebugGetSharedFragments(), TEXT("Shared Fragments"), Brush, InDebuggerModel);
+		UE::Mass::Debugger::UI::AddBitSet(Box, Composition.DebugGetConstSharedFragments(), TEXT("Const Shared Fragments"), Brush, InDebuggerModel);
 	}
 
 	TSharedRef<SVerticalBox> MainBox = SNew(SVerticalBox);
diff --git a/Engine/Source/Runtime/MassEntity/Internal/MassArchetypeData.h b/Engine/Source/Runtime/MassEntity/Internal/MassArchetypeData.h
--- a/Engine/Source/Runtime/MassEntity/Internal/MassArchetypeData.h
+++ b/Engine/Source/Runtime/MassEntity/Internal/MassArchetypeData.h
@@ -34,7 +34,7 @@ private:
 	FMassArchetypeSharedFragmentValues SharedFragmentValues;
 
 public:
-	explicit FMassArchetypeChunk(const SIZE_T InAllocSize, TConstArrayView<FInstancedStruct> InChunkFragmentTemplates, FMassArchetypeSharedFragmentValues InSharedFragmentValues)
+	explicit FMassArchetypeChunk(const SIZE_T InAllocSize, TConstArrayView<FInstancedStruct> InChunkFragmentTemplates, const FMassArchetypeSharedFragmentValues& InSharedFragmentValues)
 		: AllocSize(InAllocSize)
 		, ChunkFragmentData(InChunkFragmentTemplates)
 		, SharedFragmentValues(InSharedFragmentValues)
@@ -171,8 +171,8 @@ struct FMassArchetypeFragmentConfig
 struct FMassArchetypeData
 {
 private:
-	// One-stop-shop variable describing the archetype's fragment and tag composition 
-	FMassArchetypeCompositionDescriptor CompositionDescriptor;
+	// One-stop-shop variable describing composition of entities hosted by the archetype
+	FMassElementBitSet CompositionBitSet;
 
 	// Pre-created default chunk fragment templates
 	TArray<FInstancedStruct> ChunkFragmentsTemplate;
@@ -229,13 +229,9 @@ public:
 	explicit FMassArchetypeData(const FMassArchetypeCreationParams& CreationParams = FMassArchetypeCreationParams());
 
 	TConstArrayView<FMassArchetypeFragmentConfig> GetFragmentConfigs() const { return FragmentConfigs; }
-	const FMassFragmentBitSet& GetFragmentBitSet() const { return CompositionDescriptor.GetFragments(); }
-	const FMassTagBitSet& GetTagBitSet() const { return CompositionDescriptor.GetTags(); }
-	const FMassChunkFragmentBitSet& GetChunkFragmentBitSet() const { return CompositionDescriptor.GetChunkFragments(); }
-	const FMassSharedFragmentBitSet& GetSharedFragmentBitSet() const { return CompositionDescriptor.GetSharedFragments(); }
-	const FMassConstSharedFragmentBitSet& GetConstSharedFragmentBitSet() const { return CompositionDescriptor.GetConstSharedFragments(); }
+	const FMassElementBitSet& GetCompositionBitSet() const;
+	FMassArchetypeCompositionDescriptor GetCompositionDescriptor() const;
 
-	const FMassArchetypeCompositionDescriptor& GetCompositionDescriptor() const { return CompositionDescriptor; }
 	FORCEINLINE const FMassArchetypeSharedFragmentValues& GetSharedFragmentValues(int32 EntityIndex) const
 	{ 
 		const int32 AbsoluteIndex = EntityMap.FindChecked(EntityIndex);
@@ -255,7 +251,7 @@ public:
 	/** Method to iterate on all the fragment types */
 	void ForEachFragmentType(TFunction< void(const UScriptStruct* /*FragmentType*/)> Function) const;
 	bool HasFragmentType(const UScriptStruct* FragmentType) const;
-	bool HasTagType(const UScriptStruct* FragmentType) const { check(FragmentType); return CompositionDescriptor.GetTags().Contains(*FragmentType); }
+	bool HasTagType(const UScriptStruct* TagType) const;
 
 	bool IsEquivalent(const FMassArchetypeCompositionDescriptor& OtherCompositionDescriptor, const UE::Mass::FArchetypeGroups& OtherGroups) const;
 
@@ -272,8 +268,8 @@ public:
 	void RemoveEntity(FMassEntityHandle Entity);
 
 	bool HasFragmentDataForEntity(const UScriptStruct* FragmentType, int32 EntityIndex) const;
-	void* GetFragmentDataForEntityChecked(const UScriptStruct* FragmentType, int32 EntityIndex) const;
-	void* GetFragmentDataForEntity(const UScriptStruct* FragmentType, int32 EntityIndex) const;
+	void* GetFragmentDataForEntityChecked(TNotNull<const UScriptStruct*> FragmentType, int32 EntityIndex) const;
+	void* GetFragmentDataForEntity(TNotNull<const UScriptStruct*> FragmentType, int32 EntityIndex) const;
 
 	FORCEINLINE const int32* GetInternalIndexForEntity(const int32 EntityIndex) const { return EntityMap.Find(EntityIndex); }
 	FORCEINLINE int32 GetInternalIndexForEntityChecked(const int32 EntityIndex) const { return EntityMap.FindChecked(EntityIndex); }
@@ -529,6 +525,7 @@ struct FMassArchetypeHelper
 
 	MASSENTITY_API static bool DoesArchetypeMatchRequirements(const FMassArchetypeData& Archetype, const FMassFragmentRequirements& Requirements);
 	MASSENTITY_API static bool DoesArchetypeMatchRequirements(const FMassArchetypeCompositionDescriptor& ArchetypeComposition, const FMassFragmentRequirements& Requirements);
+	MASSENTITY_API static bool DoesArchetypeMatchRequirements(const FMassElementBitSet& ArchetypeCompositionBitSet, const FMassFragmentRequirements& Requirements);
 };
 
 //-----------------------------------------------------------------------------
@@ -563,3 +560,13 @@ inline uint32 FMassArchetypeData::GetEntityOrderVersion() const
 {
 	return EntityOrderVersion;
 }
+
+inline const FMassElementBitSet& FMassArchetypeData::GetCompositionBitSet() const
+{
+	return CompositionBitSet;
+}
+
+inline FMassArchetypeCompositionDescriptor FMassArchetypeData::GetCompositionDescriptor() const
+{
+	return FMassArchetypeCompositionDescriptor(CompositionBitSet);
+}
diff --git a/Engine/Source/Runtime/MassEntity/Private/MassArchetypeData.cpp b/Engine/Source/Runtime/MassEntity/Private/MassArchetypeData.cpp
--- a/Engine/Source/Runtime/MassEntity/Private/MassArchetypeData.cpp
+++ b/Engine/Source/Runtime/MassEntity/Private/MassArchetypeData.cpp
@@ -54,7 +54,14 @@ void FMassArchetypeData::ForEachFragmentType(TFunction< void(const UScriptStruct
 
 bool FMassArchetypeData::HasFragmentType(const UScriptStruct* FragmentType) const
 {
-	return (FragmentType && CompositionDescriptor.GetFragments().Contains(*FragmentType));
+	return UE::Mass::IsA<FMassFragment>(FragmentType) 
+		&& CompositionBitSet.Contains(*FragmentType);
+}
+
+bool FMassArchetypeData::HasTagType(const UScriptStruct* TagType) const
+{
+	return UE::Mass::IsA<FMassTag>(TagType)
+		&& CompositionBitSet.Contains(*TagType);
 }
 
 void FMassArchetypeData::Initialize(const FMassEntityManager& EntityManager, const FMassArchetypeCompositionDescriptor& InCompositionDescriptor, const uint32 ArchetypeDataVersion)
@@ -69,16 +76,14 @@ void FMassArchetypeData::Initialize(const FMassEntityManager& EntityManager, con
 	}
 
 	CreatedArchetypeDataVersion = ArchetypeDataVersion;
-	CompositionDescriptor.SetFragments(InCompositionDescriptor.GetFragments());
-	ConfigureFragments(EntityManager);
 
-	// Tags
-	CompositionDescriptor.SetTags(InCompositionDescriptor.GetTags());
+	CompositionBitSet = InCompositionDescriptor.GetElementsBitSet();
 
-	// Chunk fragments
-	CompositionDescriptor.SetChunkFragments(InCompositionDescriptor.GetChunkFragments());
+	ConfigureFragments(EntityManager);
+
+	// Create templates for chunk fragments, for quicker chunks creation.
 	TArray<const UScriptStruct*, TInlineAllocator<16>> ChunkFragmentList;
-	CompositionDescriptor.GetChunkFragments().ExportTypes(ChunkFragmentList);
+	CompositionBitSet.Get<FMassChunkFragmentBitSet>().ExportTypes(ChunkFragmentList);
 	ChunkFragmentList.Sort(FScriptStructSortOperator());
 	for (const UScriptStruct* ChunkFragmentType : ChunkFragmentList)
 	{
@@ -86,10 +91,6 @@ void FMassArchetypeData::Initialize(const FMassEntityManager& EntityManager, con
 		ChunkFragmentsTemplate.Emplace(ChunkFragmentType);
 	}
 
-	// Share fragments
-	CompositionDescriptor.SetSharedFragments(InCompositionDescriptor.GetSharedFragments());
-	CompositionDescriptor.SetConstSharedFragments(InCompositionDescriptor.GetConstSharedFragments());
-
 	EntityListOffsetWithinChunk = 0;
 
 #if WITH_MASSENTITY_DEBUG
@@ -97,16 +98,22 @@ void FMassArchetypeData::Initialize(const FMassEntityManager& EntityManager, con
 #endif // WITH_MASSENTITY_DEBUG
 }
 
+// @todo moving FMassArchetypeCompositionDescriptor doesn't make sense anymore, deprecated this constructor and optionally introduce a FMassElementBitSet&& one
 void FMassArchetypeData::InitializeWithSimilar(const FMassEntityManager& EntityManager, const FMassArchetypeData& BaseArchetype
 	, FMassArchetypeCompositionDescriptor&& NewComposition, const UE::Mass::FArchetypeGroups& InGroups, const uint32 ArchetypeDataVersion)
 {
 	checkf(IsInitialized() == false, TEXT("Trying to InitializeWithSimilar but this archetype has already been initialized"));
 
+	// NewComposition.GetElementsBitSet() - BaseArchetype.GetCompositionBitSet() -> should contain only tags and fragments
+	// BaseArchetype.GetCompositionBitSet() - NewComposition.GetElementsBitSet() -> should be empty
+	bool bHasTagsAndElements = FMassElementBitSet(NewComposition.GetElementsBitSet() - BaseArchetype.GetCompositionBitSet()).ContainsOnlyTagsAndFragments();
+	bool bNothingElse = FMassElementBitSet(BaseArchetype.GetCompositionBitSet() - NewComposition.GetElementsBitSet()).IsEmpty();
+	checkfSlow(bHasTagsAndElements && bNothingElse, TEXT("NewComposition is expected to differ from the BaseArchetype's composition just by tags and fragments"))
+
 	CreatedArchetypeDataVersion = ArchetypeDataVersion;
 
-	// note that we're calling this function rarely, so we can be a little bit inefficient here.
-	CompositionDescriptor = MoveTemp(NewComposition);
-	if (CompositionDescriptor.GetFragments() != BaseArchetype.GetCompositionDescriptor().GetFragments())
+	CompositionBitSet = NewComposition.GetElementsBitSet();
+	if (CompositionBitSet.Get<FMassFragmentBitSet>() != BaseArchetype.CompositionBitSet.Get<FMassFragmentBitSet>())
 	{
 		ConfigureFragments(EntityManager);
 	}
@@ -131,7 +138,7 @@ void FMassArchetypeData::InitializeWithSimilar(const FMassEntityManager& EntityM
 void FMassArchetypeData::ConfigureFragments(const FMassEntityManager& EntityManager)
 {
 	TArray<const UScriptStruct*, TInlineAllocator<16>> SortedFragmentList;
-	CompositionDescriptor.GetFragments().ExportTypes(SortedFragmentList);
+	CompositionBitSet.Get<FMassFragmentBitSet>().ExportTypes(SortedFragmentList);
 
 	SortedFragmentList.Sort(FScriptStructSortOperator());
 
@@ -196,8 +203,8 @@ void FMassArchetypeData::AddEntity(FMassEntityHandle Entity, const FMassArchetyp
 int32 FMassArchetypeData::AddEntityInternal(FMassEntityHandle Entity, const FMassArchetypeSharedFragmentValues& SharedFragmentValues)
 {
 	checkf(SharedFragmentValues.IsSorted(), TEXT("Expecting shared fragment values to be previously sorted"));
-	checkf(SharedFragmentValues.HasExactFragmentTypesMatch(CompositionDescriptor.GetSharedFragments(), CompositionDescriptor.GetConstSharedFragments())
-		, TEXT("Expecting values for every specified shared fragment in the archetype and only those"))
+	checkfSlow(SharedFragmentValues.DoesMatchComposition(CompositionBitSet)
+		, TEXT("Expecting values for every specified shared fragment in the archetype and only those"));
 
 	int32 IndexWithinChunk = 0;
 	int32 AbsoluteIndex = 0;
@@ -212,6 +219,15 @@ int32 FMassArchetypeData::AddEntityInternal(FMassEntityHandle Entity, const FMas
 	return AbsoluteIndex;
 }
 
+
+/// ------------ /// ------------ /// ------------ /// ------------ /// ------------
+namespace UE::Mass
+{
+	
+} // namespace UE::Mass
+/// ------------ /// ------------ /// ------------ /// ------------ /// ------------ 
+
+
 FMassArchetypeChunk& FMassArchetypeData::GetOrAddChunk(const FMassArchetypeSharedFragmentValues& SharedFragmentValues, int32& OutAbsoluteIndex, int32& OutIndexWithinChunk)
 {
 	OutAbsoluteIndex = 0;
@@ -368,21 +384,20 @@ void FMassArchetypeData::BatchDestroyEntityChunks(FMassArchetypeEntityCollection
 	}
 }
 
-bool FMassArchetypeData::HasFragmentDataForEntity(const UScriptStruct* FragmentType, int32 EntityIndex) const
+bool FMassArchetypeData::HasFragmentDataForEntity(const UScriptStruct* FragmentType, const int32 EntityIndex) const
 {
-	return (FragmentType && CompositionDescriptor.GetFragments().Contains(*FragmentType));
+	// @todo check if EntityIndex is disabled, or it has the FragmentType disabled.
+	return HasFragmentType(FragmentType);
 }
 
-void* FMassArchetypeData::GetFragmentDataForEntityChecked(const UScriptStruct* FragmentType, int32 EntityIndex) const
+void* FMassArchetypeData::GetFragmentDataForEntityChecked(TNotNull<const UScriptStruct*> FragmentType, int32 EntityIndex) const
 {
-	const FMassRawEntityInChunkData InternalIndex = MakeEntityHandle(EntityIndex);
-	
-	// failing the below Find means given entity's archetype is missing given FragmentType
-	const int32 FragmentIndex = FragmentIndexMap.FindChecked(FragmentType);
-	return GetFragmentData(FragmentIndex, InternalIndex);
+	void* ReturnPtr = GetFragmentDataForEntity(FragmentType, EntityIndex);
+	checkf(ReturnPtr, TEXT("Entity %d has neither a regular nor a transient fragment of type %s"), EntityIndex, *FragmentType->GetName())
+	return ReturnPtr;
 }
 
-void* FMassArchetypeData::GetFragmentDataForEntity(const UScriptStruct* FragmentType, int32 EntityIndex) const
+void* FMassArchetypeData::GetFragmentDataForEntity(TNotNull<const UScriptStruct*> FragmentType, int32 EntityIndex) const
 {
 	if (const int32* FragmentIndex = FragmentIndexMap.Find(FragmentType))
 	{
@@ -500,7 +515,7 @@ void FMassArchetypeData::ExecuteFunction(FMassExecutionContext& RunContext, cons
 	}
 
 	// @todo do we really want users to check composition of the archetype being processed at the moment?
-	RunContext.SetCurrentArchetypeCompositionDescriptor(GetCompositionDescriptor());
+	RunContext.StoreCurrentArchetypeComposition(CompositionBitSet);
 #if WITH_MASSENTITY_DEBUG
 	RunContext.DebugSetColor(DebugColor);
 #endif // WITH_MASSENTITY_DEBUG
@@ -543,7 +558,7 @@ void FMassArchetypeData::ExecuteFunction(FMassExecutionContext& RunContext, cons
 		return;
 	}
 
-	RunContext.SetCurrentArchetypeCompositionDescriptor(GetCompositionDescriptor());
+	RunContext.StoreCurrentArchetypeComposition(CompositionBitSet);
 #if WITH_MASSENTITY_DEBUG
 	RunContext.DebugSetColor(DebugColor);
 #endif // WITH_MASSENTITY_DEBUG
@@ -607,7 +622,7 @@ void FMassArchetypeData::ExecutionFunctionForChunk(FMassExecutionContext& RunCon
 		BindConstSharedFragmentRequirements(RunContext, Chunk.GetSharedFragmentValues(), RequirementMapping.ConstSharedFragments);
 		BindSharedFragmentRequirements(RunContext, Chunk.GetMutableSharedFragmentValues(), RequirementMapping.SharedFragments);
 
-		RunContext.SetCurrentArchetypeCompositionDescriptor(GetCompositionDescriptor());
+		RunContext.StoreCurrentArchetypeComposition(CompositionBitSet);
 		RunContext.SetCurrentChunkSerialModificationNumber(Chunk.GetSerialModificationNumber());
 #if WITH_MASSENTITY_DEBUG
 		RunContext.DebugSetColor(DebugColor);
@@ -988,12 +1003,8 @@ FString FMassArchetypeData::DebugGetDescription() const
 		OutDescription += GetCombinedDebugNamesAsString();
 		OutDescription += TEXT("\n");
 	}
-	OutDescription += TEXT("Chunk fragments: ");
-	CompositionDescriptor.GetChunkFragments().DebugGetStringDesc(OutDescription);
-	OutDescription += TEXT("\nTags: ");
-	CompositionDescriptor.GetTags().DebugGetStringDesc(OutDescription);
-	OutDescription += TEXT("\nFragments: ");
-	CompositionDescriptor.GetFragments().DebugGetStringDesc(OutDescription);
+	OutDescription += TEXT("\nElements: ");
+	CompositionBitSet.DebugGetStringDesc(OutDescription);
 	
 	return static_cast<FString>(OutDescription);
 #else
@@ -1026,10 +1037,9 @@ void FMassArchetypeData::DebugPrintArchetype(FOutputDevice& Ar)
 {
 	Ar.Logf(ELogVerbosity::Log, TEXT("Name: %s"), *GetCombinedDebugNamesAsString());
 
-	FStringOutputDevice TagsDescription;
-	CompositionDescriptor.GetTags().DebugGetStringDesc(TagsDescription);
-	Ar.Logf(ELogVerbosity::Log, TEXT("Tags: %s"), *TagsDescription);
-	Ar.Logf(ELogVerbosity::Log, TEXT("Fragments: %s"), *DebugGetDescription());
+	FStringOutputDevice ElementsDescription;
+	CompositionBitSet.DebugGetStringDesc(ElementsDescription);
+	Ar.Logf(ELogVerbosity::Log, TEXT("Elements: %s"), *ElementsDescription);
 	Ar.Logf(ELogVerbosity::Log, TEXT("\tChunks: %d x %llu KB = %llu KB total"), Chunks.Num(), GetChunkAllocSize() / 1024, (GetChunkAllocSize()*Chunks.Num()) / 1024);
 	
 	int ChunkWithFragmentsCount = 0;
@@ -1117,7 +1127,7 @@ void FMassArchetypeData::SetDebugColor(const FColor InDebugColor)
 	if (InDebugColor == FColor{0})
 	{
 		// pick a color based on the composition
-		const uint32 CompositionHash = CompositionDescriptor.CalculateHash();
+		const uint32 CompositionHash = GetTypeHash(CompositionBitSet);
 		const uint8* Bytes = reinterpret_cast<const uint8*>(&CompositionHash);
 		
 		const FLinearColor AdjustedColor = FLinearColor::MakeFromHSV8(
@@ -1150,8 +1160,8 @@ void FMassArchetypeData::BatchAddEntities(TConstArrayView<FMassEntityHandle> Ent
 {
 	SCOPE_CYCLE_COUNTER(STAT_Mass_ArchetypeBatchAdd);
 
-	testableCheckfReturn(SharedFragmentValues.HasExactSharedFragmentTypesMatch(GetCompositionDescriptor().GetSharedFragments()), return, TEXT("%hs parameter SharedFragmentValues doesn't match archetype's composition"), __FUNCTION__);
-	testableCheckfReturn(SharedFragmentValues.HasExactConstSharedFragmentTypesMatch(GetCompositionDescriptor().GetConstSharedFragments()), return, TEXT("%hs parameter ConstSharedFragmentValues doesn't match archetype's composition"), __FUNCTION__);
+	checkfSlow(SharedFragmentValues.DoesMatchComposition(CompositionBitSet), return
+		, TEXT("%hs parameter SharedFragmentValues doesn't match archetype's composition"), __FUNCTION__);
 
 	FMassArchetypeEntityCollection::FArchetypeEntityRange ResultSubchunk;
 	ResultSubchunk.ChunkIndex = 0;
@@ -1182,26 +1192,26 @@ void FMassArchetypeData::BatchMoveEntitiesToAnotherArchetype(const FMassArchetyp
 	check(&NewArchetype != this);
 
 	// verify the new archetype's shared fragment composition matches current archetype's composition modified as requested
-	if (SharedFragmentValuesToAdd)
+	const bool bChangeSharedFragments = SharedFragmentValuesToAdd || SharedFragmentToRemoveBitSet || ConstSharedFragmentToRemoveBitSet;
+	FMassElementBitSet NewSharedElementComposition;
+	if (bChangeSharedFragments)
 	{
-		bool bIsValidArchetype = true;
+		NewSharedElementComposition = CompositionBitSet.GetSharedFragments();
 		if (SharedFragmentToRemoveBitSet)
 		{
-			FMassSharedFragmentBitSet NewSharedFragmentsBitset = GetSharedFragmentBitSet();
-			NewSharedFragmentsBitset -= *SharedFragmentToRemoveBitSet;
-			NewSharedFragmentsBitset += SharedFragmentValuesToAdd->GetSharedFragmentBitSet();
-			bIsValidArchetype = NewArchetype.GetCompositionDescriptor().GetSharedFragments() == NewSharedFragmentsBitset;
+			NewSharedElementComposition -= *SharedFragmentToRemoveBitSet;
 		}
-
-		if (bIsValidArchetype && ConstSharedFragmentToRemoveBitSet)
+		if (ConstSharedFragmentToRemoveBitSet)
+		{
+			NewSharedElementComposition -= *ConstSharedFragmentToRemoveBitSet;
+		}
+		if (SharedFragmentValuesToAdd)
 		{
-			FMassConstSharedFragmentBitSet NewConstSharedFragmentsBitset = GetConstSharedFragmentBitSet();
-			NewConstSharedFragmentsBitset -= *ConstSharedFragmentToRemoveBitSet;
-			NewConstSharedFragmentsBitset += SharedFragmentValuesToAdd->GetConstSharedFragmentBitSet();
-			bIsValidArchetype = NewArchetype.GetCompositionDescriptor().GetConstSharedFragments() == NewConstSharedFragmentsBitset;
+			NewSharedElementComposition += SharedFragmentValuesToAdd->GetBitSet();
 		}
 
-		testableCheckfReturn(bIsValidArchetype, return, TEXT("%hs parameter SharedFragmentValues doesn't match archetype's composition"), __FUNCTION__);
+		const bool bIsValidArchetype = DoContainEquivalentSharedFragments(NewArchetype.GetCompositionBitSet(), NewSharedElementComposition);
+		testableCheckfReturn(bIsValidArchetype, return, TEXT("%hs combined new shared fragment composition doesn't match the  target archetype's composition"), __FUNCTION__);
 	}
 
 	TArray<FMassArchetypeEntityCollection::FArchetypeEntityRange> Subchunks(EntityCollection.GetRanges());
@@ -1230,7 +1240,6 @@ void FMassArchetypeData::BatchMoveEntitiesToAnotherArchetype(const FMassArchetyp
 		ResultSubChunk.ChunkIndex = 0;
 		ResultSubChunk.Length = 0;
 		int32 NumberMoved = 0;
-		const bool bChangeSharedFragments = SharedFragmentValuesToAdd || SharedFragmentToRemoveBitSet || ConstSharedFragmentToRemoveBitSet;
 
 		do
 		{
@@ -1244,20 +1253,7 @@ void FMassArchetypeData::BatchMoveEntitiesToAnotherArchetype(const FMassArchetyp
 			else
 			{
 				// create new shared values
-				FMassArchetypeSharedFragmentValues NewSharedValues = Chunk.GetSharedFragmentValues();
-				if (SharedFragmentToRemoveBitSet)
-				{
-					NewSharedValues.Remove(*SharedFragmentToRemoveBitSet);
-				}
-				if (ConstSharedFragmentToRemoveBitSet)
-				{
-					NewSharedValues.Remove(*ConstSharedFragmentToRemoveBitSet);
-				}
-				if (SharedFragmentValuesToAdd)
-				{
-					NewSharedValues.Append(*SharedFragmentValuesToAdd);
-				}
-				NewSharedValues.Sort();
+				FMassArchetypeSharedFragmentValues NewSharedValues = FMassArchetypeSharedFragmentValues::CreateCombined(Chunk.GetSharedFragmentValues(), NewSharedElementComposition, SharedFragmentValuesToAdd);
 
 				ResultSubChunk = NewArchetype.PrepareNextEntitiesSpanInternal(MakeArrayView(DyingEntityPtr + NumberMoved, RangeLength - NumberMoved)
 					, NewSharedValues, ResultSubChunk.ChunkIndex);
@@ -1305,11 +1301,12 @@ void FMassArchetypeData::BatchMoveEntitiesToAnotherArchetype(const FMassArchetyp
 	}
 }
 
-FMassArchetypeEntityCollection::FArchetypeEntityRange FMassArchetypeData::PrepareNextEntitiesSpanInternal(TConstArrayView<FMassEntityHandle> Entities, const FMassArchetypeSharedFragmentValues& SharedFragmentValues, int32 StartingChunk)
+FMassArchetypeEntityCollection::FArchetypeEntityRange FMassArchetypeData::PrepareNextEntitiesSpanInternal(TConstArrayView<FMassEntityHandle> Entities
+	, const FMassArchetypeSharedFragmentValues& SharedFragmentValues, int32 StartingChunk)
 {
 	checkf(SharedFragmentValues.IsSorted(), TEXT("Expecting shared fragment values to be previously sorted"));
-	checkf(SharedFragmentValues.HasExactFragmentTypesMatch(CompositionDescriptor.GetSharedFragments(), CompositionDescriptor.GetConstSharedFragments())
-		, TEXT("Expecting values for every specified shared fragment in the archetype and only those"))
+	checkfSlow(SharedFragmentValues.DoesMatchComposition(CompositionBitSet)
+		, TEXT("Expecting values for every specified shared fragment in the archetype and only those"));
 
 	int32 StartIndexWithinChunk = INDEX_NONE;
 	int32 AbsoluteStartIndex = 0;
@@ -1490,7 +1487,7 @@ void FMassArchetypeData::BatchSetFragmentValues(TConstArrayView<FMassArchetypeEn
 
 bool FMassArchetypeData::IsEquivalent(const FMassArchetypeCompositionDescriptor& OtherCompositionDescriptor, const UE::Mass::FArchetypeGroups& OtherGroups) const
 {
-	return CompositionDescriptor.IsEquivalent(OtherCompositionDescriptor) && Groups == OtherGroups;
+	return CompositionBitSet.IsEquivalent(OtherCompositionDescriptor.GetElementsBitSet()) && Groups == OtherGroups;
 }
 
 //-----------------------------------------------------------------------------
@@ -1498,18 +1495,23 @@ bool FMassArchetypeData::IsEquivalent(const FMassArchetypeCompositionDescriptor&
 //-----------------------------------------------------------------------------
 bool FMassArchetypeHelper::DoesArchetypeMatchRequirements(const FMassArchetypeData& Archetype, const FMassFragmentRequirements& Requirements)
 {
-	return DoesArchetypeMatchRequirements(Archetype.GetCompositionDescriptor(), Requirements);
+	return DoesArchetypeMatchRequirements(Archetype.GetCompositionBitSet(), Requirements);
 }
 bool FMassArchetypeHelper::DoesArchetypeMatchRequirements(const FMassArchetypeCompositionDescriptor& ArchetypeComposition, const FMassFragmentRequirements& Requirements)
 {
 	return Requirements.DoesArchetypeMatchRequirements(ArchetypeComposition);
 }
 
+bool FMassArchetypeHelper::DoesArchetypeMatchRequirements(const FMassElementBitSet& ArchetypeCompositionBitSet, const FMassFragmentRequirements& Requirements)
+{
+	return DoesArchetypeMatchRequirements(FMassArchetypeCompositionDescriptor(ArchetypeCompositionBitSet), Requirements);
+}
+
 #if WITH_MASSENTITY_DEBUG
 bool FMassArchetypeHelper::DoesArchetypeMatchRequirements(const FMassArchetypeData& Archetype, const FMassFragmentRequirements& Requirements
 	, const bool bBailOutOnFirstFail, FOutputDevice* OutputDevice)
 {
-	if (DoesArchetypeMatchRequirements(Archetype.GetCompositionDescriptor(), Requirements))
+	if (DoesArchetypeMatchRequirements(Archetype.GetCompositionBitSet(), Requirements))
 	{
 		// nothing to log
 		return true;
@@ -1519,216 +1521,9 @@ bool FMassArchetypeHelper::DoesArchetypeMatchRequirements(const FMassArchetypeDa
 	{
 		// do logging
 		OutputDevice->Logf(TEXT("%s")
-			, *FMassDebugger::GetArchetypeRequirementCompatibilityDescription(Requirements, Archetype.GetCompositionDescriptor()));
+			, *FMassDebugger::GetArchetypeRequirementCompatibilityDescription(Requirements, FMassArchetypeCompositionDescriptor(Archetype.GetCompositionBitSet())));
 	}
 
 	return false;
-
-	//bool bResult = true;
-	//if (Requirements.IsOptionalsOnly() == false)
-	//{
-	//	if (Archetype.GetTagBitSet().HasAll(Requirements.GetRequiredAllTags()) == false)
-	//	{
-	//		// missing some required tags, skip.
-	//		const FMassTagBitSet UnsatisfiedTags = Requirements.GetRequiredAllTags() - Archetype.GetTagBitSet();
-	//		FStringOutputDevice Description;
-	//		UnsatisfiedTags.DebugGetStringDesc(Description);
-	//		OutputDevice->Logf(TEXT("Archetype did not match due to missing tags: %s"), *Description);
-
-	//		bResult = false;
-	//		if (bBailOutOnFirstFail)
-	//		{
-	//			return false;
-	//		}
-	//	}
-
-	//	if (Archetype.GetTagBitSet().HasNone(Requirements.GetRequiredNoneTags()) == false)
-	//	{
-	//		// has some tags required to be absent
-	//		const FMassTagBitSet UnwantedTags = Requirements.GetRequiredNoneTags().GetOverlap(Archetype.GetTagBitSet());
-	//		FStringOutputDevice Description;
-	//		UnwantedTags.DebugGetStringDesc(Description);
-	//		OutputDevice->Logf(TEXT("Archetype has tags required absent: %s"), *Description);
-
-	//		bResult = false;
-	//		if (bBailOutOnFirstFail)
-	//		{
-	//			return false;
-	//		}
-	//	}
-
-	//	if (Requirements.GetRequiredAnyTags().IsEmpty() == false
-	//		&& Archetype.GetTagBitSet().HasAny(Requirements.GetRequiredAnyTags()) == false)
-	//	{
-	//		FStringOutputDevice Description;
-	//		Requirements.GetRequiredAnyTags().DebugGetStringDesc(Description);
-	//		OutputDevice->Logf(TEXT("Archetype did not match due to missing \'any\' tags: %s"), *Description);
-
-	//		bResult = false;
-	//		if (bBailOutOnFirstFail)
-	//		{
-	//			return false;
-	//		}
-	//	}
-
-	//	if (Archetype.GetFragmentBitSet().HasAll(Requirements.GetRequiredAllFragments()) == false)
-	//	{
-	//		// missing some required fragments, skip.
-	//		const FMassFragmentBitSet UnsatisfiedFragments = Requirements.GetRequiredAllFragments() - Archetype.GetFragmentBitSet();
-	//		FStringOutputDevice Description;
-	//		UnsatisfiedFragments.DebugGetStringDesc(Description);
-	//		OutputDevice->Logf(TEXT("Archetype did not match due to missing Fragments: %s"), *Description);
-
-	//		bResult = false;
-	//		if (bBailOutOnFirstFail)
-	//		{
-	//			return false;
-	//		}
-	//	}
-
-	//	if (Archetype.GetFragmentBitSet().HasNone(Requirements.GetRequiredNoneFragments()) == false)
-	//	{
-	//		// has some Fragments required to be absent
-	//		const FMassFragmentBitSet UnwantedFragments = Requirements.GetRequiredNoneFragments().GetOverlap(Archetype.GetFragmentBitSet());
-	//		FStringOutputDevice Description;
-	//		UnwantedFragments.DebugGetStringDesc(Description);
-	//		OutputDevice->Logf(TEXT("Archetype has Fragments required absent: %s"), *Description);
-
-	//		bResult = false;
-	//		if (bBailOutOnFirstFail)
-	//		{
-	//			return false;
-	//		}
-	//	}
-
-	//	if (Requirements.GetRequiredAnyFragments().IsEmpty() == false
-	//		&& Archetype.GetFragmentBitSet().HasAny(Requirements.GetRequiredAnyFragments()) == false)
-	//	{
-	//		FStringOutputDevice Description;
-	//		Requirements.GetRequiredAnyFragments().DebugGetStringDesc(Description);
-	//		OutputDevice->Logf(TEXT("Archetype did not match due to missing \'any\' fragments: %s"), *Description);
-
-	//		bResult = false;
-	//		if (bBailOutOnFirstFail)
-	//		{
-	//			return false;
-	//		}
-	//	}
-
-	//	if (Archetype.GetChunkFragmentBitSet().HasAll(Requirements.GetRequiredAllChunkFragments()) == false)
-	//	{
-	//		// missing some required fragments, skip.
-	//		const FMassChunkFragmentBitSet UnsatisfiedFragments = Requirements.GetRequiredAllChunkFragments() - Archetype.GetChunkFragmentBitSet();
-	//		FStringOutputDevice Description;
-	//		UnsatisfiedFragments.DebugGetStringDesc(Description);
-	//		OutputDevice->Logf(TEXT("Archetype did not match due to missing Chunk Fragments: %s"), *Description);
-
-	//		bResult = false;
-	//		if (bBailOutOnFirstFail)
-	//		{
-	//			return false;
-	//		}
-	//	}
-
-	//	if (Archetype.GetChunkFragmentBitSet().HasNone(Requirements.GetRequiredNoneChunkFragments()) == false)
-	//	{
-	//		// has some Fragments required to be absent
-	//		const FMassChunkFragmentBitSet UnwantedFragments = Requirements.GetRequiredNoneChunkFragments().GetOverlap(Archetype.GetChunkFragmentBitSet());
-	//		FStringOutputDevice Description;
-	//		UnwantedFragments.DebugGetStringDesc(Description);
-	//		OutputDevice->Logf(TEXT("Archetype has Chunk Fragments required absent: %s"), *Description);
-
-	//		bResult = false;
-	//		if (bBailOutOnFirstFail)
-	//		{
-	//			return false;
-	//		}
-	//	}
-
-	//	if (Archetype.GetSharedFragmentBitSet().HasAll(Requirements.GetRequiredAllSharedFragments()) == false)
-	//	{
-	//		// missing some required fragments, skip.
-	//		const FMassSharedFragmentBitSet UnsatisfiedFragments = Requirements.GetRequiredAllSharedFragments() - Archetype.GetSharedFragmentBitSet();
-	//		FStringOutputDevice Description;
-	//		UnsatisfiedFragments.DebugGetStringDesc(Description);
-	//		OutputDevice->Logf(TEXT("Archetype did not match due to missing Shared Fragments: %s"), *Description);
-
-	//		bResult = false;
-	//		if (bBailOutOnFirstFail)
-	//		{
-	//			return false;
-	//		}
-	//	}
-
-	//	if (Archetype.GetSharedFragmentBitSet().HasNone(Requirements.GetRequiredNoneSharedFragments()) == false)
-	//	{
-	//		// has some Fragments required to be absent
-	//		const FMassSharedFragmentBitSet UnwantedFragments = Requirements.GetRequiredNoneSharedFragments().GetOverlap(Archetype.GetSharedFragmentBitSet());
-	//		FStringOutputDevice Description;
-	//		UnwantedFragments.DebugGetStringDesc(Description);
-	//		OutputDevice->Logf(TEXT("Archetype has Shared Fragments required absent: %s"), *Description);
-
-	//		bResult = false;
-	//		if (bBailOutOnFirstFail)
-	//		{
-	//			return false;
-	//		}
-	//	}
-
-	//	if (Archetype.GetConstSharedFragmentBitSet().HasAll(Requirements.GetRequiredAllConstSharedFragments()) == false)
-	//	{
-	//		// missing some required fragments, skip.
-	//		const FMassConstSharedFragmentBitSet UnsatisfiedFragments = Requirements.GetRequiredAllConstSharedFragments() - Archetype.GetConstSharedFragmentBitSet();
-	//		FStringOutputDevice Description;
-	//		UnsatisfiedFragments.DebugGetStringDesc(Description);
-	//		OutputDevice->Logf(TEXT("Archetype did not match due to missing Const Shared Fragments: %s"), *Description);
-
-	//		bResult = false;
-	//		if (bBailOutOnFirstFail)
-	//		{
-	//			return false;
-	//		}
-	//	}
-
-	//	if (Archetype.GetConstSharedFragmentBitSet().HasNone(Requirements.GetRequiredNoneConstSharedFragments()) == false)
-	//	{
-	//		// has some Fragments required to be absent
-	//		const FMassConstSharedFragmentBitSet UnwantedFragments = Requirements.GetRequiredNoneConstSharedFragments().GetOverlap(Archetype.GetConstSharedFragmentBitSet());
-	//		FStringOutputDevice Description;
-	//		UnwantedFragments.DebugGetStringDesc(Description);
-	//		OutputDevice->Logf(TEXT("Archetype has Const Shared Fragments required absent: %s"), *Description);
-
-	//		// could skip the test as it's the final check we're performing, but leaving it in in case more cases are added
-	//		// later without checking the existing code
-	//		bResult = false;
-	//		if (bBailOutOnFirstFail)
-	//		{
-	//			return false;
-	//		}
-	//	}
-	//}
-	//else
-	//{
-	//	// test if contains any of the optional elements
-	//	if (Archetype.GetFragmentBitSet().HasNone(Requirements.GetRequiredOptionalFragments())
-	//		&& Archetype.GetTagBitSet().HasNone(Requirements.GetRequiredOptionalTags())
-	//		&& Archetype.GetChunkFragmentBitSet().HasNone(Requirements.GetRequiredOptionalChunkFragments())
-	//		&& Archetype.GetSharedFragmentBitSet().HasNone(Requirements.GetRequiredOptionalSharedFragments())
-	//		&& Archetype.GetConstSharedFragmentBitSet().HasNone(Requirements.GetRequiredOptionalConstSharedFragments()))
-	//	{
-	//		FStringOutputDevice Description;
-	//		Requirements.GetRequiredOptionalFragments().DebugGetStringDesc(Description);
-	//		Requirements.GetRequiredOptionalTags().DebugGetStringDesc(Description);
-	//		Requirements.GetRequiredOptionalChunkFragments().DebugGetStringDesc(Description);
-	//		Requirements.GetRequiredOptionalSharedFragments().DebugGetStringDesc(Description);
-	//		Requirements.GetRequiredOptionalConstSharedFragments().DebugGetStringDesc(Description);
-	//		
-	//		OutputDevice->Logf(TEXT("Archetype has none of the optional elements: %s"), *Description);
-
-	//		bResult = false;
-	//	}
-	//}
-
-	//return bResult;
 }
-#endif // WITH_MASSENTITY_DEBUG
+#endif // WITH_MASSENTITY_DEBUG
\ No newline at end of file
diff --git a/Engine/Source/Runtime/MassEntity/Private/MassDebugger.cpp b/Engine/Source/Runtime/MassEntity/Private/MassDebugger.cpp
--- a/Engine/Source/Runtime/MassEntity/Private/MassDebugger.cpp
+++ b/Engine/Source/Runtime/MassEntity/Private/MassDebugger.cpp
@@ -583,10 +583,10 @@ TArray<FMassArchetypeHandle> FMassDebugger::GetAllArchetypes(const FMassEntityMa
 	return Archetypes;
 }
 
-const FMassArchetypeCompositionDescriptor& FMassDebugger::GetArchetypeComposition(const FMassArchetypeHandle& ArchetypeHandle)
+FMassArchetypeCompositionDescriptor FMassDebugger::GetArchetypeComposition(const FMassArchetypeHandle& ArchetypeHandle)
 {
 	const FMassArchetypeData& ArchetypeData = FMassArchetypeHelper::ArchetypeDataFromHandleChecked(ArchetypeHandle);
-	return ArchetypeData.CompositionDescriptor;
+	return FMassArchetypeCompositionDescriptor(ArchetypeData.GetCompositionBitSet());
 }
 
 uint64 FMassDebugger::GetArchetypeTraceID(const FMassArchetypeData& ArchetypeData)
@@ -693,105 +693,111 @@ FString FMassDebugger::GetArchetypeRequirementCompatibilityDescription(const FMa
 	}
 
 	const FMassArchetypeData& ArchetypeData = FMassArchetypeHelper::ArchetypeDataFromHandleChecked(ArchetypeHandle);
-	return FMassDebugger::GetArchetypeRequirementCompatibilityDescription(Requirements, ArchetypeData.GetCompositionDescriptor());
+	return FMassDebugger::GetArchetypeRequirementCompatibilityDescription(Requirements, ArchetypeData.GetCompositionBitSet());
 }
 	
-FString FMassDebugger::GetArchetypeRequirementCompatibilityDescription(const FMassFragmentRequirements& Requirements, const FMassArchetypeCompositionDescriptor& ArchetypeComposition)
+FString FMassDebugger::GetArchetypeRequirementCompatibilityDescription(const FMassFragmentRequirements& Requirements, const FMassElementBitSet& ArchetypeCompositionBitSet)
 {
 	FStringOutputDevice OutDescription;
 
+	FMassFragmentBitSet FragmentBitSet = ArchetypeCompositionBitSet.Get<FMassFragmentBitSet>();
+	FMassTagBitSet TagBitSet = ArchetypeCompositionBitSet.Get<FMassTagBitSet>();
+	FMassChunkFragmentBitSet ChunkFragmentBitSet = ArchetypeCompositionBitSet.Get<FMassChunkFragmentBitSet>();
+	FMassSharedFragmentBitSet SharedFragmentBitSet = ArchetypeCompositionBitSet.Get<FMassSharedFragmentBitSet>();
+	FMassConstSharedFragmentBitSet ConstSharedFragmentBitSet = ArchetypeCompositionBitSet.Get<FMassConstSharedFragmentBitSet>();
+
 	if (Requirements.HasNegativeRequirements())
 	{
-		if (ArchetypeComposition.GetFragments().HasNone(Requirements.RequiredNoneFragments) == false)
+		if (FragmentBitSet.HasNone(Requirements.RequiredNoneFragments) == false)
 		{
 			// has some of the fragments required absent
 			OutDescription += TEXT("\nHas fragments required absent: ");
-			(Requirements.RequiredNoneFragments & ArchetypeComposition.GetFragments()).DebugGetStringDesc(OutDescription);
+			(Requirements.RequiredNoneFragments & FragmentBitSet).DebugGetStringDesc(OutDescription);
 		}
 
-		if (ArchetypeComposition.GetTags().HasNone(Requirements.RequiredNoneTags) == false)
+		if (TagBitSet.HasNone(Requirements.RequiredNoneTags) == false)
 		{
 			// has some of the tags required absent
 			OutDescription += TEXT("\nHas tags required absent: ");
-			(Requirements.RequiredNoneTags & ArchetypeComposition.GetTags()).DebugGetStringDesc(OutDescription);
+			(Requirements.RequiredNoneTags & TagBitSet).DebugGetStringDesc(OutDescription);
 		}
 
-		if (ArchetypeComposition.GetChunkFragments().HasNone(Requirements.RequiredNoneChunkFragments) == false)
+		if (ChunkFragmentBitSet.HasNone(Requirements.RequiredNoneChunkFragments) == false)
 		{
 			// has some of the chunk fragments required absent
 			OutDescription += TEXT("\nHas chunk fragments required absent: ");
-			(Requirements.RequiredNoneChunkFragments & ArchetypeComposition.GetChunkFragments()).DebugGetStringDesc(OutDescription);
+			(Requirements.RequiredNoneChunkFragments & ChunkFragmentBitSet).DebugGetStringDesc(OutDescription);
 		}
 
-		if (ArchetypeComposition.GetSharedFragments().HasNone(Requirements.RequiredNoneSharedFragments) == false)
+		if (SharedFragmentBitSet.HasNone(Requirements.RequiredNoneSharedFragments) == false)
 		{
 			// has some of the chunk fragments required absent
 			OutDescription += TEXT("\nHas shared fragments required absent: ");
-			(Requirements.RequiredNoneSharedFragments & ArchetypeComposition.GetSharedFragments()).DebugGetStringDesc(OutDescription);
+			(Requirements.RequiredNoneSharedFragments & SharedFragmentBitSet).DebugGetStringDesc(OutDescription);
 		}
 
-		if (ArchetypeComposition.GetConstSharedFragments().HasNone(Requirements.RequiredNoneConstSharedFragments) == false)
+		if (ConstSharedFragmentBitSet.HasNone(Requirements.RequiredNoneConstSharedFragments) == false)
 		{
 			// has some of the chunk fragments required absent
-			OutDescription += TEXT("\nHas shared fragments required absent: ");
-			(Requirements.RequiredNoneConstSharedFragments & ArchetypeComposition.GetConstSharedFragments()).DebugGetStringDesc(OutDescription);
+			OutDescription += TEXT("\nHas const shared fragments required absent: ");
+			(Requirements.RequiredNoneConstSharedFragments & ConstSharedFragmentBitSet).DebugGetStringDesc(OutDescription);
 		}
 	}
 
 	// if we have regular (i.e. non-optional) positive requirements then these are the determining factor, we don't check optionals
 	if (Requirements.HasPositiveRequirements())
 	{
-		if (ArchetypeComposition.GetFragments().HasAll(Requirements.RequiredAllFragments) == false)
+		if (FragmentBitSet.HasAll(Requirements.RequiredAllFragments) == false)
 		{
 			// missing one of the strictly required fragments
 			OutDescription += TEXT("\nMissing required fragments: ");
-			(Requirements.RequiredAllFragments - ArchetypeComposition.GetFragments()).DebugGetStringDesc(OutDescription);
+			(Requirements.RequiredAllFragments - FragmentBitSet).DebugGetStringDesc(OutDescription);
 		}
 
-		if (Requirements.RequiredAnyFragments.IsEmpty() == false && ArchetypeComposition.GetFragments().HasAny(Requirements.RequiredAnyFragments) == false)
+		if (Requirements.RequiredAnyFragments.IsEmpty() == false && FragmentBitSet.HasAny(Requirements.RequiredAnyFragments) == false)
 		{
 			// missing all of the "any" fragments
 			OutDescription += TEXT("\nMissing all \'any\' fragments: ");
 			Requirements.RequiredAnyFragments.DebugGetStringDesc(OutDescription);
 		}
 
-		if (ArchetypeComposition.GetTags().HasAll(Requirements.RequiredAllTags) == false)
+		if (TagBitSet.HasAll(Requirements.RequiredAllTags) == false)
 		{
 			// missing one of the strictly required tags
 			OutDescription += TEXT("\nMissing required tags: ");
-			(Requirements.RequiredAllTags - ArchetypeComposition.GetTags()).DebugGetStringDesc(OutDescription);
+			(Requirements.RequiredAllTags - TagBitSet).DebugGetStringDesc(OutDescription);
 		}
 
-		if (Requirements.RequiredAnyTags.IsEmpty() == false && ArchetypeComposition.GetTags().HasAny(Requirements.RequiredAnyTags) == false)
+		if (Requirements.RequiredAnyTags.IsEmpty() == false && TagBitSet.HasAny(Requirements.RequiredAnyTags) == false)
 		{
 			// missing all of the "any" tags
 			OutDescription += TEXT("\nMissing all \'any\' tags: ");
 			Requirements.RequiredAnyTags.DebugGetStringDesc(OutDescription);
 		}
 
-		if (ArchetypeComposition.GetChunkFragments().HasAll(Requirements.RequiredAllChunkFragments) == false)
+		if (ChunkFragmentBitSet.HasAll(Requirements.RequiredAllChunkFragments) == false)
 		{
 			// missing one of the strictly required chunk fragments
 			OutDescription += TEXT("\nMissing required chunk fragments: ");
-			(Requirements.RequiredAllChunkFragments - ArchetypeComposition.GetChunkFragments()).DebugGetStringDesc(OutDescription);
+			(Requirements.RequiredAllChunkFragments - ChunkFragmentBitSet).DebugGetStringDesc(OutDescription);
 		}
 
-		if (ArchetypeComposition.GetSharedFragments().HasAll(Requirements.RequiredAllSharedFragments) == false)
+		if (SharedFragmentBitSet.HasAll(Requirements.RequiredAllSharedFragments) == false)
 		{
 			// missing one of the strictly required Shared fragments
 			OutDescription += TEXT("\nMissing required Shared fragments: ");
-			(Requirements.RequiredAllSharedFragments - ArchetypeComposition.GetSharedFragments()).DebugGetStringDesc(OutDescription);
+			(Requirements.RequiredAllSharedFragments - SharedFragmentBitSet).DebugGetStringDesc(OutDescription);
 		}
 
-		if (ArchetypeComposition.GetConstSharedFragments().HasAll(Requirements.RequiredAllConstSharedFragments) == false)
+		if (ConstSharedFragmentBitSet.HasAll(Requirements.RequiredAllConstSharedFragments) == false)
 		{
 			// missing one of the strictly required Shared fragments
-			OutDescription += TEXT("\nMissing required Shared fragments: ");
-			(Requirements.RequiredAllConstSharedFragments - ArchetypeComposition.GetConstSharedFragments()).DebugGetStringDesc(OutDescription);
+			OutDescription += TEXT("\nMissing required Const Shared fragments: ");
+			(Requirements.RequiredAllConstSharedFragments - ConstSharedFragmentBitSet).DebugGetStringDesc(OutDescription);
 		}
 	}
 	// else we check if there are any optionals and if so test them
-	else if (Requirements.HasOptionalRequirements() && (Requirements.DoesMatchAnyOptionals(ArchetypeComposition) == false))
+	else if (Requirements.HasOptionalRequirements() && (Requirements.DoesMatchAnyOptionals(ArchetypeCompositionBitSet) == false))
 	{
 		// we report that none of the optionals has been met
 		OutDescription += TEXT("\nNone of the optionals were safisfied while not having other positive hard requirements: ");
@@ -1354,41 +1360,41 @@ const UScriptStruct* FMassDebugger::GetFragmentTypeFromName(FName FragmentName)
 		{
 			const FMassArchetypeCompositionDescriptor& Composition = GetArchetypeComposition(ArchetypeHandle);
 
-			FMassFragmentBitSet::FIndexIterator It = Composition.GetFragments().GetIndexIterator();
+			FMassFragmentBitSet::FIndexIterator It = Composition.DebugGetFragments().GetIndexIterator();
 			while (It)
 			{
-				FName StructName = Composition.GetFragments().DebugGetStructTypeName(*It);
-				const UScriptStruct* StructType = Composition.GetFragments().GetTypeAtIndex(*It);
+				FName StructName = Composition.DebugGetFragments().DebugGetStructTypeName(*It);
+				const UScriptStruct* StructType = Composition.DebugGetFragments().GetTypeAtIndex(*It);
 				
 				FragmentsByName.Add(StructName, StructType);
 				++It;
 			}
 
-			FMassChunkFragmentBitSet::FIndexIterator ChunkIt = Composition.GetChunkFragments().GetIndexIterator();
+			FMassChunkFragmentBitSet::FIndexIterator ChunkIt = Composition.DebugGetChunkFragments().GetIndexIterator();
 			while (ChunkIt)
 			{
-				FName StructName = Composition.GetChunkFragments().DebugGetStructTypeName(*ChunkIt);
-				const UScriptStruct* StructType = Composition.GetChunkFragments().GetTypeAtIndex(*ChunkIt);
+				FName StructName = Composition.DebugGetChunkFragments().DebugGetStructTypeName(*ChunkIt);
+				const UScriptStruct* StructType = Composition.DebugGetChunkFragments().GetTypeAtIndex(*ChunkIt);
 
 				FragmentsByName.Add(StructName, StructType);
 				++ChunkIt;
 			}
 
-			FMassSharedFragmentBitSet::FIndexIterator SharedFragIt = Composition.GetSharedFragments().GetIndexIterator();
+			FMassSharedFragmentBitSet::FIndexIterator SharedFragIt = Composition.DebugGetSharedFragments().GetIndexIterator();
 			while (SharedFragIt)
 			{
-				FName StructName = Composition.GetSharedFragments().DebugGetStructTypeName(*SharedFragIt);
-				const UScriptStruct* StructType = Composition.GetSharedFragments().GetTypeAtIndex(*SharedFragIt);
+				FName StructName = Composition.DebugGetSharedFragments().DebugGetStructTypeName(*SharedFragIt);
+				const UScriptStruct* StructType = Composition.DebugGetSharedFragments().GetTypeAtIndex(*SharedFragIt);
 
 				FragmentsByName.Add(StructName, StructType);
 				++SharedFragIt;
 			}
 
-			FMassConstSharedFragmentBitSet::FIndexIterator ConstSharedFragIt = Composition.GetConstSharedFragments().GetIndexIterator();
+			FMassConstSharedFragmentBitSet::FIndexIterator ConstSharedFragIt = Composition.DebugGetConstSharedFragments().GetIndexIterator();
 			while (ConstSharedFragIt)
 			{
-				FName StructName = Composition.GetConstSharedFragments().DebugGetStructTypeName(*ConstSharedFragIt);
-				const UScriptStruct* StructType = Composition.GetConstSharedFragments().GetTypeAtIndex(*ConstSharedFragIt);
+				FName StructName = Composition.DebugGetConstSharedFragments().DebugGetStructTypeName(*ConstSharedFragIt);
+				const UScriptStruct* StructType = Composition.DebugGetConstSharedFragments().GetTypeAtIndex(*ConstSharedFragIt);
 
 				FragmentsByName.Add(StructName, StructType);
 				++ConstSharedFragIt;
diff --git a/Engine/Source/Runtime/MassEntity/Private/MassElement.cpp b/Engine/Source/Runtime/MassEntity/Private/MassElement.cpp
new file mode 100644
--- /dev/null
+++ b/Engine/Source/Runtime/MassEntity/Private/MassElement.cpp
@@ -0,0 +1,16 @@
+// Copyright Epic Games, Inc. All Rights Reserved.
+
+#include "MassElement.h"
+
+DEFINE_TYPEBITSET(FMassExternalSubsystemBitSet);
+
+namespace UE::Mass
+{
+	DEFINE_TYPEBITSET(FElementBitSetBase);
+
+	TStaticArray<FElementBitSetBase, static_cast<uint8>(EElementType::MAX)> FElementBitSet::ElementTypeBitArrays;
+	FElementBitSet FElementBitSet::AllSharedFragmentsBitSet;
+	FElementBitSet FElementBitSet::AllFragmentsAndTagsBitSet;
+	bool FElementBitSet::bTypeBitSetsDirty = false;
+	FDelegateHandle FElementBitSet::OnTypeRegisteredDelegateHandle;
+}
diff --git a/Engine/Source/Runtime/MassEntity/Private/MassEntityBuilder.cpp b/Engine/Source/Runtime/MassEntity/Private/MassEntityBuilder.cpp
--- a/Engine/Source/Runtime/MassEntity/Private/MassEntityBuilder.cpp
+++ b/Engine/Source/Runtime/MassEntity/Private/MassEntityBuilder.cpp
@@ -33,7 +33,7 @@ namespace UE::Mass
 			template<typename T>
 			static void AppendFromEntity(FEntityBuilder& Builder, const FMassEntityHandle SourceEntityHandle, const FMassArchetypeCompositionDescriptor& ArchetypeComposition)
 			{
-				const auto& SourceContainer = ArchetypeComposition.GetContainer<T>();
+				auto SourceContainer = ArchetypeComposition.GetContainer<T>();
 				TArray<FInstancedStruct>& ElementInstanceContainer = Builder.GetInstancedStructContainerInternal<T>();
 
 				// remove all the existing entries that match SourceContainer, and then just copy.
@@ -57,7 +57,7 @@ namespace UE::Mass
 			template<typename T>
 			static void CopyFromEntity(FEntityBuilder& Builder, const FMassEntityHandle SourceEntityHandle, const FMassArchetypeCompositionDescriptor& ArchetypeComposition)
 			{
-				const auto& SourceContainer = ArchetypeComposition.GetContainer<T>();
+				auto SourceContainer = ArchetypeComposition.GetContainer<T>();
 				TArray<FInstancedStruct>& ElementInstanceContainer = Builder.GetInstancedStructContainerInternal<T>();
 
 				ElementInstanceContainer.Reserve(ElementInstanceContainer.Num() + SourceContainer.CountStoredTypes());
@@ -106,9 +106,9 @@ namespace UE::Mass
 			, TConstArrayView<FInstancedStruct> InitialFragmentValues, TConstArrayView<FConstSharedStruct> ConstSharedFragments, TConstArrayView<FSharedStruct> SharedFragments)
 		{
 			const UObject* LogOwner = InEntityManager->GetOwner();
-			bool bIssuesFound = CheckStructContainer<FMassFragment>(InitialFragmentValues, Composition.GetFragments(), LogOwner);
-			bIssuesFound |= CheckStructContainer<FMassConstSharedFragment>(ConstSharedFragments, Composition.GetConstSharedFragments(), LogOwner);
-			bIssuesFound |= CheckStructContainer<FMassSharedFragment>(SharedFragments, Composition.GetSharedFragments(), LogOwner);
+			bool bIssuesFound = CheckStructContainer<FMassFragment>(InitialFragmentValues, Composition.GetElementsBitSet(), LogOwner);
+			bIssuesFound |= CheckStructContainer<FMassConstSharedFragment>(ConstSharedFragments, Composition.GetElementsBitSet(), LogOwner);
+			bIssuesFound |= CheckStructContainer<FMassSharedFragment>(SharedFragments, Composition.GetElementsBitSet(), LogOwner);
 			
 			return !bIssuesFound;
 		}
@@ -410,6 +410,7 @@ bool FEntityBuilder::AppendDataFromEntity(const FMassEntityHandle SourceEntityHa
 
 bool FEntityBuilder::CopyDataFromEntity(const FMassEntityHandle SourceEntityHandle)
 {
+	// @todo should we copy transient data too?
 	if (!testableEnsureMsgf(EntityManager->IsEntityActive(SourceEntityHandle)
 		, TEXT("%hs expecting a valid, built entity as input"), __FUNCTION__))
 	{
@@ -503,7 +504,7 @@ bool FEntityBuilder::HandleTypeInstance(const UScriptStruct* Type, TArray<FInsta
 {
 	if (UE::Mass::IsA<T>(Type))
 	{
-		auto& BitSet = Composition.GetContainer<T>();
+		auto& BitSet = Composition.GetElementsBitSet();
 		const int32 TypeIndex = BitSet.GetTypeIndex(*Type);
 		OutTargetContainer = &GetInstancedStructContainerInternal<T>();
 		if (BitSet.IsBitSet(TypeIndex) == false)
@@ -569,30 +570,12 @@ FEntityBuilder& FEntityBuilder::Add(FInstancedStruct&& ElementInstance)
 
 FEntityBuilder& FEntityBuilder::Add(TNotNull<const UScriptStruct*> ElementType)
 {
-	if (UE::Mass::IsA<FMassTag>(ElementType))
+	if (UE::Mass::IsA<FMassElement>(ElementType))
 	{
-		Composition.GetTags().Add(*ElementType);
-	}
-	else if (UE::Mass::IsA<FMassFragment>(ElementType))
-	{
-		Composition.GetFragments().Add(*ElementType);
-	}
-	else if (UE::Mass::IsA<FMassSharedFragment>(ElementType))
-	{
-		Composition.GetSharedFragments().Add(*ElementType);
-	}
-	else if (UE::Mass::IsA<FMassConstSharedFragment>(ElementType))
-	{
-		Composition.GetConstSharedFragments().Add(*ElementType);
-	}
-	else if (ensureMsgf(UE::Mass::IsA<FMassChunkFragment>(ElementType)
-		, TEXT("Unhandled element type %s"), *ElementType->GetName()))
-	{
-		Composition.GetChunkFragments().Add(*ElementType);
+		Composition.GetElementsBitSet().Add(*ElementType);
+		State = EState::ReadyToCommit;
+		CachedArchetypeHandle = FMassArchetypeHandle();
 	}
-
-	State = EState::ReadyToCommit;
-	CachedArchetypeHandle = FMassArchetypeHandle();
 	
 	return *this;
 }
diff --git a/Engine/Source/Runtime/MassEntity/Private/MassEntityManager.cpp b/Engine/Source/Runtime/MassEntity/Private/MassEntityManager.cpp
--- a/Engine/Source/Runtime/MassEntity/Private/MassEntityManager.cpp
+++ b/Engine/Source/Runtime/MassEntity/Private/MassEntityManager.cpp
@@ -332,7 +332,7 @@ FMassArchetypeHandle FMassEntityManager::CreateArchetype(FMassArchetypeHandle So
 	TConstArrayView<const UScriptStruct*> FragmentsAndTagsList, const FMassArchetypeCreationParams& CreationParams)
 {
 	const FMassArchetypeData& ArchetypeData = FMassArchetypeHelper::ArchetypeDataFromHandleChecked(SourceArchetype);
-	FMassArchetypeCompositionDescriptor Composition = ArchetypeData.GetCompositionDescriptor();
+	FMassArchetypeCompositionDescriptor Composition = FMassArchetypeCompositionDescriptor(ArchetypeData.GetCompositionBitSet());
 	InternalAppendFragmentsAndTagsToArchetypeCompositionDescriptor(Composition, FragmentsAndTagsList);
 	return CreateArchetype(Composition, CreationParams);
 }
@@ -347,11 +347,7 @@ FMassArchetypeHandle FMassEntityManager::CreateArchetype(const TSharedPtr<FMassA
 	check(SourceArchetype.IsValid());
 	checkf(AddedFragments.IsEmpty() == false, TEXT("%hs Adding an empty fragment list to an archetype is not supported."), __FUNCTION__);
 
-	const FMassArchetypeCompositionDescriptor Composition(AddedFragments + SourceArchetype->GetFragmentBitSet()
-		, SourceArchetype->GetTagBitSet()
-		, SourceArchetype->GetChunkFragmentBitSet()
-		, SourceArchetype->GetSharedFragmentBitSet()
-		, SourceArchetype->GetConstSharedFragmentBitSet());
+	FMassArchetypeCompositionDescriptor Composition(SourceArchetype->GetCompositionBitSet() + AddedFragments);
 	return CreateArchetype(Composition, CreationParams);
 }
 
@@ -359,14 +355,21 @@ FMassArchetypeHandle FMassEntityManager::GetOrCreateSuitableArchetype(const FMas
 	, const FMassSharedFragmentBitSet& SharedFragmentBitSet
 	, const FMassConstSharedFragmentBitSet& ConstSharedFragmentBitSet
 	, const FMassArchetypeCreationParams& CreationParams)
+{
+	return GetOrCreateSuitableArchetype(ArchetypeHandle, FMassElementBitSet(SharedFragmentBitSet) + ConstSharedFragmentBitSet, CreationParams);
+}
+
+FMassArchetypeHandle FMassEntityManager::GetOrCreateSuitableArchetype(const FMassArchetypeHandle& ArchetypeHandle
+	, const FMassElementBitSet& SharedElementsBitSet, const FMassArchetypeCreationParams& CreationParams)
 {
 	const FMassArchetypeData& ArchetypeData = FMassArchetypeHelper::ArchetypeDataFromHandleChecked(ArchetypeHandle);
-	if (SharedFragmentBitSet != ArchetypeData.GetSharedFragmentBitSet()
-		|| ConstSharedFragmentBitSet != ArchetypeData.GetConstSharedFragmentBitSet())
+	if (UE::Mass::DoContainEquivalentSharedFragments(ArchetypeData.GetCompositionBitSet(), SharedElementsBitSet) == false)
 	{
-		FMassArchetypeCompositionDescriptor NewDescriptor = ArchetypeData.GetCompositionDescriptor();
-		NewDescriptor.SetSharedFragments(SharedFragmentBitSet);
-		NewDescriptor.SetConstSharedFragments(ConstSharedFragmentBitSet);
+		checkf(UE::Mass::DoesContainOnlySharedFragments(SharedElementsBitSet)
+			, TEXT("SharedElementsBitSet is expected to contain only shared fragments (const and non-const)"));
+
+		FMassArchetypeCompositionDescriptor NewDescriptor(ArchetypeData.GetCompositionBitSet());
+		NewDescriptor.SetAllSharedElements(SharedElementsBitSet);
 		return CreateArchetype(NewDescriptor, CreationParams);
 	}
 	return ArchetypeHandle;
@@ -426,32 +429,19 @@ FMassArchetypeHandle FMassEntityManager::CreateArchetype(const FMassArchetypeCom
 	return FMassArchetypeHelper::ArchetypeHandleFromData(ArchetypeDataPtr);
 }
 
-FMassArchetypeHandle FMassEntityManager::InternalCreateSimilarArchetype(const TSharedPtr<FMassArchetypeData>& SourceArchetype, const FMassTagBitSet& OverrideTags)
-{
-	checkSlow(SourceArchetype.IsValid());
-	const FMassArchetypeData& SourceArchetypeRef = *SourceArchetype.Get();
-	FMassArchetypeCompositionDescriptor NewComposition(SourceArchetypeRef.GetCompositionDescriptor());
-	NewComposition.SetTags(OverrideTags);
-
-	return InternalCreateSimilarArchetype(SourceArchetypeRef, MoveTemp(NewComposition), SourceArchetypeRef.GetGroups());
-}
-
-FMassArchetypeHandle FMassEntityManager::InternalCreateSimilarArchetype(const TSharedPtr<FMassArchetypeData>& SourceArchetype, const FMassFragmentBitSet& OverrideFragments)
+FMassArchetypeHandle FMassEntityManager::InternalCreateSimilarArchetype(const TSharedPtr<FMassArchetypeData>& SourceArchetype, FMassElementBitSet&& NewCompositionBitSet)
 {
 	checkSlow(SourceArchetype.IsValid());
 	const FMassArchetypeData& SourceArchetypeRef = *SourceArchetype.Get();
-	FMassArchetypeCompositionDescriptor NewComposition(SourceArchetypeRef.GetCompositionDescriptor());
-	NewComposition.SetFragments(OverrideFragments);
 
-	return InternalCreateSimilarArchetype(SourceArchetypeRef, MoveTemp(NewComposition), SourceArchetypeRef.GetGroups());
+	return InternalCreateSimilarArchetype(SourceArchetypeRef, MoveTemp(NewCompositionBitSet), SourceArchetypeRef.GetGroups());
 }
 
 FMassArchetypeHandle FMassEntityManager::InternalCreateSimilarArchetype(const TSharedPtr<FMassArchetypeData>& SourceArchetype, const UE::Mass::FArchetypeGroups& GroupsOverride)
 {
 	checkSlow(SourceArchetype.IsValid());
 	const FMassArchetypeData& SourceArchetypeRef = *SourceArchetype.Get();
-	FMassArchetypeCompositionDescriptor NewComposition = SourceArchetype->GetCompositionDescriptor();
-	return InternalCreateSimilarArchetype(SourceArchetypeRef, MoveTemp(NewComposition), GroupsOverride);
+	return InternalCreateSimilarArchetype(SourceArchetypeRef, FMassArchetypeCompositionDescriptor(SourceArchetype->GetCompositionBitSet()), GroupsOverride);
 }
 
 FMassArchetypeHandle FMassEntityManager::InternalCreateSimilarArchetype(const FMassArchetypeData& SourceArchetypeRef, FMassArchetypeCompositionDescriptor&& NewComposition, const UE::Mass::FArchetypeGroups& Groups)
@@ -507,17 +497,11 @@ void FMassEntityManager::InternalAppendFragmentsAndTagsToArchetypeCompositionDes
 {
 	for (const UScriptStruct* Type : FragmentsAndTagsList)
 	{
-		if (UE::Mass::IsA<FMassFragment>(Type))
+		if (UE::Mass::IsA<FMassFragment>(Type)
+			|| UE::Mass::IsA<FMassTag>(Type)
+			|| UE::Mass::IsA<FMassChunkFragment>(Type))
 		{
-			InOutComposition.GetFragments().Add(*Type);
-		}
-		else if (UE::Mass::IsA<FMassTag>(Type))
-		{
-			InOutComposition.GetTags().Add(*Type);
-		}
-		else if (UE::Mass::IsA<FMassChunkFragment>(Type))
-		{
-			InOutComposition.GetChunkFragments().Add(*Type);
+			InOutComposition.Add(Type);
 		}
 		else
 		{
@@ -589,7 +573,7 @@ FMassEntityHandle FMassEntityManager::CreateEntity(const FMassArchetypeHandle& A
 
 	const FMassEntityHandle Entity = ReserveEntity();
 	InternalBuildEntity(Entity
-		, GetOrCreateSuitableArchetype(ArchetypeHandle, SharedFragmentValues.GetSharedFragmentBitSet(), SharedFragmentValues.GetConstSharedFragmentBitSet())
+		, GetOrCreateSuitableArchetype(ArchetypeHandle, SharedFragmentValues.GetBitSet())
 		, SharedFragmentValues);
 
 	return Entity;
@@ -651,11 +635,11 @@ void FMassEntityManager::BuildEntity(FMassEntityHandle Entity, TConstArrayView<F
 	FMassArchetypeCompositionDescriptor Composition(FragmentInstanceList, FMassTagBitSet(), FMassChunkFragmentBitSet(), FMassSharedFragmentBitSet(), FMassConstSharedFragmentBitSet());
 	for (const FConstSharedStruct& SharedFragment : SharedFragmentValues.GetConstSharedFragments())
 	{
-		Composition.GetConstSharedFragments().Add(*SharedFragment.GetScriptStruct());
+		Composition.Add(SharedFragment.GetScriptStruct());
 	}
 	for (const FSharedStruct& SharedFragment : SharedFragmentValues.GetSharedFragments())
 	{
-		Composition.GetSharedFragments().Add(*SharedFragment.GetScriptStruct());
+		Composition.Add(SharedFragment.GetScriptStruct());
 	}
 
 	const FMassArchetypeHandle& ArchetypeHandle = CreateArchetype(Composition);
@@ -694,11 +678,11 @@ TSharedRef<FMassEntityManager::FEntityCreationContext> FMassEntityManager::Batch
 	FMassArchetypeCompositionDescriptor Composition(FragmentsAffected, FMassTagBitSet(), FMassChunkFragmentBitSet(), FMassSharedFragmentBitSet(), FMassConstSharedFragmentBitSet());
 	for (const FConstSharedStruct& SharedFragment : SharedFragmentValues.GetConstSharedFragments())
 	{
-		Composition.GetConstSharedFragments().Add(*SharedFragment.GetScriptStruct());
+		Composition.Add(SharedFragment.GetScriptStruct());
 	}
 	for (const FSharedStruct& SharedFragment : SharedFragmentValues.GetSharedFragments())
 	{
-		Composition.GetSharedFragments().Add(*SharedFragment.GetScriptStruct());
+		Composition.Add(SharedFragment.GetScriptStruct());
 	}
 
 	return BatchBuildEntities(EncodedEntitiesWithPayload, MoveTemp(Composition), SharedFragmentValues, CreationParams);
@@ -770,7 +754,7 @@ TSharedRef<FMassEntityManager::FEntityCreationContext> FMassEntityManager::Batch
 	TRACE_CPUPROFILER_EVENT_SCOPE(Mass_BatchCreateReservedEntities);
 
 	return InternalBatchCreateReservedEntities(
-		GetOrCreateSuitableArchetype(ArchetypeHandle, SharedFragmentValues.GetSharedFragmentBitSet(), SharedFragmentValues.GetConstSharedFragmentBitSet())
+		GetOrCreateSuitableArchetype(ArchetypeHandle, SharedFragmentValues.GetBitSet())
 		, SharedFragmentValues, ReservedEntities);
 }
 
@@ -788,7 +772,7 @@ TSharedRef<FMassEntityManager::FEntityCreationContext> FMassEntityManager::Batch
 	TConstArrayView<FMassEntityHandle> ReservedEntities = BatchReserveEntities(Count, InOutEntities);
 	
 	return InternalBatchCreateReservedEntities(
-		GetOrCreateSuitableArchetype(ArchetypeHandle, SharedFragmentValues.GetSharedFragmentBitSet(), SharedFragmentValues.GetConstSharedFragmentBitSet())
+		GetOrCreateSuitableArchetype(ArchetypeHandle, SharedFragmentValues.GetBitSet())
 		, SharedFragmentValues, ReservedEntities);
 }
 
@@ -1063,8 +1047,8 @@ void FMassEntityManager::AddCompositionToEntity_GetDelta(FMassEntityHandle Entit
 
 	InOutDescriptor.Remove(OldArchetype->GetCompositionDescriptor());
 
-	ensureMsgf(InOutDescriptor.GetChunkFragments().IsEmpty(), TEXT("Adding new chunk fragments is not supported"));
-	ensureMsgf(InOutDescriptor.GetSharedFragments().IsEmpty() 
+	ensureMsgf(InOutDescriptor.HasAny<FMassChunkFragment>() == false, TEXT("Adding new chunk fragments is not supported"));
+	ensureMsgf(InOutDescriptor.HasAny<FMassSharedFragment>() == false
 		|| (AddedSharedFragmentValues && AddedSharedFragmentValues->DoesMatchComposition(InOutDescriptor))
 		, TEXT("When adding new shared fragments it's required to provide values for said fragments"));
 
@@ -1114,7 +1098,7 @@ void FMassEntityManager::RemoveCompositionFromEntity(FMassEntityHandle Entity, c
 		FMassArchetypeCompositionDescriptor NewDescriptor = OldArchetype->GetCompositionDescriptor();
 		NewDescriptor.Remove(InDescriptor);
 
-		ensureMsgf(InDescriptor.GetChunkFragments().IsEmpty(), TEXT("Removing chunk fragments is not supported"));
+		ensureMsgf(InDescriptor.HasAny<FMassChunkFragment>() == false, TEXT("Removing chunk fragments is not supported"));
 
 		if (NewDescriptor.IsEquivalent(OldArchetype->GetCompositionDescriptor()) == false)
 		{
@@ -1127,7 +1111,7 @@ void FMassEntityManager::RemoveCompositionFromEntity(FMassEntityHandle Entity, c
 				// Move the entity over
 				FMassArchetypeData& NewArchetype = FMassArchetypeHelper::ArchetypeDataFromHandleChecked(NewArchetypeHandle);
 				NewArchetype.CopyDebugNamesFrom(*OldArchetype);
-				if (InDescriptor.GetSharedFragments().IsEmpty() && InDescriptor.GetConstSharedFragments().IsEmpty())
+				if (InDescriptor.GetElementsBitSet().HasAnySharedElements() == false)
 				{
 					OldArchetype->MoveEntityToAnotherArchetype(Entity, NewArchetype);
 				}
@@ -1145,7 +1129,7 @@ void FMassEntityManager::RemoveCompositionFromEntity(FMassEntityHandle Entity, c
 	}
 }
 
-const FMassArchetypeCompositionDescriptor& FMassEntityManager::GetArchetypeComposition(const FMassArchetypeHandle& ArchetypeHandle) const
+FMassArchetypeCompositionDescriptor FMassEntityManager::GetArchetypeComposition(const FMassArchetypeHandle& ArchetypeHandle) const
 {
 	const FMassArchetypeData& ArchetypeData = FMassArchetypeHelper::ArchetypeDataFromHandleChecked(ArchetypeHandle);
 	return ArchetypeData.GetCompositionDescriptor();
@@ -1173,11 +1157,11 @@ FMassFragmentBitSet FMassEntityManager::InternalAddFragmentListToEntityChecked(F
 	TSharedPtr<FMassArchetypeData>& OldArchetype = GetEntityStorageInterface().GetArchetypeAsShared(Entity.Index);
 	check(OldArchetype);
 
-	UE_CLOG(OldArchetype->GetFragmentBitSet().HasAny(InFragments), LogMass, Log
+	UE_CLOG(OldArchetype->GetCompositionBitSet().HasAny(InFragments), LogMass, Log
 		, TEXT("Trying to add a new fragment type to an entity, but it already has some of them. (%s)")
-		, *InFragments.GetOverlap(OldArchetype->GetFragmentBitSet()).DebugGetStringDesc());
+		, *InFragments.GetOverlap(OldArchetype->GetCompositionBitSet()).DebugGetStringDesc());
 
-	FMassFragmentBitSet NewFragments = InFragments - OldArchetype->GetFragmentBitSet();
+	FMassFragmentBitSet NewFragments = InFragments - OldArchetype->GetCompositionBitSet();
 	if (NewFragments.IsEmpty() == false)
 	{
 		InternalAddFragmentListToEntity(Entity, NewFragments);
@@ -1236,18 +1220,14 @@ void FMassEntityManager::RemoveFragmentListFromEntity(FMassEntityHandle Entity,
 
 	const FMassFragmentBitSet FragmentsToRemove(FragmentList);
 
-	if (OldArchetype->GetFragmentBitSet().HasAny(FragmentsToRemove))
+	if (OldArchetype->GetCompositionBitSet().HasAny(FragmentsToRemove))
 	{
 		// If all the fragments got removed this will result in fetching of the empty archetype
-		const FMassArchetypeCompositionDescriptor NewComposition(OldArchetype->GetFragmentBitSet() - FragmentsToRemove
-			, OldArchetype->GetTagBitSet()
-			, OldArchetype->GetChunkFragmentBitSet()
-			, OldArchetype->GetSharedFragmentBitSet()
-			, OldArchetype->GetConstSharedFragmentBitSet());
+		const FMassArchetypeCompositionDescriptor NewComposition(OldArchetype->GetCompositionBitSet() - FragmentsToRemove);
 		const FMassArchetypeHandle NewArchetypeHandle = CreateArchetype(NewComposition, FMassArchetypeCreationParams(*OldArchetype));
 
 		// Find overlap.  It isn't guaranteed that the old archetype has all of the fragments being removed.
-		FMassArchetypeCompositionDescriptor CompositionDelta(OldArchetype->GetFragmentBitSet().GetOverlap(FragmentsToRemove));
+		FMassArchetypeCompositionDescriptor CompositionDelta(OldArchetype->GetCompositionBitSet().GetOverlap(FragmentsToRemove));
 		ObserverManager.OnPreCompositionRemoved(Entity, CompositionDelta);
 
 		// Move the entity over
@@ -1263,6 +1243,11 @@ void FMassEntityManager::SwapTagsForEntity(FMassEntityHandle Entity, const UScri
 {
 	CHECK_SYNC_API();
 
+	if (OldTagType == NewTagType)
+	{
+		return;
+	}
+
 	CheckIfEntityIsActive(Entity);
 
 	checkf(UE::Mass::IsA<FMassTag>(OldTagType), TEXT("%hs works only with tags while '%s' is not one."), __FUNCTION__, *GetPathNameSafe(OldTagType));
@@ -1271,20 +1256,19 @@ void FMassEntityManager::SwapTagsForEntity(FMassEntityHandle Entity, const UScri
 	TSharedPtr<FMassArchetypeData>& CurrentArchetype = GetEntityStorageInterface().GetArchetypeAsShared(Entity.Index);
 	check(CurrentArchetype);
 
-	FMassTagBitSet NewTagBitSet = CurrentArchetype->GetTagBitSet();
-	NewTagBitSet.Remove(*OldTagType);
-	NewTagBitSet.Add(*NewTagType);
+	FMassElementBitSet NewCompositionBitSet = CurrentArchetype->GetCompositionBitSet();
+	NewCompositionBitSet.Remove(*OldTagType);
+	NewCompositionBitSet.Add(*NewTagType);
 	
-	if (NewTagBitSet != CurrentArchetype->GetTagBitSet())
-	{
-		const FMassArchetypeHandle NewArchetypeHandle = InternalCreateSimilarArchetype(CurrentArchetype, NewTagBitSet);
-		checkSlow(NewArchetypeHandle.IsValid());
+	ensureMsgf(NewCompositionBitSet.IsEquivalent(CurrentArchetype->GetCompositionBitSet()) == false, TEXT("Failing this test at this point is unexpected."));
 
-		// Move the entity over
-		CurrentArchetype->MoveEntityToAnotherArchetype(Entity, *NewArchetypeHandle.DataPtr.Get());
-		
-		GetEntityStorageInterface().SetArchetypeFromShared(Entity.Index, NewArchetypeHandle.DataPtr);
-	}
+	const FMassArchetypeHandle NewArchetypeHandle = InternalCreateSimilarArchetype(CurrentArchetype, MoveTemp(NewCompositionBitSet));
+	checkSlow(NewArchetypeHandle.IsValid());
+
+	// Move the entity over
+	CurrentArchetype->MoveEntityToAnotherArchetype(Entity, *NewArchetypeHandle.DataPtr.Get());
+
+	GetEntityStorageInterface().SetArchetypeFromShared(Entity.Index, NewArchetypeHandle.DataPtr);
 }
 
 void FMassEntityManager::AddTagToEntity(FMassEntityHandle Entity, const UScriptStruct* TagType)
@@ -1299,9 +1283,9 @@ void FMassEntityManager::AddTagToEntity(FMassEntityHandle Entity, const UScriptS
 
 	if (CurrentArchetype->HasTagType(TagType) == false)
 	{
-		FMassTagBitSet NewTags = CurrentArchetype->GetTagBitSet();
-		NewTags.Add(*TagType);
-		const FMassArchetypeHandle NewArchetypeHandle = InternalCreateSimilarArchetype(CurrentArchetype, NewTags);
+		FMassElementBitSet NewCompositionBitSet = CurrentArchetype->GetCompositionBitSet();
+		NewCompositionBitSet.Add(*TagType);
+		const FMassArchetypeHandle NewArchetypeHandle = InternalCreateSimilarArchetype(CurrentArchetype, MoveTemp(NewCompositionBitSet));
 		checkSlow(NewArchetypeHandle.IsValid());
 
 		// Move the entity over
@@ -1324,11 +1308,12 @@ void FMassEntityManager::RemoveTagFromEntity(FMassEntityHandle Entity, const USc
 
 	if (CurrentArchetype->HasTagType(TagType))
 	{
-		FMassTagBitSet TagDelta(*TagType);
-		const FMassTagBitSet NewTagComposition = CurrentArchetype->GetTagBitSet() - TagDelta;
+		FMassElementBitSet TagDelta(*TagType);
+		FMassElementBitSet NewComposition = CurrentArchetype->GetCompositionBitSet() - TagDelta;
+
 		ObserverManager.OnPreCompositionRemoved(Entity, FMassArchetypeCompositionDescriptor(MoveTemp(TagDelta)));
 
-		const FMassArchetypeHandle NewArchetypeHandle = InternalCreateSimilarArchetype(CurrentArchetype, NewTagComposition);
+		const FMassArchetypeHandle NewArchetypeHandle = InternalCreateSimilarArchetype(CurrentArchetype, MoveTemp(NewComposition));
 		checkSlow(NewArchetypeHandle.IsValid());
 
 		// Move the entity over
@@ -1346,10 +1331,10 @@ bool FMassEntityManager::DoesEntityHaveElement(FMassEntityHandle Entity, TNotNul
 	check(CurrentArchetype);
 	if (UE::Mass::IsA<FMassTag>(ElementType))
 	{
-		return CurrentArchetype->GetCompositionDescriptor().GetTags().Contains(*ElementType);
+		return CurrentArchetype->GetCompositionDescriptor().Contains(ElementType);
 	}
 	
-	return CurrentArchetype->GetCompositionDescriptor().GetFragments().Contains(*ElementType);
+	return CurrentArchetype->GetCompositionDescriptor().Contains(ElementType);
 }
 
 void FMassEntityManager::AddElementToEntities(TConstArrayView<FMassEntityHandle> Entities, TNotNull<const UScriptStruct*> ElementType)
@@ -1430,7 +1415,7 @@ bool FMassEntityManager::AddConstSharedFragmentToEntity(const FMassEntityHandle
 
 	const UScriptStruct* StructType = InConstSharedFragment.GetScriptStruct();
 	CA_ASSUME(StructType);
-	if (CurrentArchetype->GetCompositionDescriptor().GetConstSharedFragments().Contains(*StructType))
+	if (CurrentArchetype->GetCompositionDescriptor().Contains(StructType))
 	{
 		const FMassArchetypeSharedFragmentValues& SharedFragmentValues = CurrentArchetype->GetSharedFragmentValues(Entity);
 		FConstSharedStruct ExistingConstSharedStruct = SharedFragmentValues.GetConstSharedFragmentStruct(StructType);
@@ -1444,7 +1429,7 @@ bool FMassEntityManager::AddConstSharedFragmentToEntity(const FMassEntityHandle
 	}
 	
 	FMassArchetypeCompositionDescriptor NewComposition(CurrentArchetype->GetCompositionDescriptor());
-	NewComposition.GetConstSharedFragments().Add(*StructType);
+	NewComposition.Add(StructType);
 	const FMassArchetypeHandle NewArchetypeHandle = CreateArchetype(NewComposition, FMassArchetypeCreationParams(*CurrentArchetype));
 	check(NewArchetypeHandle.IsValid());
 	FMassArchetypeData* NewArchetype = NewArchetypeHandle.DataPtr.Get();
@@ -1478,14 +1463,14 @@ bool FMassEntityManager::RemoveConstSharedFragmentFromEntity(const FMassEntityHa
 	FMassArchetypeData* CurrentArchetype = GetEntityStorageInterface().GetArchetypeAsShared(Entity.Index).Get();
 	check(CurrentArchetype);
 	
-	if (!CurrentArchetype->GetCompositionDescriptor().GetConstSharedFragments().Contains(ConstSharedFragmentType))
+	if (!CurrentArchetype->GetCompositionDescriptor().Contains(&ConstSharedFragmentType))
 	{
 		// Nothing to do. Returning false to indicate nothing has been removed, as per function's documentation 
 		return false;
 	}
 
 	FMassArchetypeCompositionDescriptor NewComposition(CurrentArchetype->GetCompositionDescriptor());
-	NewComposition.GetConstSharedFragments().Remove(ConstSharedFragmentType);
+	NewComposition.GetElementsBitSet().Remove(ConstSharedFragmentType);
 	const FMassArchetypeHandle NewArchetypeHandle = CreateArchetype(NewComposition);
 	check(NewArchetypeHandle.IsValid());
 	FMassArchetypeData* NewArchetype = NewArchetypeHandle.DataPtr.Get();
@@ -1523,7 +1508,7 @@ bool FMassEntityManager::AddSharedFragmentToEntity(const FMassEntityHandle Entit
 
 	const UScriptStruct* StructType = InSharedFragment.GetScriptStruct();
 	CA_ASSUME(StructType);
-	if (CurrentArchetype->GetCompositionDescriptor().GetSharedFragments().Contains(*StructType))
+	if (CurrentArchetype->GetCompositionBitSet().Contains(*StructType))
 	{
 		const FMassArchetypeSharedFragmentValues& SharedFragmentValues = CurrentArchetype->GetSharedFragmentValues(Entity);
 		FConstSharedStruct ExistingSharedStruct = SharedFragmentValues.GetSharedFragmentStruct(StructType);
@@ -1537,7 +1522,7 @@ bool FMassEntityManager::AddSharedFragmentToEntity(const FMassEntityHandle Entit
 	}
 	
 	FMassArchetypeCompositionDescriptor NewComposition(CurrentArchetype->GetCompositionDescriptor());
-	NewComposition.GetSharedFragments().Add(*StructType);
+	NewComposition.GetElementsBitSet().Add(*StructType);
 	const FMassArchetypeHandle NewArchetypeHandle = CreateArchetype(NewComposition, FMassArchetypeCreationParams(*CurrentArchetype));
 	check(NewArchetypeHandle.IsValid());
 	FMassArchetypeData* NewArchetype = NewArchetypeHandle.DataPtr.Get();
@@ -1571,14 +1556,14 @@ bool FMassEntityManager::RemoveSharedFragmentFromEntity(const FMassEntityHandle
 	FMassArchetypeData* CurrentArchetype = GetEntityStorageInterface().GetArchetypeAsShared(Entity.Index).Get();
 	check(CurrentArchetype);
 	
-	if (!CurrentArchetype->GetCompositionDescriptor().GetSharedFragments().Contains(SharedFragmentType))
+	if (!CurrentArchetype->GetCompositionBitSet().Contains(SharedFragmentType))
 	{
 		// Nothing to do. Returning false to indicate nothing has been removed, as per function's documentation 
 		return false;
 	}
 
 	FMassArchetypeCompositionDescriptor NewComposition(CurrentArchetype->GetCompositionDescriptor());
-	NewComposition.GetSharedFragments().Remove(SharedFragmentType);
+	NewComposition.GetElementsBitSet().Remove(SharedFragmentType);
 	const FMassArchetypeHandle NewArchetypeHandle = CreateArchetype(NewComposition);
 	check(NewArchetypeHandle.IsValid());
 	FMassArchetypeData* NewArchetype = NewArchetypeHandle.DataPtr.Get();
@@ -1600,37 +1585,49 @@ bool FMassEntityManager::RemoveSharedFragmentFromEntity(const FMassEntityHandle
 	return true;
 }
 
-void FMassEntityManager::BatchChangeTagsForEntities(TConstArrayView<FMassArchetypeEntityCollection> EntityCollections, const FMassTagBitSet& TagsToAdd, const FMassTagBitSet& TagsToRemove)
+void FMassEntityManager::BatchChangeCompositionForEntities(TConstArrayView<FMassArchetypeEntityCollection> EntityCollections, const FMassElementBitSet& ElementsToAdd, const FMassElementBitSet& ElementsToRemove)
 {
 	CHECK_SYNC_API();
 
-	TRACE_CPUPROFILER_EVENT_SCOPE(Mass_BatchChangeTagsForEntities);
-	
+	TRACE_CPUPROFILER_EVENT_SCOPE(Mass_BatchChangeCompositionForEntities);
+
+	const bool bElementsAddedAreObserved = ObserverManager.HasObserversForBitSet(ElementsToAdd, EMassObservedOperation::AddElement);
+	const bool bElementsRemovedAreObserved = ObserverManager.HasObserversForBitSet(ElementsToRemove, EMassObservedOperation::RemoveElement);
+
 	for (const FMassArchetypeEntityCollection& Collection : EntityCollections)
 	{
-		FMassArchetypeData* CurrentArchetype = Collection.GetArchetype().DataPtr.Get();
-		const FMassTagBitSet NewTagComposition = CurrentArchetype
-			? (CurrentArchetype->GetTagBitSet() + TagsToAdd - TagsToRemove)
-			: (TagsToAdd - TagsToRemove);
-
-		if (ensure(CurrentArchetype) && CurrentArchetype->GetTagBitSet() != NewTagComposition)
+		if (FMassArchetypeData* CurrentArchetype = Collection.GetArchetype().DataPtr.Get())
 		{
-			FMassTagBitSet TagsAdded = TagsToAdd - CurrentArchetype->GetTagBitSet();
-			const bool bTagsAddedAreObserved = ObserverManager.HasObserversForBitSet(TagsAdded, EMassObservedOperation::AddElement);
-			FMassTagBitSet TagsRemoved = TagsToRemove.GetOverlap(CurrentArchetype->GetTagBitSet());
-			if (TagsRemoved.IsEmpty() == false)
+			FMassElementBitSet ElementsAdded;
+			FMassElementBitSet NewComposition = CurrentArchetype->GetCompositionBitSet() + ElementsToAdd - ElementsToRemove;
+
+			if (UNLIKELY(CurrentArchetype->GetCompositionBitSet().IsEquivalent(NewComposition)))
+			{
+				continue;
+			}
+
+			if (bElementsAddedAreObserved)
+			{
+				ElementsAdded = ElementsToAdd - CurrentArchetype->GetCompositionBitSet();
+			}
+
+			if (bElementsRemovedAreObserved)
 			{
-				ObserverManager.OnCompositionChanged(Collection, MoveTemp(TagsRemoved), EMassObservedOperation::RemoveElement);
+				FMassElementBitSet ElementsRemoved = ElementsToRemove.GetOverlap(CurrentArchetype->GetCompositionBitSet());
+				if (ElementsRemoved.IsEmpty() == false)
+				{
+					ObserverManager.OnCompositionChanged(Collection, MoveTemp(ElementsRemoved), EMassObservedOperation::RemoveElement);
+				}
 			}
 			
-			FMassArchetypeHandle NewArchetypeHandle = InternalCreateSimilarArchetype(Collection.GetArchetype().DataPtr, NewTagComposition);
+			FMassArchetypeHandle NewArchetypeHandle = InternalCreateSimilarArchetype(Collection.GetArchetype().DataPtr, MoveTemp(NewComposition));
 			checkSlow(NewArchetypeHandle.IsValid());
 
 			// Move the entity over
 			FMassArchetypeEntityCollection::FEntityRangeArray NewArchetypeEntityRanges;
 			TArray<FMassEntityHandle> EntitiesBeingMoved;
 			CurrentArchetype->BatchMoveEntitiesToAnotherArchetype(Collection, *NewArchetypeHandle.DataPtr.Get(), EntitiesBeingMoved
-				, bTagsAddedAreObserved ? &NewArchetypeEntityRanges : nullptr);
+				, bElementsAddedAreObserved && (ElementsAdded.IsEmpty() == false) ? &NewArchetypeEntityRanges : nullptr);
 
 			for (const FMassEntityHandle& Entity : EntitiesBeingMoved)
 			{
@@ -1639,71 +1636,18 @@ void FMassEntityManager::BatchChangeTagsForEntities(TConstArrayView<FMassArchety
 				GetEntityStorageInterface().SetArchetypeFromShared(Entity.Index, NewArchetypeHandle.DataPtr);
 			}
 
-			if (TagsAdded.IsEmpty() == false)
+			if (ElementsAdded.IsEmpty() == false)
 			{
 				ObserverManager.OnCompositionChanged(
 					FMassArchetypeEntityCollection(NewArchetypeHandle, MoveTemp(NewArchetypeEntityRanges))
-					, MoveTemp(TagsAdded)
+					, MoveTemp(ElementsAdded)
 					, EMassObservedOperation::AddElement);
 			}
 		}
-	}
-}
-
-void FMassEntityManager::BatchChangeFragmentCompositionForEntities(TConstArrayView<FMassArchetypeEntityCollection> EntityCollections, const FMassFragmentBitSet& FragmentsToAdd, const FMassFragmentBitSet& FragmentsToRemove)
-{
-	CHECK_SYNC_API();
-
-	TRACE_CPUPROFILER_EVENT_SCOPE(Mass_BatchChangeFragmentCompositionForEntities);
-
-	for (const FMassArchetypeEntityCollection& Collection : EntityCollections)
-	{
-		FMassArchetypeData* CurrentArchetype = Collection.GetArchetype().DataPtr.Get();
-		const FMassFragmentBitSet NewFragmentComposition = CurrentArchetype
-			? (CurrentArchetype->GetFragmentBitSet() + FragmentsToAdd - FragmentsToRemove)
-			: (FragmentsToAdd - FragmentsToRemove);
-
-		if (CurrentArchetype)
+		else // only valid if there are fragments in NewComposition
 		{
-			if (CurrentArchetype->GetFragmentBitSet() != NewFragmentComposition)
-			{
-				FMassFragmentBitSet FragmentsAdded = FragmentsToAdd - CurrentArchetype->GetFragmentBitSet();
-				const bool bFragmentsAddedAreObserved = ObserverManager.HasObserversForBitSet(FragmentsAdded, EMassObservedOperation::AddElement);
-				FMassFragmentBitSet FragmentsRemoved = FragmentsToRemove.GetOverlap(CurrentArchetype->GetFragmentBitSet());
-				
-				if (FragmentsRemoved.IsEmpty() == false)
-				{
-					ObserverManager.OnCompositionChanged(Collection, MoveTemp(FragmentsRemoved), EMassObservedOperation::RemoveElement);
-				}
-
-				FMassArchetypeHandle NewArchetypeHandle = InternalCreateSimilarArchetype(Collection.GetArchetype().DataPtr, NewFragmentComposition);
-				checkSlow(NewArchetypeHandle.IsValid());
-
-				// Move the entity over
-				FMassArchetypeEntityCollection::FEntityRangeArray NewArchetypeEntityRanges;
-				TArray<FMassEntityHandle> EntitiesBeingMoved;
-				CurrentArchetype->BatchMoveEntitiesToAnotherArchetype(Collection, *NewArchetypeHandle.DataPtr.Get(), EntitiesBeingMoved
-					, bFragmentsAddedAreObserved ? &NewArchetypeEntityRanges : nullptr);
-
-				for (const FMassEntityHandle& Entity : EntitiesBeingMoved)
-				{
-					check(GetEntityStorageInterface().IsValidIndex(Entity.Index));
-
-					GetEntityStorageInterface().SetArchetypeFromShared(Entity.Index, NewArchetypeHandle.DataPtr);
-				}
-
-				if (bFragmentsAddedAreObserved)
-				{
-					ObserverManager.OnCompositionChanged(
-						FMassArchetypeEntityCollection(NewArchetypeHandle, MoveTemp(NewArchetypeEntityRanges))
-						, MoveTemp(FragmentsAdded)
-						, EMassObservedOperation::AddElement);
-				}
-			}
-		}
-		else
-		{
-			BatchBuildEntities(FMassArchetypeEntityCollectionWithPayload(Collection), NewFragmentComposition, FMassArchetypeSharedFragmentValues());
+			FMassElementBitSet NewComposition = (ElementsToAdd - ElementsToRemove);
+			BatchBuildEntities(FMassArchetypeEntityCollectionWithPayload(Collection), FMassArchetypeCompositionDescriptor(MoveTemp(NewComposition)));
 		}
 	}
 }
@@ -1714,6 +1658,8 @@ void FMassEntityManager::BatchAddFragmentInstancesForEntities(TConstArrayView<FM
 
 	TRACE_CPUPROFILER_EVENT_SCOPE(Mass_BatchAddFragmentInstancesForEntities);
 
+	FMassElementBitSet ElementsAffected = FragmentsAffected;
+
 	// here's the scenario:
 	// * we get entities from potentially different archetypes
 	// * adding a fragment instance consists of two operations: A) add fragment type & B) set fragment value
@@ -1723,7 +1669,8 @@ void FMassEntityManager::BatchAddFragmentInstancesForEntities(TConstArrayView<FM
 	// * setting value can be done uniformly for all entities, remembering some might be in different chunks already
 	// * @todo note that after adding fragment type some entities originally in different archetypes end up in the same 
 	//		archetype. This could be utilized as a basis for optimization. To be investigated.
-	// 
+	//
+	const bool bFragmentsAddedAreObserved = ObserverManager.HasObserversForBitSet(ElementsAffected, EMassObservedOperation::AddElement);
 
 	for (const FMassArchetypeEntityCollectionWithPayload& EntityRangesWithPayload : EntityCollections)
 	{
@@ -1733,18 +1680,18 @@ void FMassEntityManager::BatchAddFragmentInstancesForEntities(TConstArrayView<FM
 		if (CurrentArchetype)
 		{
 			FMassArchetypeEntityCollection::FEntityRangeArray TargetArchetypeEntityRanges;
-			bool bFragmentsAddedAreObserved = false;
-			FMassFragmentBitSet NewFragmentComposition = CurrentArchetype
-				? (CurrentArchetype->GetFragmentBitSet() + FragmentsAffected)
-				: FragmentsAffected;
-			FMassFragmentBitSet FragmentsAdded;
+			FMassElementBitSet NewComposition = ElementsAffected + CurrentArchetype->GetCompositionBitSet();
+			FMassElementBitSet ElementsAdded;
 
-			if (CurrentArchetype->GetFragmentBitSet() != NewFragmentComposition)
+			if (CurrentArchetype->GetCompositionBitSet().IsEquivalent(NewComposition) == false)
 			{
-				FragmentsAdded = FragmentsAffected - CurrentArchetype->GetFragmentBitSet();
-				bFragmentsAddedAreObserved = ObserverManager.HasObserversForBitSet(FragmentsAdded, EMassObservedOperation::AddElement);
+				if (bFragmentsAddedAreObserved)
+				{
+					ElementsAdded = ElementsAffected - CurrentArchetype->GetCompositionBitSet();
+				}
+				
 
-				FMassArchetypeHandle NewArchetypeHandle = InternalCreateSimilarArchetype(TargetArchetypeHandle.DataPtr, NewFragmentComposition);
+				FMassArchetypeHandle NewArchetypeHandle = InternalCreateSimilarArchetype(TargetArchetypeHandle.DataPtr, MoveTemp(NewComposition));
 				checkSlow(NewArchetypeHandle.IsValid());
 
 				// Move the entity over
@@ -1771,11 +1718,11 @@ void FMassEntityManager::BatchAddFragmentInstancesForEntities(TConstArrayView<FM
 			// corresponds to the order in FMassArchetypeEntityCollectionWithPayload's payload
 			TargetArchetypeHandle.DataPtr->BatchSetFragmentValues(TargetArchetypeEntityRanges, EntityRangesWithPayload.GetPayload());
 			
-			if (bFragmentsAddedAreObserved)
+			if (bFragmentsAddedAreObserved && (ElementsAdded.IsEmpty() == false))
 			{
 				ObserverManager.OnCompositionChanged(
 					FMassArchetypeEntityCollection(TargetArchetypeHandle, MoveTemp(TargetArchetypeEntityRanges))
-					, MoveTemp(FragmentsAdded)
+					, MoveTemp(ElementsAdded)
 					, EMassObservedOperation::AddElement);
 			}
 		}
@@ -1799,8 +1746,7 @@ void FMassEntityManager::BatchAddSharedFragmentsForEntities(TConstArrayView<FMas
 		testableCheckfReturn(CurrentArchetype, continue, TEXT("Adding shared fragments to archetype-less entities is not supported"));
 
 		FMassArchetypeCompositionDescriptor NewComposition(CurrentArchetype->GetCompositionDescriptor());
-		NewComposition.GetSharedFragments() += AddedFragmentValues.GetSharedFragmentBitSet();
-		NewComposition.GetConstSharedFragments() += AddedFragmentValues.GetConstSharedFragmentBitSet();
+		NewComposition.GetElementsBitSet() += AddedFragmentValues.GetBitSet();
 
 		const FMassArchetypeHandle NewArchetypeHandle = CreateArchetype(NewComposition, FMassArchetypeCreationParams(*CurrentArchetype));
 		check(NewArchetypeHandle.IsValid());
@@ -2186,7 +2132,7 @@ void FMassEntityManager::DebugGetArchetypeFragmentTypes(const FMassArchetypeHand
 	if (Archetype.IsValid())
 	{
 		const FMassArchetypeData& ArchetypeData = FMassArchetypeHelper::ArchetypeDataFromHandleChecked(Archetype);
-		ArchetypeData.GetCompositionDescriptor().GetFragments().ExportTypes(InOutFragmentList);
+		ArchetypeData.GetCompositionDescriptor().Get<FMassFragmentBitSet>().ExportTypes(InOutFragmentList);
 	}
 }
 
@@ -2251,7 +2197,7 @@ void FMassEntityManager::DebugGetArchetypeStrings(const FMassArchetypeHandle& Ar
 		OutFragmentNames.Add(FragmentConfig.FragmentType->GetFName());
 	}
 
-	ArchetypeRef.GetTagBitSet().DebugGetIndividualNames(OutTagNames);
+	ArchetypeRef.GetCompositionBitSet().Get<FMassTagBitSet>().DebugGetIndividualNames(OutTagNames);
 }
 
 FMassEntityHandle FMassEntityManager::DebugGetEntityIndexHandle(const int32 EntityIndex) const
diff --git a/Engine/Source/Runtime/MassEntity/Private/MassEntityModule.cpp b/Engine/Source/Runtime/MassEntity/Private/MassEntityModule.cpp
--- a/Engine/Source/Runtime/MassEntity/Private/MassEntityModule.cpp
+++ b/Engine/Source/Runtime/MassEntity/Private/MassEntityModule.cpp
@@ -6,6 +6,7 @@
 #include "MassProcessingTypes.h"
 #include "Modules/ModuleManager.h"
 #include "Relations/MassChildOf.h"
+#include "MassElement.h"
 
 #define LOCTEXT_NAMESPACE "Mass"
 
@@ -31,6 +32,7 @@ void FMassEntityModule::StartupModule()
 #endif // MASS_DO_PARALLEL
 
 	UE::Mass::Relations::RegisterChildOfRelation();
+	UE::Mass::FElementBitSet::OnModuleInitialized();
 }
 
 void FMassEntityModule::ShutdownModule()
@@ -38,6 +40,7 @@ void FMassEntityModule::ShutdownModule()
 #if WITH_MASSENTITY_DEBUG
 	FMassDebugger::ShutdownDebugger();
 #endif
+	UE::Mass::FElementBitSet::OnModuleShutdown();
 }
 
 #undef LOCTEXT_NAMESPACE 
diff --git a/Engine/Source/Runtime/MassEntity/Private/MassEntityTrace.cpp b/Engine/Source/Runtime/MassEntity/Private/MassEntityTrace.cpp
--- a/Engine/Source/Runtime/MassEntity/Private/MassEntityTrace.cpp
+++ b/Engine/Source/Runtime/MassEntity/Private/MassEntityTrace.cpp
@@ -239,12 +239,13 @@ void FMassTrace::OnPhaseEnd(uint64 PhaseId)
 uint64 FMassTrace::OutputRegisterArchetype(uint64 ArchetypeID, const FMassArchetypeCompositionDescriptor& CompositionDescriptor)
 {
 	TArray<uint64> FragmentsScratch;
-	FragmentsScratch.Reserve(CompositionDescriptor.GetFragments().CountStoredTypes() + CompositionDescriptor.GetTags().CountStoredTypes());
+	FMassElementBitSet FragmentsAndTags = CompositionDescriptor.GetElementsBitSet().GetFragmentsAndTags();
+	FragmentsScratch.Reserve(FragmentsAndTags.CountStoredTypes());
 
-	auto FragmentIterator = CompositionDescriptor.GetFragments().GetIndexIterator();
+	auto FragmentIterator = FragmentsAndTags.GetIndexIterator();
 	while (FragmentIterator)
 	{
-		const UScriptStruct* FragmentStruct = CompositionDescriptor.GetFragments().GetTypeAtIndex(*FragmentIterator);
+		const UScriptStruct* FragmentStruct = FragmentsAndTags.GetTypeAtIndex(*FragmentIterator);
 
 		OutputRegisterFragment(FragmentStruct);
 
@@ -253,18 +254,6 @@ uint64 FMassTrace::OutputRegisterArchetype(uint64 ArchetypeID, const FMassArchet
 		++FragmentIterator;
 	}
 
-	auto TagIterator = CompositionDescriptor.GetTags().GetIndexIterator();
-	while (TagIterator)
-	{
-		const UScriptStruct* FragmentStruct = CompositionDescriptor.GetTags().GetTypeAtIndex(*TagIterator);
-		{
-			OutputRegisterFragment(FragmentStruct);
-		}
-		// TODO Should have utility function to go from UScriptStruct to the fragment ID
-		FragmentsScratch.Add(reinterpret_cast<uint64>(FragmentStruct));
-		++TagIterator;
-	}
-
 	UE_TRACE_LOG(MassTrace, RegisterMassArchetype, MassChannel)
 		<< RegisterMassArchetype.ArchetypeID(ArchetypeID)
 		<< RegisterMassArchetype.Fragments(FragmentsScratch.GetData(), FragmentsScratch.Num());
diff --git a/Engine/Source/Runtime/MassEntity/Private/MassEntityTypes.cpp b/Engine/Source/Runtime/MassEntity/Private/MassEntityTypes.cpp
--- a/Engine/Source/Runtime/MassEntity/Private/MassEntityTypes.cpp
+++ b/Engine/Source/Runtime/MassEntity/Private/MassEntityTypes.cpp
@@ -8,14 +8,6 @@
 
 DEFINE_STAT(STAT_Mass_Total);
 
-DEFINE_TYPEBITSET(FMassFragmentBitSet);
-DEFINE_TYPEBITSET(FMassTagBitSet);
-DEFINE_TYPEBITSET(FMassChunkFragmentBitSet);
-DEFINE_TYPEBITSET(FMassSharedFragmentBitSet);
-DEFINE_TYPEBITSET(FMassConstSharedFragmentBitSet);
-DEFINE_TYPEBITSET(FMassExternalSubsystemBitSet);
-
-
 FString LexToString(const EMassObservedOperationFlags Value)
 {
 	switch (Value)
@@ -44,39 +36,16 @@ FString LexToString(const EMassObservedOperationFlags Value)
 //-----------------------------------------------------------------------------
 // FMassArchetypeCompositionDescriptor
 //-----------------------------------------------------------------------------
-uint32 FMassArchetypeCompositionDescriptor::CalculateHash(const FMassFragmentBitSet& InFragments, const FMassTagBitSet& InTags
-	, const FMassChunkFragmentBitSet& InChunkFragments, const FMassSharedFragmentBitSet& InSharedFragmentBitSet
-	, const FMassConstSharedFragmentBitSet& InConstSharedFragmentBitSet)
-{
-	const uint32 FragmentsHash = GetTypeHash(InFragments);
-	const uint32 TagsHash = GetTypeHash(InTags);
-	const uint32 ChunkFragmentsHash = GetTypeHash(InChunkFragments);
-	const uint32 SharedFragmentsHash = GetTypeHash(InSharedFragmentBitSet);
-	const uint32 ConstSharedFragmentsHash = GetTypeHash(InConstSharedFragmentBitSet);
-
-	return HashCombine(
-			HashCombine(FragmentsHash, TagsHash)
-			, HashCombine(
-				HashCombine(ChunkFragmentsHash, SharedFragmentsHash)
-				, ConstSharedFragmentsHash));
-}
-
 PRAGMA_DISABLE_DEPRECATION_WARNINGS
 int32 FMassArchetypeCompositionDescriptor::CountStoredTypes() const
 {
-	return Fragments.CountStoredTypes()
-		+ Tags.CountStoredTypes()
-		+ ChunkFragments.CountStoredTypes()
-		+ SharedFragments.CountStoredTypes()
-		+ ConstSharedFragments.CountStoredTypes();
+	return ElementsBitSet.CountStoredTypes();
 }
 
 void FMassArchetypeCompositionDescriptor::DebugOutputDescription(FOutputDevice& Ar) const
 {
 #if WITH_MASSENTITY_DEBUG 
-	if (Fragments.IsEmpty()
-		&& Tags.IsEmpty()
-		&& ChunkFragments.IsEmpty())
+	if (ElementsBitSet.IsEmpty())
 	{
 		Ar.Logf(TEXT("Empty"));
 		return;
@@ -85,35 +54,8 @@ void FMassArchetypeCompositionDescriptor::DebugOutputDescription(FOutputDevice&
 	const bool bAutoLineEnd = Ar.GetAutoEmitLineTerminator();
 	Ar.SetAutoEmitLineTerminator(false);
 
-	if (!Fragments.IsEmpty())
-	{
-		Ar.Logf(TEXT("Fragments:\n"));
-		Fragments.DebugGetStringDesc(Ar);
-	}
-
-	if (!Tags.IsEmpty())
-	{
-		Ar.Logf(TEXT("Tags:\n"));
-		Tags.DebugGetStringDesc(Ar);
-	}
-
-	if (!ChunkFragments.IsEmpty())
-	{
-		Ar.Logf(TEXT("ChunkFragments:\n"));
-		ChunkFragments.DebugGetStringDesc(Ar);
-	}
-
-	if (!SharedFragments.IsEmpty())
-	{
-		Ar.Logf(TEXT("SharedFragments:\n"));
-		SharedFragments.DebugGetStringDesc(Ar);
-	}
-
-	if (!ConstSharedFragments.IsEmpty())
-	{
-		Ar.Logf(TEXT("ConstSharedFragments:\n"));
-		ConstSharedFragments.DebugGetStringDesc(Ar);
-	}
+	Ar.Logf(TEXT("Elements:\n"));
+	ElementsBitSet.DebugGetStringDesc(Ar);
 
 	Ar.SetAutoEmitLineTerminator(bAutoLineEnd);
 #endif // WITH_MASSENTITY_DEBUG
@@ -123,6 +65,28 @@ PRAGMA_ENABLE_DEPRECATION_WARNINGS
 //-----------------------------------------------------------------------------
 // FMassArchetypeSharedFragmentValues
 //-----------------------------------------------------------------------------
+FMassArchetypeSharedFragmentValues FMassArchetypeSharedFragmentValues::CreateCombined(const FMassArchetypeSharedFragmentValues& Original
+	, const FMassElementBitSet& NewSharedElementComposition, const FMassArchetypeSharedFragmentValues* SharedFragmentValuesModification)
+{
+	FMassArchetypeSharedFragmentValues NewSharedValues = Original;
+	if (SharedFragmentValuesModification)
+	{
+		NewSharedValues.Append(*SharedFragmentValuesModification);
+	}
+
+	// remove everything that's in ConstSharedFragments or SharedFragments that's not in NewSharedElementComposition
+	auto RemovePredicate = [&NewSharedElementComposition](const FConstSharedStruct& SharedStruct)
+		{
+			return SharedStruct.GetScriptStruct() == nullptr
+				|| NewSharedElementComposition.Contains(*SharedStruct.GetScriptStruct()) == false;
+		};
+	NewSharedValues.SharedFragments.RemoveAllSwap(RemovePredicate);
+	NewSharedValues.ConstSharedFragments.RemoveAllSwap(RemovePredicate);
+
+	NewSharedValues.Sort();
+	return NewSharedValues;
+}
+
 FConstSharedStruct FMassArchetypeSharedFragmentValues::Add_GetRef(const FConstSharedStruct& Fragment)
 {
 	check(Fragment.IsValid());
@@ -137,7 +101,7 @@ FConstSharedStruct FMassArchetypeSharedFragmentValues::Add_GetRef(const FConstSh
 	}
 
 	check(StructType);
-	ConstSharedFragmentBitSet.Add(*StructType);
+	StoredElementsBitSet.Add(*StructType);
 	FConstSharedStruct& StructInstance = ConstSharedFragments.Add_GetRef(Fragment);
 	DirtyHashCache();
 	return StructInstance;
@@ -157,7 +121,7 @@ FSharedStruct FMassArchetypeSharedFragmentValues::Add_GetRef(const FSharedStruct
 	}
 
 	check(StructType);
-	SharedFragmentBitSet.Add(*StructType);
+	StoredElementsBitSet.Add(*StructType);
 	FSharedStruct& StructInstance = SharedFragments.Add_GetRef(Fragment);
 	DirtyHashCache();
 	return StructInstance;
@@ -266,8 +230,7 @@ namespace UE::Mass::Private
 
 bool FMassArchetypeSharedFragmentValues::HasSameValues(const FMassArchetypeSharedFragmentValues& Other) const
 {
-	if (SharedFragmentBitSet.IsEquivalent(Other.SharedFragmentBitSet) == false
-		|| ConstSharedFragmentBitSet.IsEquivalent(Other.ConstSharedFragmentBitSet) == false)
+	if (StoredElementsBitSet.IsEquivalent(Other.StoredElementsBitSet) == false)
 	{
 		return false;
 	}
@@ -284,7 +247,7 @@ int32 FMassArchetypeSharedFragmentValues::Append(const FMassArchetypeSharedFragm
 	{
 		const UScriptStruct* StructType = SharedStruct.GetScriptStruct();
 		check(StructType);
-		if (SharedFragmentBitSet.Contains(*StructType))
+		if (StoredElementsBitSet.Contains(*StructType))
 		{
 			const int32 FragmentIndex = SharedFragments.IndexOfByPredicate(FStructTypeEqualOperator(StructType));
 			checkf(FragmentIndex != INDEX_NONE, TEXT("Mismatch between shared fragment bitset and stored values"));
@@ -302,7 +265,7 @@ int32 FMassArchetypeSharedFragmentValues::Append(const FMassArchetypeSharedFragm
 	{
 		const UScriptStruct* StructType = SharedStruct.GetScriptStruct();
 		check(StructType);
-		if (ConstSharedFragmentBitSet.Contains(*StructType))
+		if (StoredElementsBitSet.Contains(*StructType))
 		{
 			const int32 FragmentIndex = ConstSharedFragments.IndexOfByPredicate(FStructTypeEqualOperator(StructType));
 			checkf(FragmentIndex != INDEX_NONE, TEXT("Mismatch between const shared fragment bitset and stored values"));
@@ -316,18 +279,18 @@ int32 FMassArchetypeSharedFragmentValues::Append(const FMassArchetypeSharedFragm
 		}
 	}
 
-	SharedFragmentBitSet += Other.SharedFragmentBitSet;
-	ConstSharedFragmentBitSet += Other.ConstSharedFragmentBitSet;
+	StoredElementsBitSet += Other.StoredElementsBitSet;
+
 	DirtyHashCache();
 
 	return AddedOrModifiedCount;
 }
 
-int32 FMassArchetypeSharedFragmentValues::Remove(const FMassSharedFragmentBitSet& SharedFragmentToRemoveBitSet)
+int32 FMassArchetypeSharedFragmentValues::Remove(const FMassElementBitSet& ElementBitSet)
 {
 	int32 RemovedCount = 0;
-	FMassSharedFragmentBitSet CommonFragments = (SharedFragmentBitSet & SharedFragmentToRemoveBitSet);
-	FMassSharedFragmentBitSet::FIndexIterator It = CommonFragments.GetIndexIterator();
+	FMassElementBitSet CommonFragments = (StoredElementsBitSet & ElementBitSet);
+	FMassConstSharedFragmentBitSet::FIndexIterator It = CommonFragments.GetIndexIterator();
 	while(It)
 	{
 		const UScriptStruct* StructType = CommonFragments.GetTypeAtIndex(*It);
@@ -339,34 +302,14 @@ int32 FMassArchetypeSharedFragmentValues::Remove(const FMassSharedFragmentBitSet
 			SharedFragments[RegularFragmentIndex].Reset();
 			++RemovedCount;
 		}
-
-		++It;
-	}
-
-	if (RemovedCount)
-	{
-		SharedFragments.RemoveAllSwap([](const FSharedStruct& SharedStruct) { return !SharedStruct.IsValid(); });
-		SharedFragmentBitSet -= CommonFragments;
-		DirtyHashCache();
-	}
-	return RemovedCount;
-}
-
-int32 FMassArchetypeSharedFragmentValues::Remove(const FMassConstSharedFragmentBitSet& ConstSharedFragmentToRemoveBitSet)
-{
-	int32 RemovedCount = 0;
-	FMassConstSharedFragmentBitSet CommonFragments = (ConstSharedFragmentBitSet & ConstSharedFragmentToRemoveBitSet);
-	FMassConstSharedFragmentBitSet::FIndexIterator It = CommonFragments.GetIndexIterator();
-	while(It)
-	{
-		const UScriptStruct* StructType = CommonFragments.GetTypeAtIndex(*It);
-		check(StructType);
-
-		const int32 RegularFragmentIndex = ConstSharedFragments.IndexOfByPredicate(FStructTypeEqualOperator(StructType));
-		if (RegularFragmentIndex != INDEX_NONE)
+		else
 		{
-			ConstSharedFragments[RegularFragmentIndex].Reset();
-			++RemovedCount;
+			const int32 ConstRegularFragmentIndex = ConstSharedFragments.IndexOfByPredicate(FStructTypeEqualOperator(StructType));
+			if (ConstRegularFragmentIndex != INDEX_NONE)
+			{
+				ConstSharedFragments[ConstRegularFragmentIndex].Reset();
+				++RemovedCount;
+			}
 		}
 
 		++It;
@@ -375,7 +318,9 @@ int32 FMassArchetypeSharedFragmentValues::Remove(const FMassConstSharedFragmentB
 	if (RemovedCount)
 	{
 		ConstSharedFragments.RemoveAllSwap([](const FConstSharedStruct& SharedStruct) { return !SharedStruct.IsValid(); });
-		ConstSharedFragmentBitSet -= CommonFragments;
+		SharedFragments.RemoveAllSwap([](const FConstSharedStruct& SharedStruct) { return !SharedStruct.IsValid(); });
+
+		StoredElementsBitSet -= CommonFragments;
 		DirtyHashCache();
 	}
 	return RemovedCount;
diff --git a/Engine/Source/Runtime/MassEntity/Private/MassExecutionContext.cpp b/Engine/Source/Runtime/MassEntity/Private/MassExecutionContext.cpp
--- a/Engine/Source/Runtime/MassEntity/Private/MassExecutionContext.cpp
+++ b/Engine/Source/Runtime/MassEntity/Private/MassExecutionContext.cpp
@@ -56,7 +56,7 @@ void FMassExecutionContext::ClearExecutionData()
 	ChunkFragmentViews.Reset();
 	ConstSharedFragmentViews.Reset();
 	SharedFragmentViews.Reset();
-	CurrentArchetypeCompositionDescriptor.Reset();
+	CurrentArchetypeCompositionBitSet.Reset();
 	EntityListView = {};
 	ChunkSerialModificationNumber = -1;
 #if WITH_MASSENTITY_DEBUG
diff --git a/Engine/Source/Runtime/MassEntity/Private/MassObserverManager.cpp b/Engine/Source/Runtime/MassEntity/Private/MassObserverManager.cpp
--- a/Engine/Source/Runtime/MassEntity/Private/MassObserverManager.cpp
+++ b/Engine/Source/Runtime/MassEntity/Private/MassObserverManager.cpp
@@ -62,7 +62,7 @@ namespace UE::Mass::ObserverManager
 						}
 						else if (ProcessorClass->GetDefaultObject<UMassProcessor>()->ShouldExecute(WorldExecutionFlags))
 						{
-							UMassProcessor* ObserverInstance = NewObject<UMassProcessor>(&Owner, ProcessorClass);							
+							UMassProcessor* ObserverInstance = NewObject<UMassProcessor>(&Owner, ProcessorClass);
 							ObserverProcessors.Add(ObserverInstance);
 							InOutExistingProcessorsMap.Add(ProcessorClass, ObserverInstance);
 						}
@@ -98,30 +98,23 @@ namespace UE::Mass::ObserverManager
 			return EMassObservedOperation::MAX;
 		};
 
-		static void HandleSingleElement(TNotNull<FMassObserverManager*> ObserverManager, const UScriptStruct& ElementType, const FMassArchetypeEntityCollection& EntityCollection, FMassObserversMap& HandlersContainer)
+		static void HandleSingleElement(TNotNull<FMassObserverManager*> ObserverManager, const UScriptStruct& ElementType, const FMassArchetypeEntityCollection& EntityCollection, const FMassObserversMap& HandlersContainer)
 		{
-			const bool bIsFragment = UE::Mass::IsA<FMassFragment>(&ElementType);
-			check(bIsFragment || UE::Mass::IsA<FMassTag>(&ElementType));
-
-			const EMassObservedOperation Operation = bIsFragment
-				? DetermineOperationType(&HandlersContainer, ObserverManager->EntityManager.GetObserverManager().FragmentObservers)
-				: DetermineOperationType(&HandlersContainer, ObserverManager->EntityManager.GetObserverManager().TagObservers);
-
+			const EMassObservedOperation Operation = DetermineOperationType(&HandlersContainer, ObserverManager->EntityManager.GetObserverManager().ElementObservers);
 			HandleSingleElement(ObserverManager, ElementType, EntityCollection, Operation);
 		}
 
 		static void HandleSingleElement(TNotNull<FMassObserverManager*> ObserverManager, const UScriptStruct& ElementType, const FMassArchetypeEntityCollection& EntityCollection, const EMassObservedOperation Operation)
 		{
-			const int32 OperationIndex = static_cast<int32>(Operation);
-			const bool bIsFragment = UE::Mass::IsA<FMassFragment>(&ElementType);
-			check(bIsFragment || UE::Mass::IsA<FMassTag>(&ElementType));
+			checkSlow(Operation != EMassObservedOperation::MAX);
 
+			const int32 OperationIndex = static_cast<int32>(Operation);
 			const UScriptStruct* ElementTypePtr = &ElementType;
 
 			UE::Mass::FProcessingContext ProcessingContext(ObserverManager->EntityManager);
 			FMassObserverManager::HandleElementsImpl(ProcessingContext, MakeArrayView(&EntityCollection, 1)
 				, {Operation, MakeArrayView(&ElementTypePtr, 1)}
-				, bIsFragment ? ObserverManager->FragmentObservers[OperationIndex] : ObserverManager->TagObservers[OperationIndex]);
+				, ObserverManager->ElementObservers[OperationIndex]);
 		}
 	};
 
@@ -131,95 +124,6 @@ namespace UE::Mass::ObserverManager
 		FProcessingContext& ProcessingContext;
 	};
 
-	struct FBufferedNotificationExecutioner 
-	{
-		FBufferedNotificationExecutioner(FNotificationContext& InNotificationContext)
-			: NotificationContext(InNotificationContext)
-		{}
-
-		template<typename TEntities>
-		void operator()(const FBufferedNotification::FEmptyComposition&, TEntities)
-		{
-			// no-op
-		}
-
-		void operator()(const FMassArchetypeCompositionDescriptor& Change, const FEntityCollection& Entities)
-		{
-			if (Change.GetFragments().IsEmpty() == false)
-			{
-				(*this)(Change.GetFragments(), Entities);
-			}
-			if (Change.GetTags().IsEmpty() == false)
-			{
-				(*this)(Change.GetTags(), Entities);
-			}
-		}
-
-		void operator()(const FMassArchetypeCompositionDescriptor& Change, const FMassEntityHandle EntityHandle)
-		{
-			if (Change.GetFragments().IsEmpty() == false)
-			{
-				(*this)(Change.GetFragments(), EntityHandle);
-			}
-			if (Change.GetTags().IsEmpty() == false)
-			{
-				(*this)(Change.GetTags(), EntityHandle);
-			}
-		}
-
-		void operator()(const FMassFragmentBitSet& Change, const FEntityCollection& Entities)
-		{
-			ObservedTypesOverlap.Reset();
-			Change.ExportTypes(ObservedTypesOverlap);
-			FMassObserverManager::HandleElementsImpl(NotificationContext.ProcessingContext
-				, Entities.GetUpToDatePerArchetypeCollections(NotificationContext.ObserverManager.EntityManager)
-				, {OpType, ObservedTypesOverlap}
-				, NotificationContext.ObserverManager.FragmentObservers[static_cast<uint8>(OpType)]);
-		}
-
-		void operator()(const FMassFragmentBitSet& Change, const FMassEntityHandle EntityHandle)
-		{
-			ObservedTypesOverlap.Reset();
-			Change.ExportTypes(ObservedTypesOverlap);
-			
-			FMassArchetypeHandle ArchetypeHandle = NotificationContext.ObserverManager.GetEntityManager().GetArchetypeForEntity(EntityHandle);
-			FMassArchetypeEntityCollection AsEntityCollection(MoveTemp(ArchetypeHandle), EntityHandle);
-
-			NotificationContext.ObserverManager.HandleElementsImpl(NotificationContext.ProcessingContext
-				, MakeArrayView(&AsEntityCollection, 1)
-				, {OpType, ObservedTypesOverlap}
-				, NotificationContext.ObserverManager.FragmentObservers[static_cast<uint8>(OpType)]);
-		}
-
-		void operator()(const FMassTagBitSet& Change, const FEntityCollection& Entities)
-		{
-			ObservedTypesOverlap.Reset();
-			Change.ExportTypes(ObservedTypesOverlap);
-			FMassObserverManager::HandleElementsImpl(NotificationContext.ProcessingContext
-				, Entities.GetUpToDatePerArchetypeCollections(NotificationContext.ObserverManager.EntityManager)
-				, {OpType, ObservedTypesOverlap}
-				, NotificationContext.ObserverManager.TagObservers[static_cast<uint8>(OpType)]);
-		}
-
-		void operator()(const FMassTagBitSet& Change, const FMassEntityHandle EntityHandle)
-		{
-			ObservedTypesOverlap.Reset();
-			Change.ExportTypes(ObservedTypesOverlap);
-
-			FMassArchetypeHandle ArchetypeHandle = NotificationContext.ObserverManager.GetEntityManager().GetArchetypeForEntity(EntityHandle);
-			FMassArchetypeEntityCollection AsEntityCollection(MoveTemp(ArchetypeHandle), EntityHandle);
-
-			NotificationContext.ObserverManager.HandleElementsImpl(NotificationContext.ProcessingContext
-				, MakeArrayView(&AsEntityCollection, 1)
-				, {OpType, ObservedTypesOverlap}
-				, NotificationContext.ObserverManager.TagObservers[static_cast<uint8>(OpType)]);
-		}
-
-		TArray<const UScriptStruct*, TInlineAllocator<ObserverManager::Tweakables::InlineAllocatorElementsForOverlapTypes>> ObservedTypesOverlap;
-		FNotificationContext& NotificationContext;
-		EMassObservedOperation OpType;
-	};
-
 	struct FBufferedCreationNotificationExecutioner
 	{
 		explicit FBufferedCreationNotificationExecutioner(FNotificationContext& InNotificationContext)
@@ -304,8 +208,7 @@ void FMassObserverManager::DebugGatherUniqueProcessors(TArray<const UMassProcess
 #if WITH_MASSENTITY_DEBUG
 	for (int32 OperationIndex = 0; OperationIndex < static_cast<uint32>(EMassObservedOperation::MAX); ++OperationIndex)
 	{
-		FragmentObservers[OperationIndex].DebugAddUniqueProcessors(OutProcessors);
-		TagObservers[OperationIndex].DebugAddUniqueProcessors(OutProcessors);
+		ElementObservers[OperationIndex].DebugAddUniqueProcessors(OutProcessors);
 	}
 #endif // WITH_MASSENTITY_DEBUG
 }
@@ -324,8 +227,7 @@ void FMassObserverManager::Initialize()
 	using UE::Mass::ObserverManager::Private::AddRegisteredObserverProcessorInstances;
 	for (int i = 0; i < (int)EMassObservedOperation::MAX; ++i)
 	{
-		AddRegisteredObserverProcessorInstances(EntityManager, WorldExecutionFlags, *Owner, Registry.FragmentObserverMaps[i], ObservedFragments[i], FragmentObservers[i], TransientProcessorsMap);
-		AddRegisteredObserverProcessorInstances(EntityManager, WorldExecutionFlags, *Owner, Registry.TagObserverMaps[i], ObservedTags[i], TagObservers[i], TransientProcessorsMap);
+		AddRegisteredObserverProcessorInstances(EntityManager, WorldExecutionFlags, *Owner, Registry.ElementObserverMaps[i], ObservedElements[i], ElementObservers[i], TransientProcessorsMap);
 	}
 
 #if WITH_MASSENTITY_DEBUG
@@ -356,8 +258,7 @@ void FMassObserverManager::DeInitialize()
 
 	for (int32 i = 0; i < (int32)EMassObservedOperation::MAX; ++i)
 	{
-		(*FragmentObservers[i]).Empty();
-		(*TagObservers[i]).Empty();
+		(*ElementObservers[i]).Empty();
 	}
 }
 
@@ -417,7 +318,7 @@ bool FMassObserverManager::OnPreEntityDestroyed(const FMassArchetypeCompositionD
 	return OnCompositionChanged(Entity, ArchetypeComposition, EMassObservedOperation::DestroyEntity);
 }
 
-bool FMassObserverManager::OnCompositionChanged(FCollectionRefOrHandle&& EntityCollection, const FMassArchetypeCompositionDescriptor& CompositionDelta
+bool FMassObserverManager::OnCompositionChanged(FCollectionRefOrHandle&& EntityCollection, const FMassElementBitSet& CompositionDelta
 	, const EMassObservedOperation Operation, UE::Mass::FProcessingContext* ProcessingContext)
 {
 	using UE::Mass::ObserverManager::Tweakables::InlineAllocatorElementsForOverlapTypes;
@@ -450,12 +351,9 @@ bool FMassObserverManager::OnCompositionChanged(FCollectionRefOrHandle&& EntityC
 
 	const int32 OperationIndex = static_cast<int32>(Operation);
 
-	FMassFragmentBitSet FragmentOverlap = ObservedFragments[OperationIndex].GetOverlap(CompositionDelta.GetFragments());
-	FMassTagBitSet TagOverlap = ObservedTags[OperationIndex].GetOverlap(CompositionDelta.GetTags());
-	const bool bHasFragmentsOverlap = !FragmentOverlap.IsEmpty();
-	const bool bHasTagsOverlap = !TagOverlap.IsEmpty();
+	FMassElementBitSet OverlapBitSet = ObservedElements[OperationIndex].GetOverlap(CompositionDelta);
 
-	if (bHasFragmentsOverlap || bHasTagsOverlap)
+	if (OverlapBitSet.IsEmpty() == false)
 	{
 		if (LocksCount > 0)
 		{
@@ -468,15 +366,11 @@ bool FMassObserverManager::OnCompositionChanged(FCollectionRefOrHandle&& EntityC
 
 				if (EntityCollection.EntityHandle.IsSet())
 				{
-					ObserverLockRef->AddNotification(Operation, EntityCollection.EntityHandle
-						, bHasFragmentsOverlap, MoveTemp(FragmentOverlap)
-						, bHasTagsOverlap, MoveTemp(TagOverlap));
+					ObserverLockRef->AddNotification(Operation, EntityCollection.EntityHandle, MoveTemp(OverlapBitSet));
 				}
 				else
 				{
-					ObserverLockRef->AddNotification(Operation, EntityCollection.EntityCollection
-						, bHasFragmentsOverlap, MoveTemp(FragmentOverlap) 
-						, bHasTagsOverlap, MoveTemp(TagOverlap));
+					ObserverLockRef->AddNotification(Operation, EntityCollection.EntityCollection, MoveTemp(OverlapBitSet));
 				}
 			}
 			else
@@ -493,19 +387,7 @@ bool FMassObserverManager::OnCompositionChanged(FCollectionRefOrHandle&& EntityC
 					Entities.Emplace<UE::Mass::FEntityCollection>(EntityCollection.EntityCollection);
 				}
 
-				if (bHasFragmentsOverlap && bHasTagsOverlap)
-				{
-					FMassArchetypeCompositionDescriptor ChangedComposition(MoveTemp(FragmentOverlap), MoveTemp(TagOverlap), {}, {}, {});
-					ObserverLockRef->BufferedNotifications.Emplace(Operation, MoveTemp(ChangedComposition), MoveTemp(Entities));
-				}
-				else if (bHasFragmentsOverlap)
-				{
-					ObserverLockRef->BufferedNotifications.Emplace(Operation, MoveTemp(FragmentOverlap), MoveTemp(Entities));
-				}
-				else // bHasTagsOverlap
-				{
-					ObserverLockRef->BufferedNotifications.Emplace(Operation, MoveTemp(TagOverlap), MoveTemp(Entities));
-				}
+				ObserverLockRef->BufferedNotifications.Emplace(Operation, MoveTemp(OverlapBitSet), MoveTemp(Entities));
 			}
 		}
 		else
@@ -519,21 +401,9 @@ bool FMassObserverManager::OnCompositionChanged(FCollectionRefOrHandle&& EntityC
 					? new(&LocalContextBuffer) UE::Mass::FProcessingContext(EntityManager, /*DeltaSeconds=*/0.f, /*bFlushCommandBuffer=*/false)
 					: ProcessingContext;
 
-				if (bHasFragmentsOverlap)
-				{
-					FragmentOverlap.ExportTypes(ObservedTypesOverlap);
-
-					HandleElementsImpl(*LocalProcessingContext, {Collection}, {Operation, ObservedTypesOverlap}, FragmentObservers[OperationIndex]);
-				}
-
-				if (bHasTagsOverlap)
-				{
-					ObservedTypesOverlap.Reset();
-					TagOverlap.ExportTypes(ObservedTypesOverlap);
-
-					HandleElementsImpl(*LocalProcessingContext, {Collection}, {Operation, ObservedTypesOverlap}, TagObservers[OperationIndex]);
-				}
-
+				OverlapBitSet.ExportTypes(ObservedTypesOverlap);
+				HandleElementsImpl(*LocalProcessingContext, {Collection}, {Operation, ObservedTypesOverlap}, ElementObservers[OperationIndex]);
+				
 				if (ProcessingContext == nullptr)
 				{
 					LocalProcessingContext->~FProcessingContext();
@@ -568,38 +438,25 @@ bool FMassObserverManager::OnCollectionsCreatedImpl(UE::Mass::FProcessingContext
 	constexpr EMassObservedOperation Operation = EMassObservedOperation::CreateEntity;
 	constexpr int32 OperationIndex = static_cast<int32>(Operation);
 
-	FMassFragmentBitSet FragmentOverlap;
-	FMassTagBitSet TagOverlap;
+	FMassElementBitSet OverlapBitSet;
 
 	for (FMassArchetypeEntityCollection Collection : EntityCollections)
 	{
 		checkfSlow(Collection.IsUpToDate(), TEXT("Out-of-date FMassArchetypeEntityCollection used. Stored information is unreliable."));
 
 		const FMassArchetypeCompositionDescriptor& ArchetypeComposition = EntityManager.GetArchetypeComposition(Collection.GetArchetype());
-		FragmentOverlap += ArchetypeComposition.GetFragments();
-		TagOverlap += ArchetypeComposition.GetTags();
+		OverlapBitSet += ArchetypeComposition.GetElementsBitSet();
 	}
-	FragmentOverlap = ObservedFragments[OperationIndex].GetOverlap(FragmentOverlap);
-	TagOverlap = ObservedTags[OperationIndex].GetOverlap(TagOverlap);
 
-	const bool bHasFragmentsOverlap = !FragmentOverlap.IsEmpty();
-	const bool bHasTagsOverlap = !TagOverlap.IsEmpty();
-	if (bHasFragmentsOverlap || bHasTagsOverlap)
+	OverlapBitSet = ObservedElements[OperationIndex].GetOverlap(OverlapBitSet);
+	const bool bHasOverlap = !OverlapBitSet.IsEmpty();
+
+	if (bHasOverlap)
 	{
 		TArray<const UScriptStruct*, TInlineAllocator<InlineAllocatorElementsForOverlapTypes>> ObservedTypesOverlap;
 
-		if (bHasFragmentsOverlap)
-		{
-			FragmentOverlap.ExportTypes(ObservedTypesOverlap);
-			HandleElementsImpl(ProcessingContext, EntityCollections, {Operation, ObservedTypesOverlap}, FragmentObservers[OperationIndex]);
-		}
-
-		if (bHasTagsOverlap)
-		{
-			ObservedTypesOverlap.Reset();
-			TagOverlap.ExportTypes(ObservedTypesOverlap);
-			HandleElementsImpl(ProcessingContext, EntityCollections, {Operation, ObservedTypesOverlap}, TagObservers[OperationIndex]);
-		}
+		OverlapBitSet.ExportTypes(ObservedTypesOverlap);
+		HandleElementsImpl(ProcessingContext, EntityCollections, {Operation, ObservedTypesOverlap}, ElementObservers[OperationIndex]);
 
 		return true;
 	}
@@ -629,29 +486,15 @@ void FMassObserverManager::HandleElementsImpl(UE::Mass::FProcessingContext& Proc
 
 void FMassObserverManager::AddObserverInstance(TNotNull<const UScriptStruct*> ElementType, const EMassObservedOperationFlags OperationFlags, TNotNull<UMassProcessor*> ObserverProcessor)
 {
-	const bool bIsFragment = UE::Mass::IsA<FMassFragment>(ElementType);
-	checkSlow(bIsFragment || UE::Mass::IsA<FMassTag>(ElementType));
-
-	FMassObserversMap* ObserversMap = bIsFragment
-		? FragmentObservers
-		: TagObservers;
-
 	bool bInitializeCalled = false;
 
 	for (uint8 OperationIndex = 0; OperationIndex < static_cast<uint8>(EMassObservedOperation::MAX); ++OperationIndex)
 	{
 		if (!!(OperationFlags & static_cast<EMassObservedOperationFlags>(1 << OperationIndex)))
 		{
-			if (bIsFragment)
-			{
-				ObservedFragments[OperationIndex].Add(*ElementType);
-			}
-			else
-			{
-				ObservedTags[OperationIndex].Add(*ElementType);
-			}
+			ObservedElements[OperationIndex].Add(*ElementType);
 
-			FMassRuntimePipeline& Pipeline = (*ObserversMap[OperationIndex]).FindOrAdd(ElementType);
+			FMassRuntimePipeline& Pipeline = (*ElementObservers[OperationIndex]).FindOrAdd(ElementType);
 			if (Pipeline.AppendUniqueProcessor(*ObserverProcessor) && !bInitializeCalled)
 			{
 				if (UObject* Owner = EntityManager.GetOwner())
@@ -687,36 +530,18 @@ void FMassObserverManager::AddObserverInstance(TNotNull<UMassObserverProcessor*>
 
 void FMassObserverManager::RemoveObserverInstance(TNotNull<const UScriptStruct*> ElementType, const EMassObservedOperationFlags OperationFlags, TNotNull<UMassProcessor*> ObserverProcessor)
 {
-	const bool bIsFragment = UE::Mass::IsA<FMassFragment>(ElementType);
-
-	if (!ensure(bIsFragment || UE::Mass::IsA<FMassTag>(ElementType)))
-	{
-		return;
-	}
-
-	FMassObserversMap* ObserversMap = bIsFragment
-		? FragmentObservers
-		: TagObservers;
-
 	for (uint8 OperationIndex = 0; OperationIndex < static_cast<uint8>(EMassObservedOperation::MAX); ++OperationIndex)
 	{
 		if (!!(OperationFlags & static_cast<EMassObservedOperationFlags>(1 << OperationIndex)))
 		{
-			FMassRuntimePipeline* Pipeline = (*ObserversMap[OperationIndex]).Find(ElementType);
+			FMassRuntimePipeline* Pipeline = (*ElementObservers[OperationIndex]).Find(ElementType);
 			if (ensureMsgf(Pipeline, TEXT("Trying to remove an observer for a fragment/tag that does not seem to be observed.")))
 			{
 				Pipeline->RemoveProcessor(*ObserverProcessor);
 				if (Pipeline->Num() == 0)
 				{
-					(*ObserversMap[OperationIndex]).Remove(ElementType);
-					if (bIsFragment)
-					{
-						ObservedFragments[OperationIndex].Remove(*ElementType);
-					}
-					else
-					{
-						ObservedTags[OperationIndex].Remove(*ElementType);
-					}
+					(*ElementObservers[OperationIndex]).Remove(ElementType);
+					ObservedElements[OperationIndex].Remove(*ElementType);
 				}
 			}
 		}
@@ -821,14 +646,14 @@ void FMassObserverManager::OnModulePackagesUnloaded(TConstArrayView<UPackage*> P
 
 	for (int32 OperationIndex = 0; OperationIndex < static_cast<uint32>(EMassObservedOperation::MAX); ++OperationIndex)
 	{
-		ProcessObserversMap(FragmentObservers[OperationIndex]);
-		ProcessObserversMap(TagObservers[OperationIndex]);
+		ProcessObserversMap(ElementObservers[OperationIndex]);
 	}
 }
 
 void FMassObserverManager::ResumeExecution(FObserverLock& LockBeingReleased)
 {
 	using namespace UE::Mass::ObserverManager;
+	using UE::Mass::ObserverManager::Tweakables::InlineAllocatorElementsForOverlapTypes;
 
 	ensureMsgf(LocksCount == 0, TEXT("We only expect this function to be called if all locks are released."));
 #if WITH_MASSENTITY_DEBUG
@@ -844,9 +669,8 @@ void FMassObserverManager::ResumeExecution(FObserverLock& LockBeingReleased)
 		UE::Mass::FProcessingContext ProcessingContext(EntityManager);
 
 		FNotificationContext NotificationContext{*this, ProcessingContext};
-
-		FBufferedNotificationExecutioner RegularOperationNotifier(NotificationContext);
 		FBufferedCreationNotificationExecutioner CreationOpExecutioner(NotificationContext);
+		TArray<const UScriptStruct*, TInlineAllocator<InlineAllocatorElementsForOverlapTypes>> ObservedTypesOverlap;
 
 		for (FBufferedNotification& Op : LockBeingReleased.BufferedNotifications)
 		{
@@ -856,8 +680,29 @@ void FMassObserverManager::ResumeExecution(FObserverLock& LockBeingReleased)
 			}
 			else
 			{
-				RegularOperationNotifier.OpType = Op.Type;
-				Visit(RegularOperationNotifier, Op.CompositionChange, Op.AffectedEntities);
+				if (const UE::Mass::FEntityCollection* Entities = Op.AffectedEntities.TryGet<UE::Mass::FEntityCollection>())
+				{
+					ObservedTypesOverlap.Reset();
+					Op.CompositionChange.ExportTypes(ObservedTypesOverlap);
+					FMassObserverManager::HandleElementsImpl(NotificationContext.ProcessingContext
+						, Entities->GetUpToDatePerArchetypeCollections(NotificationContext.ObserverManager.EntityManager)
+						, {Op.Type, ObservedTypesOverlap}
+						, NotificationContext.ObserverManager.ElementObservers[static_cast<uint8>(Op.Type)]);
+				}
+				else 
+				{
+					const FMassEntityHandle EntityHandle = Op.AffectedEntities.Get<FMassEntityHandle>();
+					ObservedTypesOverlap.Reset();
+					Op.CompositionChange.ExportTypes(ObservedTypesOverlap);
+					
+					FMassArchetypeHandle ArchetypeHandle = NotificationContext.ObserverManager.GetEntityManager().GetArchetypeForEntity(EntityHandle);
+					FMassArchetypeEntityCollection AsEntityCollection(MoveTemp(ArchetypeHandle), EntityHandle);
+
+					NotificationContext.ObserverManager.HandleElementsImpl(NotificationContext.ProcessingContext
+						, MakeArrayView(&AsEntityCollection, 1)
+						, {Op.Type, ObservedTypesOverlap}
+						, NotificationContext.ObserverManager.ElementObservers[static_cast<uint8>(Op.Type)]);
+				}
 			}
 		}
 #if WITH_MASSENTITY_DEBUG
@@ -978,9 +823,7 @@ void FMassObserverManager::HandleElementsImpl(UE::Mass::FProcessingContext& Proc
 	check(bIsFragment || UE::Mass::IsA<FMassTag>(ObservedTypes[0]));
 
 	using UE::Mass::ObserverManager::FDeprecationHelper;
-	const EMassObservedOperation Operation = bIsFragment
-		? FDeprecationHelper::DetermineOperationType(&HandlersContainer, ProcessingContext.GetEntityManager()->GetObserverManager().FragmentObservers)
-		: FDeprecationHelper::DetermineOperationType(&HandlersContainer, ProcessingContext.GetEntityManager()->GetObserverManager().TagObservers);
+	const EMassObservedOperation Operation = FDeprecationHelper::DetermineOperationType(&HandlersContainer, ProcessingContext.GetEntityManager()->GetObserverManager().ElementObservers);
 
 	checkf(Operation < EMassObservedOperation::MAX, TEXT("Unable to determine the legacy operation type"));
 
@@ -998,4 +841,4 @@ PRAGMA_ENABLE_DEPRECATION_WARNINGS
 bool FMassObserverManager::OnCollectionsCreatedImpl(UE::Mass::FProcessingContext& ProcessingContext, TArray<FMassArchetypeEntityCollection>&& EntityCollections)
 {
 	return OnCollectionsCreatedImpl(ProcessingContext, MakeArrayView(EntityCollections));
-}
+}
\ No newline at end of file
diff --git a/Engine/Source/Runtime/MassEntity/Private/MassObserverRegistry.cpp b/Engine/Source/Runtime/MassEntity/Private/MassObserverRegistry.cpp
--- a/Engine/Source/Runtime/MassEntity/Private/MassObserverRegistry.cpp
+++ b/Engine/Source/Runtime/MassEntity/Private/MassObserverRegistry.cpp
@@ -19,26 +19,20 @@ UMassObserverRegistry::UMassObserverRegistry()
 	}
 }
 
-void UMassObserverRegistry::RegisterObserver(TNotNull<const UScriptStruct*> ObservedType, const uint8 OperationFlags, TSubclassOf<UMassProcessor> ObserverClass)
+void UMassObserverRegistry::RegisterObserver(TNotNull<const UScriptStruct*> ObservedType, const uint8 OperationFlags, const TSubclassOf<UMassProcessor>& ObserverClass)
 {
 	check(ObserverClass);
-	const bool bIsFragment = UE::Mass::IsA<FMassFragment>(ObservedType);
-	checkSlow(bIsFragment || UE::Mass::IsA<FMassTag>(ObservedType));
-
-	FObserverClassesMap* ObserversMap = bIsFragment
-		? FragmentObserverMaps
-		: TagObserverMaps; 
 
 	for (uint8 OperationIndex = 0; OperationIndex < static_cast<uint8>(EMassObservedOperation::MAX); ++OperationIndex)
 	{
 		if ((OperationFlags & (1 << OperationIndex)))
 		{
-			ObserversMap[OperationIndex].FindOrAdd(ObservedType).AddUnique(FSoftClassPath(ObserverClass));
+			ElementObserverMaps[OperationIndex].FindOrAdd(ObservedType).AddUnique(FSoftClassPath(ObserverClass));
 		}
 	}
 }
 
-void UMassObserverRegistry::RegisterObserver(const UScriptStruct& ObservedType, const EMassObservedOperation Operation, TSubclassOf<UMassProcessor> ObserverClass)
+void UMassObserverRegistry::RegisterObserver(const UScriptStruct& ObservedType, const EMassObservedOperation Operation, const TSubclassOf<UMassProcessor>& ObserverClass)
 {
 	RegisterObserver(&ObservedType, static_cast<uint8>(1 << static_cast<uint8>(Operation)), ObserverClass);
 }
@@ -79,7 +73,6 @@ void UMassObserverRegistry::OnModulePackagesUnloaded(TConstArrayView<UPackage*>
 
 	for (uint8 OperationIndex = 0; OperationIndex < static_cast<uint8>(EMassObservedOperation::MAX); ++OperationIndex)
 	{
-		ProcessObserversMap(FragmentObserverMaps[OperationIndex]);
-		ProcessObserversMap(TagObserverMaps[OperationIndex]);
+		ProcessObserversMap(ElementObserverMaps[OperationIndex]);
 	}
 }
diff --git a/Engine/Source/Runtime/MassEntity/Private/MassRequirementAccessDetector.cpp b/Engine/Source/Runtime/MassEntity/Private/MassRequirementAccessDetector.cpp
--- a/Engine/Source/Runtime/MassEntity/Private/MassRequirementAccessDetector.cpp
+++ b/Engine/Source/Runtime/MassEntity/Private/MassRequirementAccessDetector.cpp
@@ -17,10 +17,6 @@ namespace UE::Mass::Private
 void FMassRequirementAccessDetector::Initialize()
 {
 	check(IsInGameThread());
-	AddDetectors(FMassFragmentBitSet::FStructTrackerWrapper::StructTracker);
-	AddDetectors(FMassChunkFragmentBitSet::FStructTrackerWrapper::StructTracker);
-	AddDetectors(FMassSharedFragmentBitSet::FStructTrackerWrapper::StructTracker);
-	AddDetectors(FMassConstSharedFragmentBitSet::FStructTrackerWrapper::StructTracker);
 	AddDetectors(FMassExternalSubsystemBitSet::FStructTrackerWrapper::StructTracker);
 }
 
@@ -29,8 +25,10 @@ void FMassRequirementAccessDetector::AddDetectors(const FStructTracker& StructTr
 	TConstArrayView<TWeakObjectPtr<const UStruct>> Types = StructTracker.DebugGetAllStructTypes<UStruct>();
 	for (TWeakObjectPtr<const UStruct> Type : Types)
 	{
-		check(Type.Get());
-		Detectors.Add(Type.Get(), MakeShareable(new FRWAccessDetector()));
+		if (Type.Get())
+		{
+			Detectors.Add(Type.Get(), MakeShareable(new FRWAccessDetector()));
+		}
 	}
 }
 
diff --git a/Engine/Source/Runtime/MassEntity/Private/MassRequirements.cpp b/Engine/Source/Runtime/MassEntity/Private/MassRequirements.cpp
--- a/Engine/Source/Runtime/MassEntity/Private/MassRequirements.cpp
+++ b/Engine/Source/Runtime/MassEntity/Private/MassRequirements.cpp
@@ -247,14 +247,10 @@ bool FMassFragmentRequirements::IsEmpty() const
 	return !bHasPositiveRequirements && !bHasNegativeRequirements && !bHasOptionalRequirements;
 }
 
-bool FMassFragmentRequirements::DoesMatchAnyOptionals(const FMassArchetypeCompositionDescriptor& ArchetypeComposition) const
+bool FMassFragmentRequirements::DoesMatchAnyOptionals(const FMassElementBitSet& ArchetypeCompositionBitSet) const
 {
 	return bHasOptionalRequirements
-		&& (ArchetypeComposition.GetFragments().HasAny(RequiredOptionalFragments)
-			|| ArchetypeComposition.GetTags().HasAny(RequiredOptionalTags)
-			|| ArchetypeComposition.GetChunkFragments().HasAny(RequiredOptionalChunkFragments)
-			|| ArchetypeComposition.GetSharedFragments().HasAny(RequiredOptionalSharedFragments)
-			|| ArchetypeComposition.GetConstSharedFragments().HasAny(RequiredOptionalConstSharedFragments));
+		&& ArchetypeCompositionBitSet.HasAny(RequiredOptionalFragments | RequiredOptionalTags | RequiredOptionalChunkFragments | RequiredOptionalSharedFragments | RequiredOptionalConstSharedFragments);
 }
 
 bool FMassFragmentRequirements::DoesArchetypeMatchRequirements(const FMassArchetypeHandle& ArchetypeHandle) const
@@ -263,35 +259,29 @@ bool FMassFragmentRequirements::DoesArchetypeMatchRequirements(const FMassArchet
 	const FMassArchetypeData* Archetype = FMassArchetypeHelper::ArchetypeDataFromHandle(ArchetypeHandle);
 	CA_ASSUME(Archetype);
 
-	return DoesArchetypeMatchRequirements(Archetype->GetCompositionDescriptor());
+	return DoesArchetypeMatchRequirements(Archetype->GetCompositionBitSet());
 }
 	
-bool FMassFragmentRequirements::DoesArchetypeMatchRequirements(const FMassArchetypeCompositionDescriptor& ArchetypeComposition) const
+bool FMassFragmentRequirements::DoesArchetypeMatchRequirements(const FMassElementBitSet& ArchetypeCompositionBitSet) const
 {
 	CacheProperties();
 
 	const bool bPassNegativeFilter = bHasNegativeRequirements == false
-		|| (ArchetypeComposition.GetFragments().HasNone(RequiredNoneFragments)
-			&& ArchetypeComposition.GetTags().HasNone(RequiredNoneTags)
-			&& ArchetypeComposition.GetChunkFragments().HasNone(RequiredNoneChunkFragments)
-			&& ArchetypeComposition.GetSharedFragments().HasNone(RequiredNoneSharedFragments)
-			&& ArchetypeComposition.GetConstSharedFragments().HasNone(RequiredNoneConstSharedFragments));
-	
+		|| ArchetypeCompositionBitSet.HasNone(RequiredNoneFragments | RequiredNoneTags | RequiredNoneChunkFragments | RequiredNoneSharedFragments | RequiredNoneConstSharedFragments);
+
 	if (bPassNegativeFilter)
 	{
 		if (bHasPositiveRequirements)
 		{
-			return ArchetypeComposition.GetFragments().HasAll(RequiredAllFragments)
-				&& (RequiredAnyFragments.IsEmpty() || ArchetypeComposition.GetFragments().HasAny(RequiredAnyFragments))
-				&& ArchetypeComposition.GetTags().HasAll(RequiredAllTags)
-				&& (RequiredAnyTags.IsEmpty() || ArchetypeComposition.GetTags().HasAny(RequiredAnyTags))
-				&& ArchetypeComposition.GetChunkFragments().HasAll(RequiredAllChunkFragments)
-				&& ArchetypeComposition.GetSharedFragments().HasAll(RequiredAllSharedFragments)
-				&& ArchetypeComposition.GetConstSharedFragments().HasAll(RequiredAllConstSharedFragments);
+			return ArchetypeCompositionBitSet.HasAll(RequiredAllFragments)
+				&& (RequiredAnyFragments.IsEmpty() || ArchetypeCompositionBitSet.HasAny(RequiredAnyFragments))
+				&& ArchetypeCompositionBitSet.HasAll(RequiredAllTags)
+				&& (RequiredAnyTags.IsEmpty() || ArchetypeCompositionBitSet.HasAny(RequiredAnyTags))
+				&& ArchetypeCompositionBitSet.HasAll(RequiredAllChunkFragments | RequiredAllSharedFragments | RequiredAllConstSharedFragments);
 		}
 		else if (bHasOptionalRequirements)
 		{
-			return DoesMatchAnyOptionals(ArchetypeComposition);
+			return DoesMatchAnyOptionals(ArchetypeCompositionBitSet);
 		}
 		// else - it's fine, we passed all the filters that have been set up
 		return true;
diff --git a/Engine/Source/Runtime/MassEntity/Public/MassDebugger.h b/Engine/Source/Runtime/MassEntity/Public/MassDebugger.h
--- a/Engine/Source/Runtime/MassEntity/Public/MassDebugger.h
+++ b/Engine/Source/Runtime/MassEntity/Public/MassDebugger.h
@@ -187,7 +187,7 @@ struct FMassDebugger
 
 	static MASSENTITY_API void ForEachArchetype(const FMassEntityManager& EntityManager, const UE::Mass::Debug::FArchetypeFunction& Function);
 	static MASSENTITY_API TArray<FMassArchetypeHandle> GetAllArchetypes(const FMassEntityManager& EntityManager);
-	static MASSENTITY_API const FMassArchetypeCompositionDescriptor& GetArchetypeComposition(const FMassArchetypeHandle& ArchetypeHandle);
+	static MASSENTITY_API FMassArchetypeCompositionDescriptor GetArchetypeComposition(const FMassArchetypeHandle& ArchetypeHandle);
 
 	static MASSENTITY_API uint64 GetArchetypeTraceID(const FMassArchetypeData& ArchetypeData);
 	static MASSENTITY_API uint64 GetArchetypeTraceID(const FMassArchetypeHandle& ArchetypeHandle);
@@ -209,7 +209,11 @@ struct FMassDebugger
 	static MASSENTITY_API FString GetSingleRequirementDescription(const FMassFragmentRequirementDescription& Requirement);
 	static MASSENTITY_API FString GetRequirementsDescription(const FMassFragmentRequirements& Requirements);
 	static MASSENTITY_API FString GetArchetypeRequirementCompatibilityDescription(const FMassFragmentRequirements& Requirements, const FMassArchetypeHandle& ArchetypeHandle);
-	static MASSENTITY_API FString GetArchetypeRequirementCompatibilityDescription(const FMassFragmentRequirements& Requirements, const FMassArchetypeCompositionDescriptor& ArchetypeComposition);
+	static MASSENTITY_API FString GetArchetypeRequirementCompatibilityDescription(const FMassFragmentRequirements& Requirements, const FMassElementBitSet& ArchetypeCompositionBitSet);
+	static FString GetArchetypeRequirementCompatibilityDescription(const FMassFragmentRequirements& Requirements, const FMassArchetypeCompositionDescriptor& ArchetypeComposition)
+	{
+		return GetArchetypeRequirementCompatibilityDescription(Requirements, ArchetypeComposition.GetElementsBitSet());
+	}
 
 	static MASSENTITY_API void OutputArchetypeDescription(FOutputDevice& Ar, const FMassArchetypeHandle& Archetype);
 	static MASSENTITY_API void OutputEntityDescription(FOutputDevice& Ar, const FMassEntityManager& EntityManager, const int32 EntityIndex, const TCHAR* InPrefix = TEXT(""));
diff --git a/Engine/Source/Runtime/MassEntity/Public/MassElement.h b/Engine/Source/Runtime/MassEntity/Public/MassElement.h
new file mode 100644
--- /dev/null
+++ b/Engine/Source/Runtime/MassEntity/Public/MassElement.h
@@ -0,0 +1,235 @@
+// Copyright Epic Games, Inc. All Rights Reserved.
+
+#pragma once
+
+#include "StructUtils/StructTypeBitSet.h"
+#include "MassEntityElementTypes.h"
+#include "Subsystems/Subsystem.h"
+
+#define UE_API MASSENTITY_API
+
+DECLARE_CLASSTYPEBITSET_EXPORTED(UE_API, FMassExternalSubsystemBitSet, USubsystem);
+
+namespace UE::Mass
+{
+	DECLARE_STRUCTTYPEBITSET_EXPORTED(UE_API, FElementBitSetBase, FMassElement);
+
+	enum class EElementType : uint8
+	{
+		Fragment,
+		Tag,
+		ChunkFragment,
+		SharedFragment,
+		ConstSharedFragment,
+		MAX
+	};
+
+	template<typename TElementType>
+	struct TElementTraits
+	{};
+
+	template<typename TBitSetType>
+	struct TBitSetTraits
+	{};
+
+	struct FElementBitSet : public FElementBitSetBase
+	{
+		using FElementBitSetBase::FElementBitSetBase;
+		using FElementBitSetBase::operator=;
+		using FElementBitSetBase::HasAny;
+		using FElementBitSetBase::HasAll;;
+
+		FElementBitSet(const FElementBitSetBase& OtherBitSet)
+			: FElementBitSetBase(OtherBitSet)
+		{
+		}
+
+		template<typename TBitSet>
+		TBitSet Get() const;
+
+		template<typename TElementType>
+		bool HasAny() const;
+
+		template<typename TElementType>
+		bool HasAll() const;
+
+		// @todo this is a temporary construct, until shared fragments get removed as a concept
+		bool HasAnySharedElements() const;
+		
+		bool ContainsOnlyTagsAndFragments() const;
+		FElementBitSet GetFragmentsAndTags() const;
+		FElementBitSet GetSharedFragments() const;
+	
+		static EElementType DetermineElementType(const TNotNull<const UStruct*> ElementType)
+		{
+			if (UE::Mass::IsA<FMassFragment>(ElementType))
+			{
+				return EElementType::Fragment;
+			}
+			else if (UE::Mass::IsA<FMassTag>(ElementType))
+			{
+				return EElementType::Tag;
+			}
+			else if (UE::Mass::IsA<FMassSharedFragment>(ElementType))
+			{
+				return EElementType::SharedFragment;
+			}
+			else if (UE::Mass::IsA<FMassConstSharedFragment>(ElementType))
+			{
+				return EElementType::ConstSharedFragment;
+			}
+			else if (UE::Mass::IsA<FMassChunkFragment>(ElementType))
+			{
+				return EElementType::ChunkFragment;
+			}
+
+			checkf(false, TEXT("%hs: Unhandled element type %s"), __FUNCTION__, *ElementType->GetName())
+			return EElementType::MAX;
+		}
+
+		static void OnTypeRegistered(const TNotNull<const UStruct*> Type, int32 Index)
+		{
+			const EElementType ElementType = DetermineElementType(Type);
+			GetTypeBitArray(ElementType).AddAtIndex(Index);
+			bTypeBitSetsDirty = true;
+		}
+
+		static void OnModuleInitialized()
+		{
+			OnTypeRegisteredDelegateHandle = FElementBitSetBaseStructTrackerWrapper::StructTracker.OnTypeRegistered.AddStatic(&FElementBitSet::OnTypeRegistered);
+		}
+
+		static void OnModuleShutdown()
+		{
+			FElementBitSetBaseStructTrackerWrapper::StructTracker.OnTypeRegistered.Remove(OnTypeRegisteredDelegateHandle);
+		}
+
+		UE_FORCEINLINE_HINT static void UpdateTypeBitSets()
+		{
+			if (bTypeBitSetsDirty)
+			{
+				AllSharedFragmentsBitSet = GetTypeBitArray(EElementType::ConstSharedFragment) | GetTypeBitArray(EElementType::SharedFragment);
+				AllFragmentsAndTagsBitSet = (GetTypeBitArray(EElementType::Fragment) | GetTypeBitArray(EElementType::Tag));
+				bTypeBitSetsDirty = false;
+			}
+		}
+
+		static const FElementBitSet& GetAllFragmentsAndTags()
+		{
+			UpdateTypeBitSets();
+			return AllFragmentsAndTagsBitSet;
+		}
+
+		static const FElementBitSet& GetAllSharedFragments()
+		{
+			UpdateTypeBitSets();
+			return AllSharedFragmentsBitSet;
+		}
+
+	protected:
+		UE_API static TStaticArray<FElementBitSetBase, static_cast<uint8>(EElementType::MAX)> ElementTypeBitArrays;
+		UE_API static FElementBitSet AllSharedFragmentsBitSet;
+		UE_API static FElementBitSet AllFragmentsAndTagsBitSet;
+		UE_API static bool bTypeBitSetsDirty;
+		static FDelegateHandle OnTypeRegisteredDelegateHandle;
+
+	public:
+		/**
+		 * Fetches FElementBitSetBase containing all known specific element types
+		 */
+		static FElementBitSetBase& GetTypeBitArray(const EElementType ElementType)
+		{	
+			return ElementTypeBitArrays[static_cast<uint8>(ElementType)];
+		}
+	};
+
+	template<typename TBitSet>
+	inline TBitSet FElementBitSet::Get() const
+	{
+		return GetTypeBitArray(TBitSetTraits<TBitSet>::ElementType) & *this;
+	}
+
+	template<typename TElementType>
+	bool FElementBitSet::HasAny() const
+	{
+		return HasAny(GetTypeBitArray(TElementTraits<TElementType>::ElementType));
+	}
+
+	template<typename TElementType>
+	bool FElementBitSet::HasAll() const
+	{
+		return HasAll(GetTypeBitArray(TElementTraits<TElementType>::ElementType));
+	}
+
+	inline bool FElementBitSet::HasAnySharedElements() const
+	{
+		return FElementBitSetBase::HasAny(GetTypeBitArray(EElementType::SharedFragment)) | FElementBitSetBase::HasAny(GetTypeBitArray(EElementType::ConstSharedFragment));
+	}
+
+	inline bool FElementBitSet::ContainsOnlyTagsAndFragments() const
+	{
+		return GetFragmentsAndTags().IsEquivalent(*this);
+	}
+
+	inline FElementBitSet FElementBitSet::GetFragmentsAndTags() const
+	{
+		return (*this) & GetAllFragmentsAndTags();
+	}
+
+	inline FElementBitSet FElementBitSet::GetSharedFragments() const
+	{
+		return (*this) & GetAllSharedFragments();
+	}
+
+	//-----------------------------------------------------------------------------
+	// Helpers to avoid redundant operations when dealing with specialized element bitsets
+	//-----------------------------------------------------------------------------
+	/** @return whether BitSet indicates only shared fragment types (both const and mutable) */
+	inline bool DoesContainOnlySharedFragments(const FElementBitSet& BitSet)
+	{
+		return FElementBitSet::GetAllSharedFragments().HasAll(BitSet);
+	}
+
+	/** @return whether shared fragment type information stored by both bitsets is equivalent */
+	inline bool DoContainEquivalentSharedFragments(const FElementBitSet& BitSetA, const FElementBitSet& BitSetB)
+	{
+		return BitSetA.IsEquivalentFiltered(BitSetB, FElementBitSet::GetAllSharedFragments());
+	}
+
+
+	#define DECLARE_NEWTYPEBITSET(ContainerTypeName, BaseType, ElementTypeEnumValue) struct ContainerTypeName : public FElementBitSetBase \
+	{ \
+		using FElementType = BaseType; \
+		using FElementBitSetBase::FElementBitSetBase; \
+		using FElementBitSetBase::operator=; \
+		ContainerTypeName(const FElementBitSetBase& OtherBitSet) \
+			: FElementBitSetBase(OtherBitSet) \
+		{} \
+	}; \
+	template<> struct TElementTraits<BaseType> \
+	{ \
+		static constexpr EElementType ElementType = ElementTypeEnumValue; \
+	}; \
+	template<> struct TBitSetTraits<ContainerTypeName> \
+	{ \
+		static constexpr EElementType ElementType = ElementTypeEnumValue; \
+	}
+
+
+	DECLARE_NEWTYPEBITSET(FFragmentBitSet, FMassFragment, EElementType::Fragment);
+	DECLARE_NEWTYPEBITSET(FTagBitSet, FMassTag, EElementType::Tag);
+	DECLARE_NEWTYPEBITSET(FChunkFragmentBitSet, FMassChunkFragment, EElementType::ChunkFragment);
+	DECLARE_NEWTYPEBITSET(FSharedFragmentBitSet, FMassSharedFragment, EElementType::SharedFragment);
+	DECLARE_NEWTYPEBITSET(FConstSharedFragmentBitSet, FMassConstSharedFragment, EElementType::ConstSharedFragment);
+
+	#undef DECLARE_NEWTYPEBITSET
+}
+
+using FMassElementBitSet = UE::Mass::FElementBitSet;
+using FMassFragmentBitSet = UE::Mass::FFragmentBitSet;
+using FMassTagBitSet = UE::Mass::FTagBitSet;
+using FMassChunkFragmentBitSet = UE::Mass::FChunkFragmentBitSet;
+using FMassSharedFragmentBitSet = UE::Mass::FSharedFragmentBitSet;
+using FMassConstSharedFragmentBitSet = UE::Mass::FConstSharedFragmentBitSet;
+
+#undef UE_API
diff --git a/Engine/Source/Runtime/MassEntity/Public/MassEntityBuilder.h b/Engine/Source/Runtime/MassEntity/Public/MassEntityBuilder.h
--- a/Engine/Source/Runtime/MassEntity/Public/MassEntityBuilder.h
+++ b/Engine/Source/Runtime/MassEntity/Public/MassEntityBuilder.h
@@ -300,12 +300,6 @@ private:
 		return ConstSharedFragments; 
 	}
 
-	template<typename T>
-	auto& GetBitSetBuilder()
-	{
-		return Composition.GetContainer<T>();
-	}
-
 	/** Releases reserved handle if it has not been committed yet */
 	void ConditionallyReleaseEntityHandle();
 
@@ -351,7 +345,7 @@ struct FScopedEntityBuilder : FEntityBuilder
 template<CTag T>
 FEntityBuilder& FEntityBuilder::Add()
 {
-	Composition.GetTags().Add<T>();
+	Composition.Add<T>();
 	State = EState::ReadyToCommit;
 	CachedArchetypeHandle = FMassArchetypeHandle();
 	return *this;
@@ -360,7 +354,7 @@ FEntityBuilder& FEntityBuilder::Add()
 template<CChunkFragment T>
 FEntityBuilder& FEntityBuilder::Add()
 {
-	Composition.GetChunkFragments().Add<T>();
+	Composition.Add<T>();
 	State = EState::ReadyToCommit;
 	CachedArchetypeHandle = FMassArchetypeHandle();
 	return *this;
@@ -376,13 +370,11 @@ FEntityBuilder& FEntityBuilder::Add(TArgs&&... InArgs) requires (CElement<T> &&
 template<typename T, typename... TArgs>
 T& FEntityBuilder::Add_GetRef(TArgs&&... InArgs) requires (CElement<T> && !(CTag<T> || CChunkFragment<T>))
 {
-	using FElementType = UE::Mass::TElementType<T>;
-
 	State = EState::ReadyToCommit;
 
-	if (ensureMsgf(GetBitSetBuilder<FElementType>().template Contains<T>() == false, TEXT("Element of type %s has already been added"), *T::StaticStruct()->GetName()))
+	if (ensureMsgf(Composition.GetElementsBitSet().Contains<T>() == false, TEXT("Element of type %s has already been added"), *T::StaticStruct()->GetName()))
 	{
-		GetBitSetBuilder<FElementType>().template Add<T>();
+		Composition.GetElementsBitSet().Add<T>();
 		CachedArchetypeHandle = FMassArchetypeHandle();
 
 		return GetInstancedStructContainerInternal<T>().Add_GetRef(FInstancedStruct::Make<T>(InArgs...)).template GetMutable<T>();
@@ -395,12 +387,10 @@ T& FEntityBuilder::Add_GetRef(TArgs&&... InArgs) requires (CElement<T> && !(CTag
 template<typename T, typename... TArgs>
 T& FEntityBuilder::GetOrCreate(TArgs&&... InArgs) requires (CElement<T> && !(CTag<T> || CChunkFragment<T>))
 {
-	using FElementType = UE::Mass::TElementType<T>;
-
 	State = EState::ReadyToCommit;
 
 	FInstancedStruct* ElementFound = nullptr;
-	if (GetBitSetBuilder<FElementType>().template Contains<T>())
+	if (Composition.GetElementsBitSet().Contains<T>())
 	{
 		// replace
 		ElementFound = GetInstancedStructContainerInternal<T>().FindByPredicate(FStructInstanceFindingPredicate{T::StaticStruct()});
@@ -410,7 +400,7 @@ T& FEntityBuilder::GetOrCreate(TArgs&&... InArgs) requires (CElement<T> && !(CTa
 	}
 	else
 	{
-		GetBitSetBuilder<FElementType>().template Add<T>();
+		Composition.GetElementsBitSet().Add<T>();
 		CachedArchetypeHandle = FMassArchetypeHandle();
 
 		ElementFound = &GetInstancedStructContainerInternal<T>().Add_GetRef(FInstancedStruct::Make<T>(InArgs...));
@@ -422,9 +412,7 @@ T& FEntityBuilder::GetOrCreate(TArgs&&... InArgs) requires (CElement<T> && !(CTa
 template<typename T>
 T* FEntityBuilder::Find() requires (CElement<T> && !(CTag<T> || CChunkFragment<T>))
 {
-	using FElementType = UE::Mass::TElementType<T>;
-
-	if (GetBitSetBuilder<FElementType>().template Contains<T>())
+	if (Composition.GetElementsBitSet().Contains<T>())
 	{
 		FInstancedStruct* ElementFound = GetInstancedStructContainerInternal<T>().FindByPredicate(FStructInstanceFindingPredicate{T::StaticStruct()});
 
diff --git a/Engine/Source/Runtime/MassEntity/Public/MassEntityElementTypes.h b/Engine/Source/Runtime/MassEntity/Public/MassEntityElementTypes.h
--- a/Engine/Source/Runtime/MassEntity/Public/MassEntityElementTypes.h
+++ b/Engine/Source/Runtime/MassEntity/Public/MassEntityElementTypes.h
@@ -4,9 +4,15 @@
 
 #include "MassEntityElementTypes.generated.h"
 
+USTRUCT()
+struct FMassElement
+{
+	GENERATED_BODY()
+};
+
 // This is the base class for all lightweight fragments
 USTRUCT()
-struct FMassFragment
+struct FMassFragment : public FMassElement
 {
 	GENERATED_BODY()
 };
@@ -20,25 +26,25 @@ struct FMassFragment
 // This is the base class for types that will only be tested for presence/absence, i.e. Tags.
 // Subclasses should never contain any member properties.
 USTRUCT()
-struct FMassTag
+struct FMassTag : public FMassElement
 {
 	GENERATED_BODY()
 };
 
 USTRUCT()
-struct FMassChunkFragment
+struct FMassChunkFragment : public FMassElement
 {
 	GENERATED_BODY()
 };
 
 USTRUCT()
-struct FMassSharedFragment
+struct FMassSharedFragment : public FMassElement
 {
 	GENERATED_BODY()
 };
 
 USTRUCT()
-struct FMassConstSharedFragment
+struct FMassConstSharedFragment : public FMassElement
 {
 	GENERATED_BODY()
 };
@@ -51,6 +57,13 @@ namespace UE::Mass
 		return false;
 	}
 
+	template<>
+	inline bool IsA<FMassElement>(const UStruct* Struct)
+	{
+		// @todo check that it's not any of the directly inherited types (that are "abstract" in their nature until inherited from themselves). 
+		return Struct && Struct->IsChildOf(FMassElement::StaticStruct());
+	}
+
 	template<>
 	inline bool IsA<FMassFragment>(const UStruct* Struct)
 	{
diff --git a/Engine/Source/Runtime/MassEntity/Public/MassEntityManager.h b/Engine/Source/Runtime/MassEntity/Public/MassEntityManager.h
--- a/Engine/Source/Runtime/MassEntity/Public/MassEntityManager.h
+++ b/Engine/Source/Runtime/MassEntity/Public/MassEntityManager.h
@@ -215,6 +215,10 @@ public:
 		, const FMassConstSharedFragmentBitSet& ConstSharedFragmentBitSet
 		, const FMassArchetypeCreationParams& CreationParams = FMassArchetypeCreationParams());
 
+	UE_API FMassArchetypeHandle GetOrCreateSuitableArchetype(const FMassArchetypeHandle& ArchetypeHandle
+		, const FMassElementBitSet& SharedElementsBitSet
+		, const FMassArchetypeCreationParams& CreationParams = FMassArchetypeCreationParams());
+
 	/** Fetches the archetype for a given EntityHandle. If EntityHandle is not valid it will still return a handle, just with an invalid archetype */
 	UE_API FMassArchetypeHandle GetArchetypeForEntity(FMassEntityHandle EntityHandle) const;
 	/**
@@ -486,8 +490,20 @@ public:
 		, const FMassArchetypeSharedFragmentValues& SharedFragmentValues = {}, const FMassArchetypeCreationParams& CreationParams = FMassArchetypeCreationParams());
 	UE_API TSharedRef<FEntityCreationContext> BatchBuildEntities(const FMassArchetypeEntityCollectionWithPayload& EncodedEntitiesWithPayload, const FMassArchetypeCompositionDescriptor& Composition
 		, const FMassArchetypeSharedFragmentValues& SharedFragmentValues = {}, const FMassArchetypeCreationParams& CreationParams = FMassArchetypeCreationParams());
-	UE_API void BatchChangeTagsForEntities(TConstArrayView<FMassArchetypeEntityCollection> EntityCollections, const FMassTagBitSet& TagsToAdd, const FMassTagBitSet& TagsToRemove);
-	UE_API void BatchChangeFragmentCompositionForEntities(TConstArrayView<FMassArchetypeEntityCollection> EntityCollections, const FMassFragmentBitSet& FragmentsToAdd, const FMassFragmentBitSet& FragmentsToRemove);
+	
+	UE_API void BatchChangeCompositionForEntities(TConstArrayView<FMassArchetypeEntityCollection> EntityCollections, const FMassElementBitSet& ElementsToAdd, const FMassElementBitSet& ElementsToRemove);
+
+	void BatchChangeTagsForEntities(TConstArrayView<FMassArchetypeEntityCollection> EntityCollections, const FMassTagBitSet& TagsToAdd, const FMassTagBitSet& TagsToRemove)
+	{
+		TRACE_CPUPROFILER_EVENT_SCOPE(Mass_BatchChangeTagsForEntities);
+		BatchChangeCompositionForEntities(EntityCollections, TagsToAdd, TagsToRemove);
+	}
+	void BatchChangeFragmentCompositionForEntities(TConstArrayView<FMassArchetypeEntityCollection> EntityCollections, const FMassFragmentBitSet& FragmentsToAdd, const FMassFragmentBitSet& FragmentsToRemove)
+	{
+		TRACE_CPUPROFILER_EVENT_SCOPE(Mass_BatchChangeFragmentCompositionForEntities);
+		BatchChangeCompositionForEntities(EntityCollections, FragmentsToAdd, FragmentsToRemove);
+	}
+
 	UE_API void BatchAddFragmentInstancesForEntities(TConstArrayView<FMassArchetypeEntityCollectionWithPayload> EntityCollections, const FMassFragmentBitSet& FragmentsAffected);
 	/** 
 	 * Adds a new const and non-const shared fragments to all entities provided via EntityCollections 
@@ -504,7 +520,7 @@ public:
 	UE_API void AddCompositionToEntity_GetDelta(FMassEntityHandle EntityHandle, FMassArchetypeCompositionDescriptor& InOutDescriptor, const FMassArchetypeSharedFragmentValues* AddedSharedFragmentValues = nullptr);
 	UE_API void RemoveCompositionFromEntity(FMassEntityHandle EntityHandle, const FMassArchetypeCompositionDescriptor& InDescriptor);
 
-	UE_API const FMassArchetypeCompositionDescriptor& GetArchetypeComposition(const FMassArchetypeHandle& ArchetypeHandle) const;
+	UE_API FMassArchetypeCompositionDescriptor GetArchetypeComposition(const FMassArchetypeHandle& ArchetypeHandle) const;
 
 	/** 
 	 * Moves an entity over to a new archetype by copying over fragments common to both archetypes
@@ -1130,8 +1146,7 @@ protected:
 	 * A "similar" archetype is an archetype exactly the same as SourceArchetype except for one composition aspect 
 	 * like Fragments or "Tags" 
 	 */
-	UE_API FMassArchetypeHandle InternalCreateSimilarArchetype(const TSharedPtr<FMassArchetypeData>& SourceArchetype, const FMassTagBitSet& OverrideTags);
-	UE_API FMassArchetypeHandle InternalCreateSimilarArchetype(const TSharedPtr<FMassArchetypeData>& SourceArchetype, const FMassFragmentBitSet& OverrideFragments);
+	UE_API FMassArchetypeHandle InternalCreateSimilarArchetype(const TSharedPtr<FMassArchetypeData>& SourceArchetype, FMassElementBitSet&& NewCompositionBitSet);
 	UE_API FMassArchetypeHandle InternalCreateSimilarArchetype(const TSharedPtr<FMassArchetypeData>& SourceArchetype, const UE::Mass::FArchetypeGroups& GroupsOverride);
 	UE_API FMassArchetypeHandle InternalCreateSimilarArchetype(const FMassArchetypeData& SourceArchetypeRef, FMassArchetypeCompositionDescriptor&& NewComposition, const UE::Mass::FArchetypeGroups& GroupsOverride);
 
diff --git a/Engine/Source/Runtime/MassEntity/Public/MassEntityTypes.h b/Engine/Source/Runtime/MassEntity/Public/MassEntityTypes.h
--- a/Engine/Source/Runtime/MassEntity/Public/MassEntityTypes.h
+++ b/Engine/Source/Runtime/MassEntity/Public/MassEntityTypes.h
@@ -2,23 +2,24 @@
 
 #pragma once
 
-#include "StructUtils/StructTypeBitSet.h"
 #include "MassProcessingTypes.h"
 #include "StructUtils/StructArrayView.h"
+#if UE_ENABLE_INCLUDE_ORDER_DEPRECATED_IN_5_8
 #include "Subsystems/Subsystem.h"
+#endif // UE_ENABLE_INCLUDE_ORDER_DEPRECATED_IN_5_8
 #include "Misc/ScopeRWLock.h"
 #include "Misc/TransactionallySafeRWLock.h"
 #if UE_ENABLE_INCLUDE_ORDER_DEPRECATED_IN_5_6
 #include "MassExternalSubsystemTraits.h"
 #endif // UE_ENABLE_INCLUDE_ORDER_DEPRECATED_IN_5_6
 #include "StructUtils/SharedStruct.h"
-#include "MassEntityElementTypes.h"
 #include "MassEntityConcepts.h"
 #include "MassTestableEnsures.h"
 #if UE_ENABLE_INCLUDE_ORDER_DEPRECATED_IN_5_6
 #include "MassExternalSubsystemTraits.h"
 #include "MassEntityHandle.h"
 #endif // UE_ENABLE_INCLUDE_ORDER_DEPRECATED_IN_5_6
+#include "MassElement.h"
 #include "MassEntityTypes.generated.h"
 
 
@@ -27,13 +28,6 @@ MASSENTITY_API DECLARE_LOG_CATEGORY_EXTERN(LogMass, Warning, All);
 DECLARE_STATS_GROUP(TEXT("Mass"), STATGROUP_Mass, STATCAT_Advanced);
 DECLARE_CYCLE_STAT_EXTERN(TEXT("Mass Total Frame Time"), STAT_Mass_Total, STATGROUP_Mass, MASSENTITY_API);
 
-DECLARE_STRUCTTYPEBITSET_EXPORTED(MASSENTITY_API, FMassFragmentBitSet, FMassFragment);
-DECLARE_STRUCTTYPEBITSET_EXPORTED(MASSENTITY_API, FMassTagBitSet, FMassTag);
-DECLARE_STRUCTTYPEBITSET_EXPORTED(MASSENTITY_API, FMassChunkFragmentBitSet, FMassChunkFragment);
-DECLARE_STRUCTTYPEBITSET_EXPORTED(MASSENTITY_API, FMassSharedFragmentBitSet, FMassSharedFragment);
-DECLARE_STRUCTTYPEBITSET_EXPORTED(MASSENTITY_API, FMassConstSharedFragmentBitSet, FMassConstSharedFragment);
-DECLARE_CLASSTYPEBITSET_EXPORTED(MASSENTITY_API, FMassExternalSubsystemBitSet, USubsystem);
-
 struct FMassArchetypeData;
 struct FMassEntityQuery;
 
@@ -85,7 +79,9 @@ struct FMassArchetypeCompositionDescriptor
 		, ChunkFragments(InChunkFragments)
 		, SharedFragments(InSharedFragments)
 		, ConstSharedFragments(InConstSharedFragments)
-	{}
+	{
+		SetElementsBitSet();
+	}
 
 	FMassArchetypeCompositionDescriptor(TConstArrayView<const UScriptStruct*> InFragments,
 		const FMassTagBitSet& InTags,
@@ -93,7 +89,8 @@ struct FMassArchetypeCompositionDescriptor
 		const FMassSharedFragmentBitSet& InSharedFragments,
 		const FMassConstSharedFragmentBitSet& InConstSharedFragments)
 		: FMassArchetypeCompositionDescriptor(FMassFragmentBitSet(InFragments), InTags, InChunkFragments, InSharedFragments, InConstSharedFragments)
-	{}
+	{
+	}
 
 	FMassArchetypeCompositionDescriptor(TConstArrayView<FInstancedStruct> InFragmentInstances,
 		const FMassTagBitSet& InTags,
@@ -101,7 +98,8 @@ struct FMassArchetypeCompositionDescriptor
 		const FMassSharedFragmentBitSet& InSharedFragments,
 		const FMassConstSharedFragmentBitSet& InConstSharedFragments)
 		: FMassArchetypeCompositionDescriptor(FMassFragmentBitSet(InFragmentInstances), InTags, InChunkFragments, InSharedFragments, InConstSharedFragments)
-	{}
+	{
+	}
 
 	FMassArchetypeCompositionDescriptor(FMassFragmentBitSet&& InFragments,
 		FMassTagBitSet&& InTags,
@@ -113,23 +111,21 @@ struct FMassArchetypeCompositionDescriptor
 		, ChunkFragments(MoveTemp(InChunkFragments))
 		, SharedFragments(MoveTemp(InSharedFragments))
 		, ConstSharedFragments(MoveTemp(InConstSharedFragments))
-	{}
+	{
+		SetElementsBitSet();
+	}
 
-	FMassArchetypeCompositionDescriptor(FMassFragmentBitSet&& InFragments)
-		: Fragments(MoveTemp(InFragments))
+	FMassArchetypeCompositionDescriptor(FMassElementBitSet&& InElements)
+		: ElementsBitSet(MoveTemp(InElements))
 	{}
 
-	FMassArchetypeCompositionDescriptor(FMassTagBitSet&& InTags)
-		: Tags(MoveTemp(InTags))
+	explicit FMassArchetypeCompositionDescriptor(const FMassElementBitSet& InElements)
+		: ElementsBitSet(InElements)
 	{}
 
 	void Reset()
 	{
-		Fragments.Reset();
-		Tags.Reset();
-		ChunkFragments.Reset();
-		SharedFragments.Reset();
-		ConstSharedFragments.Reset();
+		ElementsBitSet.Reset();
 	}
 
 	/**
@@ -137,11 +133,7 @@ struct FMassArchetypeCompositionDescriptor
 	 */
 	bool IsEquivalent(const FMassArchetypeCompositionDescriptor& OtherDescriptor) const
 	{
-		return Fragments.IsEquivalent(OtherDescriptor.Fragments)
-			&& Tags.IsEquivalent(OtherDescriptor.Tags)
-			&& ChunkFragments.IsEquivalent(OtherDescriptor.ChunkFragments)
-			&& SharedFragments.IsEquivalent(OtherDescriptor.SharedFragments)
-			&& ConstSharedFragments.IsEquivalent(OtherDescriptor.ConstSharedFragments);
+		return ElementsBitSet.IsEquivalent(OtherDescriptor.GetElementsBitSet());
 	}
 
 	/**
@@ -149,47 +141,27 @@ struct FMassArchetypeCompositionDescriptor
 	 */
 	bool IsIdentical(const FMassArchetypeCompositionDescriptor& OtherDescriptor) const
 	{
-		return Fragments == OtherDescriptor.Fragments
-			&& Tags == OtherDescriptor.Tags
-			&& ChunkFragments == OtherDescriptor.ChunkFragments
-			&& SharedFragments == OtherDescriptor.SharedFragments
-			&& ConstSharedFragments == OtherDescriptor.ConstSharedFragments;
+		return ElementsBitSet == OtherDescriptor.GetElementsBitSet();
 	}
 
 	bool IsEmpty() const 
-	{ 
-		return Fragments.IsEmpty()
-			&& Tags.IsEmpty()
-			&& ChunkFragments.IsEmpty()
-			&& SharedFragments.IsEmpty()
-			&& ConstSharedFragments.IsEmpty();
+	{
+		return ElementsBitSet.IsEmpty();
 	}
 
 	bool HasAll(const FMassArchetypeCompositionDescriptor& OtherDescriptor) const
 	{
-		return Fragments.HasAll(OtherDescriptor.Fragments)
-			&& Tags.HasAll(OtherDescriptor.Tags)
-			&& ChunkFragments.HasAll(OtherDescriptor.ChunkFragments)
-			&& SharedFragments.HasAll(OtherDescriptor.SharedFragments)
-			&& ConstSharedFragments.HasAll(OtherDescriptor.ConstSharedFragments);
+		return ElementsBitSet.HasAll(OtherDescriptor.GetElementsBitSet());
 	}
 
 	void Append(const FMassArchetypeCompositionDescriptor& OtherDescriptor)
 	{
-		Fragments += OtherDescriptor.Fragments;
-		Tags += OtherDescriptor.Tags;
-		ChunkFragments += OtherDescriptor.ChunkFragments;
-		SharedFragments += OtherDescriptor.SharedFragments;
-		ConstSharedFragments += OtherDescriptor.ConstSharedFragments;
+		ElementsBitSet += OtherDescriptor.ElementsBitSet;
 	}
 
 	void Remove(const FMassArchetypeCompositionDescriptor& OtherDescriptor)
 	{
-		Fragments -= OtherDescriptor.Fragments;
-		Tags -= OtherDescriptor.Tags;
-		ChunkFragments -= OtherDescriptor.ChunkFragments;
-		SharedFragments -= OtherDescriptor.SharedFragments;
-		ConstSharedFragments -= OtherDescriptor.ConstSharedFragments;
+		ElementsBitSet -= OtherDescriptor.ElementsBitSet;
 	}
 
 	/**
@@ -198,24 +170,20 @@ struct FMassArchetypeCompositionDescriptor
 	 */
 	FMassArchetypeCompositionDescriptor CalculateDifference(const FMassArchetypeCompositionDescriptor& OtherDescriptor) const
 	{
-		FMassArchetypeCompositionDescriptor Diff;
-
-		Diff.Fragments = Fragments - OtherDescriptor.Fragments;
-		Diff.Tags = Tags - OtherDescriptor.Tags;
-		Diff.ChunkFragments = ChunkFragments - OtherDescriptor.ChunkFragments;
-		Diff.SharedFragments = SharedFragments - OtherDescriptor.SharedFragments;
-		Diff.ConstSharedFragments = ConstSharedFragments - OtherDescriptor.ConstSharedFragments;
+		return FMassArchetypeCompositionDescriptor(ElementsBitSet - OtherDescriptor.ElementsBitSet);
+	}
 
-		return Diff;
+	uint32 CalculateHash() const 
+	{
+		return GetTypeHash(ElementsBitSet);
 	}
 
-	MASSENTITY_API static uint32 CalculateHash(const FMassFragmentBitSet& InFragments, const FMassTagBitSet& InTags
+	UE_DEPRECATED(5.8, "This function is deprecated. Use FMassArchetypeCompositionDescriptor::CalculateHash() instead")
+	static uint32 CalculateHash(const FMassFragmentBitSet& InFragments, const FMassTagBitSet& InTags
 		, const FMassChunkFragmentBitSet& InChunkFragments, const FMassSharedFragmentBitSet& InSharedFragmentBitSet
-		, const FMassConstSharedFragmentBitSet& InConstSharedFragmentBitSet);
-
-	uint32 CalculateHash() const 
+		, const FMassConstSharedFragmentBitSet& InConstSharedFragmentBitSet)
 	{
-		return CalculateHash(Fragments, Tags, ChunkFragments, SharedFragments, ConstSharedFragments);
+		return GetTypeHash(InFragments + InTags + InChunkFragments + InSharedFragmentBitSet + InConstSharedFragmentBitSet);
 	}
 
 	MASSENTITY_API int32 CountStoredTypes() const;
@@ -223,34 +191,73 @@ struct FMassArchetypeCompositionDescriptor
 	MASSENTITY_API void DebugOutputDescription(FOutputDevice& Ar) const;
 
 	template<typename T>
-	auto& GetContainer() const;
+	auto GetContainer() const;
 
 	template<typename T>
-	auto& GetContainer();
+	UE_DEPRECATED(5.8, "Mutable GetContainer is deprecated")
+	auto& GetContainer()
+	{
+		return GetElementsBitSet();
+	}
 
 	template<typename T>
 	bool Contains() const;
 
+	bool Contains(TNotNull<const UScriptStruct*> ElementType) const;
+
 	template<typename T>
 	void Add();
 
-	const FMassFragmentBitSet& GetFragments() const;
-	const FMassTagBitSet& GetTags() const;
-	const FMassChunkFragmentBitSet& GetChunkFragments() const;
-	const FMassSharedFragmentBitSet& GetSharedFragments() const;
-	const FMassConstSharedFragmentBitSet& GetConstSharedFragments() const;
-
-	FMassFragmentBitSet& GetFragments();
-	FMassTagBitSet& GetTags();
-	FMassChunkFragmentBitSet& GetChunkFragments();
-	FMassSharedFragmentBitSet& GetSharedFragments();
-	FMassConstSharedFragmentBitSet& GetConstSharedFragments();
+	void Add(TNotNull<const UScriptStruct*> ElementType);
+
+	FMassFragmentBitSet DebugGetFragments() const;
+	FMassTagBitSet DebugGetTags() const;
+	FMassChunkFragmentBitSet DebugGetChunkFragments() const;
+	FMassSharedFragmentBitSet DebugGetSharedFragments() const;
+	FMassConstSharedFragmentBitSet DebugGetConstSharedFragments() const;
+
+	UE_DEPRECATED(5.8, "Getting a subset of compossition's bitsets is slow now. Consider refactoring your code to work directly on FMassElementBitSet")
+	FMassFragmentBitSet GetFragments() const;
+	UE_DEPRECATED(5.8, "Getting a subset of compossition's bitsets is slow now. Consider refactoring your code to work directly on FMassElementBitSet")
+	FMassTagBitSet GetTags() const;
+	UE_DEPRECATED(5.8, "Getting a subset of compossition's bitsets is slow now. Consider refactoring your code to work directly on FMassElementBitSet")
+	FMassChunkFragmentBitSet GetChunkFragments() const;
+	UE_DEPRECATED(5.8, "Getting a subset of compossition's bitsets is slow now. Consider refactoring your code to work directly on FMassElementBitSet")
+	FMassSharedFragmentBitSet GetSharedFragments() const;
+	UE_DEPRECATED(5.8, "Getting a subset of compossition's bitsets is slow now. Consider refactoring your code to work directly on FMassElementBitSet")
+	FMassConstSharedFragmentBitSet GetConstSharedFragments() const;
 
+	/**
+	 * Functions setting specific type of bits. These functions are a lot slower now.
+	 * Consider changing your code to work directly on FMassElementBitSet
+	 */
+	UE_DEPRECATED(5.8, "Setting a subset of compossition's bitsets is slow now. Consider refactoring your code to work directly on FMassElementBitSet")
 	void SetFragments(const FMassFragmentBitSet& InBitSet);
+	UE_DEPRECATED(5.8, "Setting a subset of compossition's bitsets is slow now. Consider refactoring your code to work directly on FMassElementBitSet")
 	void SetTags(const FMassTagBitSet& InBitSet);
+	UE_DEPRECATED(5.8, "Setting a subset of compossition's bitsets is slow now. Consider refactoring your code to work directly on FMassElementBitSet")
 	void SetChunkFragments(const FMassChunkFragmentBitSet& InBitSet);
+	UE_DEPRECATED(5.8, "Setting a subset of compossition's bitsets is slow now. Consider refactoring your code to work directly on FMassElementBitSet")
 	void SetSharedFragments(const FMassSharedFragmentBitSet& InBitSet);
+	UE_DEPRECATED(5.8, "Setting a subset of compossition's bitsets is slow now. Consider refactoring your code to work directly on FMassElementBitSet")
 	void SetConstSharedFragments(const FMassConstSharedFragmentBitSet& InBitSet);
+	
+	void SetAllSharedElements(const FMassElementBitSet& InBitSet);
+
+	FMassElementBitSet GetAllSharedFragments() const
+	{
+		return ElementsBitSet & (FMassElementBitSet::GetAllSharedFragments());
+	}
+
+	const FMassElementBitSet& GetElementsBitSet() const
+	{
+		return ElementsBitSet;
+	}
+
+	FMassElementBitSet& GetElementsBitSet()
+	{
+		return ElementsBitSet;
+	}
 
 	UE_DEPRECATED(5.7, "Direct access to FMassArchetypeCompositionDescriptor's bitsets is deprecated. Use any of the newly added getters instead.")
 	FMassFragmentBitSet Fragments;
@@ -267,6 +274,37 @@ struct FMassArchetypeCompositionDescriptor
 	UE_DEPRECATED(5.7, "Direct access to FMassArchetypeCompositionDescriptor's bitsets is deprecated. Use any of the newly added getters instead.")
 	FMassConstSharedFragmentBitSet ConstSharedFragments;
 
+protected:
+	FMassElementBitSet ElementsBitSet;
+
+public:
+	template<typename TBitSet>
+	TBitSet Get() const
+	{
+		return ElementsBitSet.Get<TBitSet>();
+	}
+
+	void SetElementsBitSet()
+	{
+		ElementsBitSet = Fragments;
+		ElementsBitSet += Tags;
+		ElementsBitSet += ChunkFragments;
+		ElementsBitSet += SharedFragments;
+		ElementsBitSet += ConstSharedFragments;
+	}
+
+	template<typename TElementType>
+	bool HasAny() const
+	{
+		return ElementsBitSet.HasAny<TElementType>();
+	}
+
+	template<typename TElementType>
+	bool HasAll() const
+	{
+		return ElementsBitSet.HasAll<TElementType>();
+	}
+
 	UE_DEPRECATED(5.5, "This FMassArchetypeCompositionDescriptor constructor is deprecated. Please explicitly provide FConstSharedFragmentBitSet.")
 	FMassArchetypeCompositionDescriptor(const FMassFragmentBitSet& InFragments, const FMassTagBitSet& InTags, const FMassChunkFragmentBitSet& InChunkFragments, const FMassSharedFragmentBitSet& InSharedFragments)
 		: FMassArchetypeCompositionDescriptor(InFragments, InTags, InChunkFragments, InSharedFragments, FMassConstSharedFragmentBitSet())
@@ -302,8 +340,7 @@ struct FMassArchetypeSharedFragmentValues
 	FMassArchetypeSharedFragmentValues(const FMassArchetypeSharedFragmentValues& Other)
 		: HashCache(Other.HashCache)
 		, bSorted(Other.bSorted)
-		, SharedFragmentBitSet(Other.SharedFragmentBitSet)
-		, ConstSharedFragmentBitSet(Other.ConstSharedFragmentBitSet)
+		, StoredElementsBitSet(Other.StoredElementsBitSet)
 		, ConstSharedFragments(Other.ConstSharedFragments)
 		, SharedFragments(Other.SharedFragments)
 	{
@@ -312,8 +349,7 @@ struct FMassArchetypeSharedFragmentValues
 	FMassArchetypeSharedFragmentValues(FMassArchetypeSharedFragmentValues&& Other)
 		: HashCache(Other.HashCache)
 		, bSorted(Other.bSorted)
-		, SharedFragmentBitSet(MoveTemp(Other.SharedFragmentBitSet))
-		, ConstSharedFragmentBitSet(MoveTemp(Other.ConstSharedFragmentBitSet))
+		, StoredElementsBitSet(MoveTemp(Other.StoredElementsBitSet))
 		, ConstSharedFragments(MoveTemp(Other.ConstSharedFragments))
 		, SharedFragments(MoveTemp(Other.SharedFragments))
 	{
@@ -325,8 +361,7 @@ struct FMassArchetypeSharedFragmentValues
 		{
 			UE::TWriteScopeLock WriteLock{ HashLock };
 
-			SharedFragmentBitSet = Other.SharedFragmentBitSet;
-			ConstSharedFragmentBitSet = Other.ConstSharedFragmentBitSet;
+			StoredElementsBitSet = Other.StoredElementsBitSet;
 			ConstSharedFragments = Other.ConstSharedFragments;
 			SharedFragments = Other.SharedFragments;
 			HashCache = Other.HashCache;
@@ -341,8 +376,7 @@ struct FMassArchetypeSharedFragmentValues
 		{
 			UE::TWriteScopeLock WriteLock{ HashLock };
 
-			SharedFragmentBitSet = MoveTemp(Other.SharedFragmentBitSet);
-			ConstSharedFragmentBitSet = MoveTemp(Other.ConstSharedFragmentBitSet);
+			StoredElementsBitSet = MoveTemp(Other.StoredElementsBitSet);
 			ConstSharedFragments = MoveTemp(Other.ConstSharedFragments);
 			SharedFragments = MoveTemp(Other.SharedFragments);
 			HashCache = Other.HashCache;
@@ -351,39 +385,25 @@ struct FMassArchetypeSharedFragmentValues
 		return *this;
 	}
 
-	inline bool HasExactFragmentTypesMatch(const FMassSharedFragmentBitSet& InSharedFragmentBitSet, const FMassConstSharedFragmentBitSet& InConstSharedFragmentBitSet) const
-	{
-		return HasExactSharedFragmentTypesMatch(InSharedFragmentBitSet)
-			&& HasExactConstSharedFragmentTypesMatch(InConstSharedFragmentBitSet);
-	}
-
-	inline bool HasExactSharedFragmentTypesMatch(const FMassSharedFragmentBitSet& InSharedFragmentBitSet) const
-	{
-		return SharedFragmentBitSet.IsEquivalent(InSharedFragmentBitSet);
-	}
-
-	inline bool HasAllRequiredSharedFragmentTypes(const FMassSharedFragmentBitSet& InSharedFragmentBitSet) const
-	{
-		return SharedFragmentBitSet.HasAll(InSharedFragmentBitSet);
-	}
-
-	inline bool HasExactConstSharedFragmentTypesMatch(const FMassConstSharedFragmentBitSet& InConstSharedFragmentBitSet) const
-	{
-		return ConstSharedFragmentBitSet.IsEquivalent(InConstSharedFragmentBitSet);
-	}
-
-	inline bool HasAllRequiredConstSharedFragmentTypes(const FMassConstSharedFragmentBitSet& InConstSharedFragmentBitSet) const
-	{
-		return ConstSharedFragmentBitSet.HasAll(InConstSharedFragmentBitSet);
-	}
+	/**
+	 * @param SharedFragmentValuesModification if provided will be used to override or add values already existing in Original.
+	 *	Needs to be consistent with NewSharedElementComposition, meaning, not add anything that's not already in the bitset. 
+	 * Also, if NewSharedElementComposition contains elements not in Original, then SharedFragmentValuesModification has to provide a value for it.  
+	 */
+	static FMassArchetypeSharedFragmentValues CreateCombined(const FMassArchetypeSharedFragmentValues& Original
+		, const FMassElementBitSet& NewSharedElementComposition, const FMassArchetypeSharedFragmentValues* SharedFragmentValuesModification);
 
 	/**
 	 * @return whether the stored shared fragment values exactly match shared fragment types indicated by InDescriptor
 	 */
 	bool DoesMatchComposition(const FMassArchetypeCompositionDescriptor& InDescriptor) const 
 	{
-		return HasExactSharedFragmentTypesMatch(InDescriptor.GetSharedFragments())
-			&& HasExactConstSharedFragmentTypesMatch(InDescriptor.GetConstSharedFragments());
+		return DoContainEquivalentSharedFragments(StoredElementsBitSet, InDescriptor.GetElementsBitSet());
+	}
+
+	bool DoesMatchComposition(const FMassElementBitSet& CompositionBitSet) const 
+	{
+		return DoContainEquivalentSharedFragments(StoredElementsBitSet, CompositionBitSet);
 	}
 
 	inline bool IsEquivalent(const FMassArchetypeSharedFragmentValues& OtherSharedFragmentValues) const
@@ -401,30 +421,17 @@ struct FMassArchetypeSharedFragmentValues
 
 	inline bool ContainsType(const UScriptStruct* FragmentType) const
 	{
-		if (UE::Mass::IsA<FMassSharedFragment>(FragmentType))
-		{
-			return SharedFragmentBitSet.Contains(*FragmentType);
-		}
-
-		if (UE::Mass::IsA<FMassConstSharedFragment>(FragmentType))
-		{
-			return ConstSharedFragmentBitSet.Contains(*FragmentType);
-		}
-
-		return false;
+		return (UE::Mass::IsA<FMassSharedFragment>(FragmentType) || UE::Mass::IsA<FMassConstSharedFragment>(FragmentType))
+			&& StoredElementsBitSet.Contains(*FragmentType);
 	}
 
 	template<typename T>
 	inline bool ContainsType() const
 	{
-		if constexpr (UE::Mass::CConstSharedFragment<T>)
+		if constexpr (UE::Mass::CConstSharedFragment<T> || UE::Mass::CSharedFragment<T>)
 		{
-			return ConstSharedFragmentBitSet.Contains(*T::StaticStruct());
-		}
-		else if constexpr (UE::Mass::CSharedFragment<T>)
-		{
-			return SharedFragmentBitSet.Contains(*T::StaticStruct());
-		}
+			return StoredElementsBitSet.Contains(*T::StaticStruct());
+		}		
 		else
 		{
 			return false;
@@ -493,17 +500,25 @@ struct FMassArchetypeSharedFragmentValues
 	 */
 	MASSENTITY_API int32 Append(const FMassArchetypeSharedFragmentValues& Other);
 
+	MASSENTITY_API int32 Remove(const FMassElementBitSet& ElementBitSet);
+
 	/** 
 	 * Note that the function removes the shared fragments by type
 	 * @return number of fragments types removed
 	 */
-	MASSENTITY_API int32 Remove(const FMassSharedFragmentBitSet& SharedFragmentToRemoveBitSet);
+	int32 Remove(const FMassSharedFragmentBitSet& SharedFragmentToRemoveBitSet)
+	{
+		return Remove(FMassElementBitSet(SharedFragmentToRemoveBitSet));
+	}
 
 	/** 
 	 * Note that the function removes the const shared fragments by type
 	 * @return number of fragments types removed
 	 */
-	MASSENTITY_API int32 Remove(const FMassConstSharedFragmentBitSet& ConstSharedFragmentToRemoveBitSet);
+	int32 Remove(const FMassConstSharedFragmentBitSet& ConstSharedFragmentToRemoveBitSet)
+	{
+		return Remove(FMassElementBitSet(ConstSharedFragmentToRemoveBitSet));
+	}
 
 	/**
 	 * Remove all the shared and const shared fragments indicated by InDescriptor
@@ -511,7 +526,12 @@ struct FMassArchetypeSharedFragmentValues
 	 */
 	int32 Remove(const FMassArchetypeCompositionDescriptor& InDescriptor)
 	{
-		return Remove(InDescriptor.GetSharedFragments()) + Remove(InDescriptor.GetConstSharedFragments());
+		return Remove(InDescriptor.GetAllSharedFragments());
+	}
+
+	int32 Remove(const FMassArchetypeSharedFragmentValues& Other)
+	{
+		return Remove(Other.GetBitSet());
 	}
 
 	inline const TArray<FConstSharedStruct>& GetConstSharedFragments() const
@@ -547,14 +567,21 @@ struct FMassArchetypeSharedFragmentValues
 		return FragmentIndex != INDEX_NONE ? SharedFragments[FragmentIndex] : FSharedStruct();
 	}
 
-	const FMassSharedFragmentBitSet& GetSharedFragmentBitSet() const
+	const FMassElementBitSet& GetBitSet() const
 	{
-		return SharedFragmentBitSet;
+		return StoredElementsBitSet;
 	}
 
-	const FMassConstSharedFragmentBitSet& GetConstSharedFragmentBitSet() const
+	UE_DEPRECATED(5.8, "This function is deprecated. Use GetBitSet instead")
+	FMassSharedFragmentBitSet GetSharedFragmentBitSet() const
 	{
-		return ConstSharedFragmentBitSet;
+		return GetBitSet();
+	}
+
+	UE_DEPRECATED(5.8, "This function is deprecated. Use GetBitSet instead")
+	FMassConstSharedFragmentBitSet GetConstSharedFragmentBitSet() const
+	{
+		return GetBitSet();
 	}
 
 	inline void DirtyHashCache()
@@ -617,8 +644,7 @@ protected:
 	 */ 
 	mutable bool bSorted = true; 
 	
-	FMassSharedFragmentBitSet SharedFragmentBitSet;
-	FMassConstSharedFragmentBitSet ConstSharedFragmentBitSet;
+	FMassElementBitSet StoredElementsBitSet;
 	TArray<FConstSharedStruct> ConstSharedFragments;
 	TArray<FSharedStruct> SharedFragments;
 
@@ -632,7 +658,7 @@ public:
 	UE_DEPRECATED(5.5, "HasExactFragmentTypesMatch is deprecated. Use HasExactSharedFragmentTypesMatch or the two-parameter version of HasExactFragmentTypesMatch.")
 	inline bool HasExactFragmentTypesMatch(const FMassSharedFragmentBitSet& InSharedFragmentBitSet) const
 	{
-		return HasExactSharedFragmentTypesMatch(InSharedFragmentBitSet);
+		return (StoredElementsBitSet & InSharedFragmentBitSet).IsEquivalent(InSharedFragmentBitSet);
 	}
 };
 
@@ -898,27 +924,27 @@ struct FMassArchetypeCreationParams
 //-----------------------------------------------------------------------------
 PRAGMA_DISABLE_DEPRECATION_WARNINGS
 template<typename T>
-auto& FMassArchetypeCompositionDescriptor::GetContainer() const
+auto FMassArchetypeCompositionDescriptor::GetContainer() const
 {
 	if constexpr (std::is_same_v<FMassFragment, T>)
 	{
-		return Fragments;
+		return ElementsBitSet.Get<FMassFragmentBitSet>();
 	}
 	else if constexpr (std::is_same_v<FMassTag, T>)
 	{
-		return Tags;
+		return ElementsBitSet.Get<FMassTagBitSet>();
 	}
 	else if constexpr (std::is_same_v<FMassChunkFragment, T>)
 	{
-		return ChunkFragments;
+		return ElementsBitSet.Get<FMassChunkFragmentBitSet>();
 	}
 	else if constexpr (std::is_same_v<FMassSharedFragment, T>)
 	{
-		return SharedFragments;
+		return ElementsBitSet.Get<FMassSharedFragmentBitSet>();
 	}
 	else if constexpr (std::is_same_v<FMassConstSharedFragment, T>)
 	{
-		return ConstSharedFragments;
+		return ElementsBitSet.Get<FMassConstSharedFragmentBitSet>();
 	}
 	else
 	{
@@ -927,128 +953,111 @@ auto& FMassArchetypeCompositionDescriptor::GetContainer() const
 }
 
 template<typename T>
-auto& FMassArchetypeCompositionDescriptor::GetContainer()
+bool FMassArchetypeCompositionDescriptor::Contains() const
 {
-	if constexpr (std::is_same_v<FMassFragment, T>)
-	{
-		return Fragments;
-	}
-	else if constexpr (std::is_same_v<FMassTag, T>)
-	{
-		return Tags;
-	}
-	else if constexpr (std::is_same_v<FMassChunkFragment, T>)
-	{
-		return ChunkFragments;
-	}
-	else if constexpr (std::is_same_v<FMassSharedFragment, T>)
-	{
-		return SharedFragments;
-	}
-	else if constexpr (std::is_same_v<FMassConstSharedFragment, T>)
-	{
-		return ConstSharedFragments;
-	}
-	else
-	{
-		static_assert(UE::Mass::TAlwaysFalse<T>, "Unknown element type passed to GetContainer.");
-	}
+	return ElementsBitSet.Contains(*T::StaticStruct());
 }
 
-template<typename T>
-bool FMassArchetypeCompositionDescriptor::Contains() const
+inline bool FMassArchetypeCompositionDescriptor::Contains(TNotNull<const UScriptStruct*> ElementType) const
 {
-	using FElementType = UE::Mass::TElementType<T>;
-	return GetContainer<FElementType>().template Contains<T>();
+	return ElementsBitSet.Contains(*ElementType);
 }
 
 template<typename T>
 void FMassArchetypeCompositionDescriptor::Add()
 {
-	using FElementType = UE::Mass::TElementType<T>;
-	GetContainer<FElementType>().template Add<T>();
+	return ElementsBitSet.Add(*T::StaticStruct());
+}
+
+inline void FMassArchetypeCompositionDescriptor::Add(TNotNull<const UScriptStruct*> ElementType)
+{
+	return ElementsBitSet.Add(*ElementType);
 }
 
-inline const FMassFragmentBitSet& FMassArchetypeCompositionDescriptor::GetFragments() const 
+inline FMassFragmentBitSet FMassArchetypeCompositionDescriptor::DebugGetFragments() const 
 { 
-	return Fragments; 
+	return ElementsBitSet.Get<FMassFragmentBitSet>();
 }
 
-inline const FMassTagBitSet& FMassArchetypeCompositionDescriptor::GetTags() const 
+inline FMassTagBitSet FMassArchetypeCompositionDescriptor::DebugGetTags() const 
 { 
-	return Tags; 
+	return ElementsBitSet.Get<FMassTagBitSet>();
 }
 
-inline const FMassChunkFragmentBitSet& FMassArchetypeCompositionDescriptor::GetChunkFragments() const 
+inline FMassChunkFragmentBitSet FMassArchetypeCompositionDescriptor::DebugGetChunkFragments() const 
 { 
-	return ChunkFragments; 
+	return ElementsBitSet.Get<FMassChunkFragmentBitSet>();
 }
 
-inline const FMassSharedFragmentBitSet& FMassArchetypeCompositionDescriptor::GetSharedFragments() const 
+inline FMassSharedFragmentBitSet FMassArchetypeCompositionDescriptor::DebugGetSharedFragments() const 
 { 
-	return SharedFragments; 
+	return ElementsBitSet.Get<FMassSharedFragmentBitSet>();
 }
 
-inline const FMassConstSharedFragmentBitSet& FMassArchetypeCompositionDescriptor::GetConstSharedFragments() const 
+inline FMassConstSharedFragmentBitSet FMassArchetypeCompositionDescriptor::DebugGetConstSharedFragments() const 
 { 
-	return ConstSharedFragments; 
+	return ElementsBitSet.Get<FMassConstSharedFragmentBitSet>();
 }
 
-inline FMassFragmentBitSet& FMassArchetypeCompositionDescriptor::GetFragments() 
+inline FMassFragmentBitSet FMassArchetypeCompositionDescriptor::GetFragments() const 
 { 
-	return Fragments; 
+	return DebugGetFragments(); 
 }
 
-inline FMassTagBitSet& FMassArchetypeCompositionDescriptor::GetTags() 
+inline FMassTagBitSet FMassArchetypeCompositionDescriptor::GetTags() const 
 { 
-	return Tags; 
+	return DebugGetTags(); 
 }
 
-inline FMassChunkFragmentBitSet& FMassArchetypeCompositionDescriptor::GetChunkFragments() 
+inline FMassChunkFragmentBitSet FMassArchetypeCompositionDescriptor::GetChunkFragments() const 
 { 
-	return ChunkFragments; 
+	return DebugGetChunkFragments(); 
 }
 
-inline FMassSharedFragmentBitSet& FMassArchetypeCompositionDescriptor::GetSharedFragments() 
+inline FMassSharedFragmentBitSet FMassArchetypeCompositionDescriptor::GetSharedFragments() const 
 { 
-	return SharedFragments; 
+	return DebugGetSharedFragments(); 
 }
 
-inline FMassConstSharedFragmentBitSet& FMassArchetypeCompositionDescriptor::GetConstSharedFragments() 
+inline FMassConstSharedFragmentBitSet FMassArchetypeCompositionDescriptor::GetConstSharedFragments() const 
 { 
-	return ConstSharedFragments; 
+	return DebugGetConstSharedFragments(); 
 }
 
 inline void FMassArchetypeCompositionDescriptor::SetFragments(const FMassFragmentBitSet& InBitSet)
 { 
-	Fragments = InBitSet; 
+	ElementsBitSet = (ElementsBitSet - UE::Mass::FElementBitSet::GetTypeBitArray(UE::Mass::EElementType::Fragment)) + InBitSet;
 }
 
 inline void FMassArchetypeCompositionDescriptor::SetTags(const FMassTagBitSet& InBitSet)
 { 
-	Tags = InBitSet; 
+	ElementsBitSet = (ElementsBitSet - UE::Mass::FElementBitSet::GetTypeBitArray(UE::Mass::EElementType::Tag)) + InBitSet;
 }
 
 inline void FMassArchetypeCompositionDescriptor::SetChunkFragments(const FMassChunkFragmentBitSet& InBitSet)
 { 
-	ChunkFragments = InBitSet; 
+	ElementsBitSet = (ElementsBitSet - UE::Mass::FElementBitSet::GetTypeBitArray(UE::Mass::EElementType::ChunkFragment)) + InBitSet;
 }
 
 inline void FMassArchetypeCompositionDescriptor::SetSharedFragments(const FMassSharedFragmentBitSet& InBitSet)
 { 
-	SharedFragments = InBitSet; 
+	ElementsBitSet = (ElementsBitSet - UE::Mass::FElementBitSet::GetTypeBitArray(UE::Mass::EElementType::SharedFragment)) + InBitSet;
 }
 
 inline void FMassArchetypeCompositionDescriptor::SetConstSharedFragments(const FMassConstSharedFragmentBitSet& InBitSet)
 { 
-	ConstSharedFragments = InBitSet; 
+	ElementsBitSet = (ElementsBitSet - UE::Mass::FElementBitSet::GetTypeBitArray(UE::Mass::EElementType::ConstSharedFragment)) + InBitSet;
+}
+
+inline void FMassArchetypeCompositionDescriptor::SetAllSharedElements(const FMassElementBitSet& InBitSet)
+{ 
+	ElementsBitSet = (ElementsBitSet - ElementsBitSet.GetAllSharedFragments()) + InBitSet;
 }
 PRAGMA_ENABLE_DEPRECATION_WARNINGS
 
 inline SIZE_T FMassArchetypeSharedFragmentValues::GetAllocatedSize() const
 {
-	return SharedFragmentBitSet.GetAllocatedSize()
-		+ ConstSharedFragmentBitSet.GetAllocatedSize()
+	return StoredElementsBitSet.GetAllocatedSize()
 		+ ConstSharedFragments.GetAllocatedSize()
 		+ SharedFragments.GetAllocatedSize();
 }
@@ -1069,8 +1078,7 @@ inline void FMassArchetypeSharedFragmentValues::Reset()
 
 	HashCache = UINT32_MAX;
 	bSorted = false; 
-	SharedFragmentBitSet.Reset();
-	ConstSharedFragmentBitSet.Reset();
+	StoredElementsBitSet.Reset();
 	ConstSharedFragments.Reset();
 	SharedFragments.Reset();
 }
diff --git a/Engine/Source/Runtime/MassEntity/Public/MassExecutionContext.h b/Engine/Source/Runtime/MassEntity/Public/MassExecutionContext.h
--- a/Engine/Source/Runtime/MassEntity/Public/MassExecutionContext.h
+++ b/Engine/Source/Runtime/MassEntity/Public/MassExecutionContext.h
@@ -74,7 +74,7 @@ private:
 	FInstancedStruct AuxData;
 	float DeltaTimeSeconds = 0.0f;
 	int32 ChunkSerialModificationNumber = -1;
-	FMassArchetypeCompositionDescriptor CurrentArchetypeCompositionDescriptor;
+	FMassElementBitSet CurrentArchetypeCompositionBitSet;
 #if WITH_MASSENTITY_DEBUG
 	FColor DebugColor;
 #endif // WITH_MASSENTITY_DEBUG
@@ -311,26 +311,26 @@ public:
 
 	bool DoesArchetypeHaveFragment(const UScriptStruct& FragmentType) const
 	{
-		return CurrentArchetypeCompositionDescriptor.GetFragments().Contains(FragmentType);
+		return CurrentArchetypeCompositionBitSet.Contains(FragmentType);
 	}
 
 	template<typename T>
 	bool DoesArchetypeHaveFragment() const
 	{
 		static_assert(UE::Mass::CFragment<T>, MASS_INVALID_FRAGMENT_MSG);
-		return CurrentArchetypeCompositionDescriptor.GetFragments().Contains<T>();
+		return CurrentArchetypeCompositionBitSet.Contains<T>();
 	}
 
 	bool DoesArchetypeHaveTag(const UScriptStruct& TagType) const
 	{
-		return CurrentArchetypeCompositionDescriptor.GetTags().Contains(TagType);
+		return CurrentArchetypeCompositionBitSet.Contains(TagType);
 	}
 
 	template<typename T>
 	bool DoesArchetypeHaveTag() const
 	{
 		static_assert(UE::Mass::CTag<T>, "Given struct is not of a valid tag type.");
-		return CurrentArchetypeCompositionDescriptor.GetTags().Contains<T>();
+		return CurrentArchetypeCompositionBitSet.Contains<T>();
 	}
 
 #if WITH_MASSENTITY_DEBUG
@@ -609,7 +609,11 @@ public:
 	void ClearExecutionData();
 	void SetCurrentArchetypeCompositionDescriptor(const FMassArchetypeCompositionDescriptor& Descriptor)
 	{
-		CurrentArchetypeCompositionDescriptor = Descriptor;
+		CurrentArchetypeCompositionBitSet = Descriptor.GetElementsBitSet();
+	}
+	void StoreCurrentArchetypeComposition(const FMassElementBitSet& CompositionBitSet)
+	{
+		CurrentArchetypeCompositionBitSet = CompositionBitSet;
 	}
 
 #if WITH_MASSENTITY_DEBUG
diff --git a/Engine/Source/Runtime/MassEntity/Public/MassObserverManager.h b/Engine/Source/Runtime/MassEntity/Public/MassObserverManager.h
--- a/Engine/Source/Runtime/MassEntity/Public/MassObserverManager.h
+++ b/Engine/Source/Runtime/MassEntity/Public/MassObserverManager.h
@@ -77,31 +77,14 @@ struct FMassObserverManager
 
 	FMassEntityManager& GetEntityManager();
 
-	const FMassFragmentBitSet* GetObservedFragmentBitSets() const { return ObservedFragments; }
-	const FMassFragmentBitSet& GetObservedFragmentsBitSet(const EMassObservedOperation Operation) const 
-	{ 
-		return ObservedFragments[(uint8)Operation]; 
-	}
-
-	const FMassTagBitSet* GetObservedTagBitSets() const { return ObservedTags; }
-	const FMassTagBitSet& GetObservedTagsBitSet(const EMassObservedOperation Operation) const 
-	{ 
-		return ObservedTags[(uint8)Operation]; 
-	}
-	
-	bool HasObserversForBitSet(const FMassFragmentBitSet& InQueriedBitSet, const EMassObservedOperation Operation) const
+	bool HasObserversForBitSet(const FMassElementBitSet& InQueriedBitSet, const EMassObservedOperation Operation) const
 	{
-		return ObservedFragments[(uint8)Operation].HasAny(InQueriedBitSet);
-	}
-
-	bool HasObserversForBitSet(const FMassTagBitSet& InQueriedBitSet, const EMassObservedOperation Operation) const
-	{
-		return ObservedTags[(uint8)Operation].HasAny(InQueriedBitSet);
+		return (InQueriedBitSet.IsEmpty() == false) && ObservedElements[static_cast<uint8>(Operation)].HasAny(InQueriedBitSet);
 	}
 
 	bool HasObserversForComposition(const FMassArchetypeCompositionDescriptor& Composition, const EMassObservedOperation Operation) const
 	{
-		return HasObserversForBitSet(Composition.GetFragments(), Operation) || HasObserversForBitSet(Composition.GetTags(), Operation);
+		return HasObserversForBitSet(Composition.GetElementsBitSet(), Operation);
 	}
 
 	/** @return whether there are observers watching affected elements */
@@ -156,7 +139,7 @@ protected:
 		FMassEntityHandle EntityHandle;
 		const FMassArchetypeEntityCollection& EntityCollection;
 	};
-	MASSENTITY_API bool OnCompositionChanged(FCollectionRefOrHandle&& EntityCollection, const FMassArchetypeCompositionDescriptor& Composition
+	MASSENTITY_API bool OnCompositionChanged(FCollectionRefOrHandle&& EntityCollection, const FMassElementBitSet& Composition
 		, const EMassObservedOperation Operation, UE::Mass::FProcessingContext* ProcessingContext = nullptr);
 
 public:
@@ -164,28 +147,28 @@ public:
 		, const EMassObservedOperation Operation, UE::Mass::FProcessingContext* ProcessingContext = nullptr)
 	{
 		TRACE_CPUPROFILER_EVENT_SCOPE(MassObserver_OnCompositionChanged_Collection);
-		return OnCompositionChanged(FCollectionRefOrHandle(EntityCollection), Composition, Operation, ProcessingContext);
+		return OnCompositionChanged(FCollectionRefOrHandle(EntityCollection), Composition.GetElementsBitSet(), Operation, ProcessingContext);
 	}
 
 	bool OnCompositionChanged(FMassEntityHandle EntityHandle, const FMassArchetypeCompositionDescriptor& Composition
 		, const EMassObservedOperation Operation, UE::Mass::FProcessingContext* ProcessingContext = nullptr)
 	{
 		TRACE_CPUPROFILER_EVENT_SCOPE(MassObserver_OnCompositionChanged_Entity);
-		return OnCompositionChanged(FCollectionRefOrHandle(EntityHandle), Composition, Operation, ProcessingContext);
+		return OnCompositionChanged(FCollectionRefOrHandle(EntityHandle), Composition.GetElementsBitSet(), Operation, ProcessingContext);
 	}
 
 	template<typename T, typename U = std::decay_t<T>> requires (std::is_same_v<U, FMassFragmentBitSet> || std::is_same_v<U, FMassTagBitSet>)
 	bool OnCompositionChanged(const FMassArchetypeEntityCollection& EntityCollection, T&& BitSet
 		, const EMassObservedOperation Operation, UE::Mass::FProcessingContext* ProcessingContext = nullptr)
 	{
-		return OnCompositionChanged(EntityCollection, FMassArchetypeCompositionDescriptor(MoveTemp(BitSet)), Operation, ProcessingContext);
+		return OnCompositionChanged(EntityCollection, FMassElementBitSet(MoveTemp(BitSet)), Operation, ProcessingContext);
 	}
 
 	template<typename T, typename U = std::decay_t<T>> requires (std::is_same_v<U, FMassFragmentBitSet> || std::is_same_v<U, FMassTagBitSet>)
 	bool OnCompositionChanged(const FMassEntityHandle Entity, T&& BitSet
 		, const EMassObservedOperation Operation, UE::Mass::FProcessingContext* ProcessingContext = nullptr)
 	{
-		return OnCompositionChanged(Entity, FMassArchetypeCompositionDescriptor(MoveTemp(BitSet)), Operation, ProcessingContext);
+		return OnCompositionChanged(Entity, FMassElementBitSet(MoveTemp(BitSet)), Operation, ProcessingContext);
 	}
 
 	MASSENTITY_API void AddObserverInstance(TNotNull<const UScriptStruct*> ElementType, EMassObservedOperationFlags OperationFlags, TNotNull<UMassProcessor*> ObserverProcessor);
@@ -218,7 +201,6 @@ protected:
 
 	MASSENTITY_API void OnModulePackagesUnloaded(TConstArrayView<UPackage*> Packages);
 
-	//TSharedRef<FObserverLock> GetOrMakeObserverLock(TConstArrayView<FMassEntityHandle> ReservedEntities, FMassArchetypeEntityCollection&& EntityCollection);
 	MASSENTITY_API TSharedRef<FObserverLock> GetOrMakeObserverLock();
 	TSharedPtr<FObserverLock> GetObserverLock() const;
 	bool IsLocked() const { return ActiveObserverLock.IsValid(); }
@@ -250,14 +232,11 @@ protected:
 	int32 LocksCount = 0;
 	TWeakPtr<FCreationContext> ActiveCreationContext;
 
-	FMassFragmentBitSet ObservedFragments[(uint8)EMassObservedOperation::MAX];
-	FMassTagBitSet ObservedTags[(uint8)EMassObservedOperation::MAX];
 
-	UPROPERTY()
-	FMassObserversMap FragmentObservers[(uint8)EMassObservedOperation::MAX];
+	FMassElementBitSet ObservedElements[static_cast<uint8>(EMassObservedOperation::MAX)];
 
 	UPROPERTY()
-	FMassObserversMap TagObservers[(uint8)EMassObservedOperation::MAX];
+	FMassObserversMap ElementObservers[static_cast<uint8>(EMassObservedOperation::MAX)];
 
 	FDelegateHandle ModulesUnloadedHandle;
 
diff --git a/Engine/Source/Runtime/MassEntity/Public/MassObserverNotificationTypes.h b/Engine/Source/Runtime/MassEntity/Public/MassObserverNotificationTypes.h
--- a/Engine/Source/Runtime/MassEntity/Public/MassObserverNotificationTypes.h
+++ b/Engine/Source/Runtime/MassEntity/Public/MassObserverNotificationTypes.h
@@ -41,38 +41,41 @@ namespace UE::Mass
 			{
 			};
 			EMassObservedOperation Type;
-			using FCompositionDescription = TVariant<FEmptyComposition, FMassArchetypeCompositionDescriptor, FMassFragmentBitSet, FMassTagBitSet>;
-			FCompositionDescription CompositionChange;
+			FMassElementBitSet CompositionChange;
 			using FEntitiesContainer = TVariant<FEntityCollection, FMassEntityHandle>;
 			FEntitiesContainer AffectedEntities;
 
-			template<typename TComposition>
-			FBufferedNotification(const EMassObservedOperation InType, TComposition&& Composition, FEntitiesContainer&& Entities)
+			FBufferedNotification(const EMassObservedOperation InType, FMassElementBitSet&& Composition, FEntitiesContainer&& Entities)
 				: Type(InType)
-				, CompositionChange(TInPlaceType<typename TDecay<TComposition>::Type>(), Forward<TComposition>(Composition))
+				, CompositionChange(Forward<FMassElementBitSet>(Composition))
 				, AffectedEntities(MoveTemp(Entities))
 			{
 			}
 
-			template<typename TComposition, typename TEntities>
-			FBufferedNotification(const EMassObservedOperation InType, TComposition&& Composition, TEntities&& Entities)
+			template<typename TEntities>
+			FBufferedNotification(const EMassObservedOperation InType, FMassElementBitSet&& Composition, TEntities&& Entities)
 				: Type(InType)
-				, CompositionChange(TInPlaceType<typename TDecay<TComposition>::Type>(), Forward<TComposition>(Composition))
+				, CompositionChange(MoveTemp(Composition))
 				, AffectedEntities(TInPlaceType<typename TDecay<TEntities>::Type>(), Forward<TEntities>(Entities))
 			{
 			}
 
-			template<typename TEntities>
-			FBufferedNotification(const EMassObservedOperation InType, FCompositionDescription&& Composition, TEntities&& Entities)
+			FBufferedNotification(const EMassObservedOperation InType, FMassElementBitSet&& Composition, const FMassArchetypeEntityCollection& Entities)
 				: Type(InType)
 				, CompositionChange(MoveTemp(Composition))
+				, AffectedEntities(TInPlaceType<typename TDecay<FEntityCollection>::Type>(), Entities)
+			{
+			}
+
+			template<typename TEntities>
+			FBufferedNotification(const EMassObservedOperation InType, TEntities&& Entities)
+				: Type(InType)
 				, AffectedEntities(TInPlaceType<typename TDecay<TEntities>::Type>(), Forward<TEntities>(Entities))
 			{
 			}
 
-			FBufferedNotification(const EMassObservedOperation InType, FCompositionDescription&& Composition, const FMassArchetypeEntityCollection& Entities)
+			FBufferedNotification(const EMassObservedOperation InType, const FMassArchetypeEntityCollection& Entities)
 				: Type(InType)
-				, CompositionChange(MoveTemp(Composition))
 				, AffectedEntities(TInPlaceType<typename TDecay<FEntityCollection>::Type>(), Entities)
 			{
 			}
@@ -144,25 +147,9 @@ namespace UE::Mass
 				}
 			}
 
-			static bool AreCompositionsEqual(const FCompositionDescription& A, const FCompositionDescription& B)
+			static bool AreCompositionsEqual(const FMassElementBitSet& A, const FMassElementBitSet& B)
 			{
-				if (A.GetIndex() == B.GetIndex())
-				{
-					switch (A.GetIndex())
-					{
-						case FCompositionDescription::IndexOfType<FEmptyComposition>():
-							return true;
-						case FCompositionDescription::IndexOfType<FMassArchetypeCompositionDescriptor>():
-							return A.Get<FMassArchetypeCompositionDescriptor>().IsIdentical(B.Get<FMassArchetypeCompositionDescriptor>());
-						case FCompositionDescription::IndexOfType<FMassFragmentBitSet>():
-							return A.Get<FMassFragmentBitSet>() == B.Get<FMassFragmentBitSet>();
-						case FCompositionDescription::IndexOfType<FMassTagBitSet>():
-							return A.Get<FMassTagBitSet>() == B.Get<FMassTagBitSet>();
-						default:
-							return false;
-					}
-				}
-				return false;
+				return A.IsEquivalent(B);
 			}
 
 		private:
@@ -252,9 +239,7 @@ namespace UE::Mass
 				if (CreationNotificationIndex == INDEX_NONE)
 				{
 					CreationNotificationIndex = BufferedNotifications.Num();
-					BufferedNotifications.Emplace(EMassObservedOperation::CreateEntity
-						, FBufferedNotification::FEmptyComposition{}
-						, UE::Mass::FEntityCollection());
+					BufferedNotifications.Emplace(EMassObservedOperation::CreateEntity, UE::Mass::FEntityCollection());
 				}
 				return CreationNotificationIndex;
 			}
@@ -281,9 +266,7 @@ namespace UE::Mass
 				if (CreationNotificationIndex == INDEX_NONE)
 				{
 					CreationNotificationIndex = BufferedNotifications.Num();
-					BufferedNotifications.Emplace(EMassObservedOperation::CreateEntity
-						, FBufferedNotification::FEmptyComposition{}
-						, CreatedEntity);
+					BufferedNotifications.Emplace(EMassObservedOperation::CreateEntity, CreatedEntity);
 				}
 				else
 				{
@@ -299,9 +282,7 @@ namespace UE::Mass
 				if (CreationNotificationIndex == INDEX_NONE)
 				{
 					CreationNotificationIndex = BufferedNotifications.Num();
-					BufferedNotifications.Emplace(EMassObservedOperation::CreateEntity
-						, FBufferedNotification::FEmptyComposition{}
-						, UE::Mass::FEntityCollection(InCreatedEntities));
+					BufferedNotifications.Emplace(EMassObservedOperation::CreateEntity, UE::Mass::FEntityCollection(InCreatedEntities));
 				}
 				else
 				{
@@ -317,7 +298,7 @@ namespace UE::Mass
 				if (CreationNotificationIndex == INDEX_NONE)
 				{
 					CreationNotificationIndex = BufferedNotifications.Num();
-					BufferedNotifications.Emplace(EMassObservedOperation::CreateEntity, FBufferedNotification::FEmptyComposition{}
+					BufferedNotifications.Emplace(EMassObservedOperation::CreateEntity
 						, UE::Mass::FEntityCollection(InCreatedEntities, Forward<FMassArchetypeEntityCollection>(InEntityCollection)));
 				}
 				else
@@ -335,9 +316,7 @@ namespace UE::Mass
 				if (CreationNotificationIndex == INDEX_NONE)
 				{
 					CreationNotificationIndex = BufferedNotifications.Num();
-					BufferedNotifications.Emplace(EMassObservedOperation::CreateEntity
-						, FBufferedNotification::FEmptyComposition{}
-						, UE::Mass::FEntityCollection(Forward<T>(InEntityCollection)));
+					BufferedNotifications.Emplace(EMassObservedOperation::CreateEntity, UE::Mass::FEntityCollection(Forward<T>(InEntityCollection)));
 				}
 				else
 				{
@@ -348,27 +327,17 @@ namespace UE::Mass
 			}
 
 			template<typename TEntities>
-			void AddNotification(const EMassObservedOperation OperationType
-				, TEntities&& Entities
-				, const bool bHasFragmentsOverlap, FMassFragmentBitSet&& FragmentOverlap
-				, const bool bHasTagsOverlap, FMassTagBitSet&& TagOverlap)
+			void AddNotification(const EMassObservedOperation OperationType, TEntities&& Entities, FMassElementBitSet&& OverlapBitSet)
 			{
-				checkSlow(bHasFragmentsOverlap || bHasTagsOverlap);
-				FBufferedNotification::FCompositionDescription CompositionChange = (bHasFragmentsOverlap != bHasTagsOverlap)
-					? (bHasFragmentsOverlap
-						? FBufferedNotification::FCompositionDescription(TInPlaceType<FMassFragmentBitSet>(), MoveTemp(FragmentOverlap))
-						: FBufferedNotification::FCompositionDescription(TInPlaceType<FMassTagBitSet>(), MoveTemp(TagOverlap)))
-					: FBufferedNotification::FCompositionDescription(TInPlaceType<FMassArchetypeCompositionDescriptor>(), FMassArchetypeCompositionDescriptor(MoveTemp(FragmentOverlap), MoveTemp(TagOverlap), {}, {}, {}));
-
-				if (BufferedNotifications.Num() 
+				if (BufferedNotifications.Num()
 					&& BufferedNotifications.Last().Type == OperationType
-					&& FBufferedNotification::AreCompositionsEqual(BufferedNotifications.Last().CompositionChange, CompositionChange))
+					&& FBufferedNotification::AreCompositionsEqual(BufferedNotifications.Last().CompositionChange, OverlapBitSet))
 				{
 					BufferedNotifications.Last().AppendEntities(Forward<TEntities>(Entities));
 				}
 				else
 				{
-					BufferedNotifications.Emplace(OperationType, MoveTemp(CompositionChange), Forward<TEntities>(Entities));
+					BufferedNotifications.Emplace(OperationType, MoveTemp(OverlapBitSet), Forward<TEntities>(Entities));
 				}
 			}
 
diff --git a/Engine/Source/Runtime/MassEntity/Public/MassObserverRegistry.h b/Engine/Source/Runtime/MassEntity/Public/MassObserverRegistry.h
--- a/Engine/Source/Runtime/MassEntity/Public/MassObserverRegistry.h
+++ b/Engine/Source/Runtime/MassEntity/Public/MassObserverRegistry.h
@@ -59,14 +59,14 @@ public:
 	static UMassObserverRegistry& GetMutable() { return *GetMutableDefault<UMassObserverRegistry>(); }
 	static const UMassObserverRegistry& Get() { return *GetDefault<UMassObserverRegistry>(); }
 
-	UE_API void RegisterObserver(TNotNull<const UScriptStruct*> ObservedType, uint8 OperationFlags, TSubclassOf<UMassProcessor> ObserverClass);
+	UE_API void RegisterObserver(TNotNull<const UScriptStruct*> ObservedType, uint8 OperationFlags, const TSubclassOf<UMassProcessor>& ObserverClass);
 
-	void RegisterObserver(TNotNull<const UScriptStruct*> ObservedType, EMassObservedOperationFlags OperationFlags, TSubclassOf<UMassProcessor> ObserverClass)
+	void RegisterObserver(TNotNull<const UScriptStruct*> ObservedType, EMassObservedOperationFlags OperationFlags, const TSubclassOf<UMassProcessor>& ObserverClass)
 	{
 		RegisterObserver(ObservedType, static_cast<uint8>(OperationFlags), ObserverClass);
 	}
 
-	UE_API void RegisterObserver(const UScriptStruct& ObservedType, EMassObservedOperation Operation, TSubclassOf<UMassProcessor> ObserverClass);
+	UE_API void RegisterObserver(const UScriptStruct& ObservedType, EMassObservedOperation Operation, const TSubclassOf<UMassProcessor>& ObserverClass);
 
 	using FObserverClassesMap = TMap<TObjectKey<const UScriptStruct>, TArray<FSoftClassPath>>;
 protected:
@@ -76,14 +76,13 @@ protected:
 	FDelegateHandle ModulesUnloadedHandle;
 
 	friend FMassObserverManager;
-	FObserverClassesMap FragmentObserverMaps[static_cast<uint8>(EMassObservedOperation::MAX)];
-	FObserverClassesMap TagObserverMaps[static_cast<uint8>(EMassObservedOperation::MAX)];
+	FObserverClassesMap ElementObserverMaps[static_cast<uint8>(EMassObservedOperation::MAX)];
 
-	UE_DEPRECATED(5.7, "Use FragmentObserverMaps instead")
+	UE_DEPRECATED(5.7, "Use ElementObserverMaps instead")
 	UPROPERTY()
 	FMassEntityObserverClassesMap FragmentObservers[(uint8)EMassObservedOperation::MAX];
 
-	UE_DEPRECATED(5.7, "Use TagObserverMaps instead")
+	UE_DEPRECATED(5.7, "Use ElementObserverMaps instead")
 	UPROPERTY()
 	FMassEntityObserverClassesMap TagObservers[(uint8)EMassObservedOperation::MAX];
 };
diff --git a/Engine/Source/Runtime/MassEntity/Public/MassRequirements.h b/Engine/Source/Runtime/MassEntity/Public/MassRequirements.h
--- a/Engine/Source/Runtime/MassEntity/Public/MassRequirements.h
+++ b/Engine/Source/Runtime/MassEntity/Public/MassRequirements.h
@@ -660,8 +660,16 @@ struct FMassFragmentRequirements
 	bool HasLinkedEntityRequirements() const;
 
 	MASSENTITY_API bool DoesArchetypeMatchRequirements(const FMassArchetypeHandle& ArchetypeHandle) const;
-	MASSENTITY_API bool DoesArchetypeMatchRequirements(const FMassArchetypeCompositionDescriptor& ArchetypeComposition) const;
-	MASSENTITY_API bool DoesMatchAnyOptionals(const FMassArchetypeCompositionDescriptor& ArchetypeComposition) const;
+	MASSENTITY_API bool DoesArchetypeMatchRequirements(const FMassElementBitSet& ArchetypeCompositionBitSet) const;
+	bool DoesArchetypeMatchRequirements(const FMassArchetypeCompositionDescriptor& ArchetypeComposition) const
+	{
+		return DoesArchetypeMatchRequirements(ArchetypeComposition.GetElementsBitSet());
+	}
+	MASSENTITY_API bool DoesMatchAnyOptionals(const FMassElementBitSet& ArchetypeCompositionBitSet) const;
+	bool DoesMatchAnyOptionals(const FMassArchetypeCompositionDescriptor& ArchetypeComposition) const
+	{
+		return DoesMatchAnyOptionals(ArchetypeComposition.GetElementsBitSet());
+	}
 
 	bool DoesRequireGameThreadExecution() const;
 	MASSENTITY_API void ExportRequirements(FMassExecutionRequirements& OutRequirements) const;
