[Mass] fixed buggy checkSlow condition submitted in yesterday's "single elements bitset" change

diff --git a/Engine/Source/Runtime/MassEntity/Private/MassArchetypeData.cpp b/Engine/Source/Runtime/MassEntity/Private/MassArchetypeData.cpp
--- a/Engine/Source/Runtime/MassEntity/Private/MassArchetypeData.cpp
+++ b/Engine/Source/Runtime/MassEntity/Private/MassArchetypeData.cpp
@@ -105,10 +105,8 @@ void FMassArchetypeData::InitializeWithSimilar(const FMassEntityManager& EntityM
 	checkf(IsInitialized() == false, TEXT("Trying to InitializeWithSimilar but this archetype has already been initialized"));
 
 	// NewComposition.GetElementsBitSet() - BaseArchetype.GetCompositionBitSet() -> should contain only tags and fragments
-	// BaseArchetype.GetCompositionBitSet() - NewComposition.GetElementsBitSet() -> should be empty
-	bool bHasTagsAndElements = FMassElementBitSet(NewComposition.GetElementsBitSet() - BaseArchetype.GetCompositionBitSet()).ContainsOnlyTagsAndFragments();
-	bool bNothingElse = FMassElementBitSet(BaseArchetype.GetCompositionBitSet() - NewComposition.GetElementsBitSet()).IsEmpty();
-	checkfSlow(bHasTagsAndElements && bNothingElse, TEXT("NewComposition is expected to differ from the BaseArchetype's composition just by tags and fragments"))
+	bool bHasOnlyTagsAndElements = FMassElementBitSet(NewComposition.GetElementsBitSet() - BaseArchetype.GetCompositionBitSet()).ContainsOnlyTagsAndFragments();
+	checkfSlow(bHasOnlyTagsAndElements, TEXT("NewComposition is expected to differ from the BaseArchetype's composition just by tags and fragments"))
 
 	CreatedArchetypeDataVersion = ArchetypeDataVersion;
 
@@ -1526,4 +1524,4 @@ bool FMassArchetypeHelper::DoesArchetypeMatchRequirements(const FMassArchetypeDa
 
 	return false;
 }
-#endif // WITH_MASSENTITY_DEBUG
\ No newline at end of file
+#endif // WITH_MASSENTITY_DEBUG
