[Mass] fixed buggy checkSlow condition submitted in yesterday's "single elements bitset" change

#rnx

diff --git a/Engine/Source/Runtime/MassEntity/Private/MassArchetypeData.cpp b/Engine/Source/Runtime/MassEntity/Private/MassArchetypeData.cpp
--- a/Engine/Source/Runtime/MassEntity/Private/MassArchetypeData.cpp
+++ b/Engine/Source/Runtime/MassEntity/Private/MassArchetypeData.cpp
@@ -105,8 +105,10 @@ void FMassArchetypeData::InitializeWithSimilar(const FMassEntityManager& EntityM
 	checkf(IsInitialized() == false, TEXT("Trying to InitializeWithSimilar but this archetype has already been initialized"));
 
 	// NewComposition.GetElementsBitSet() - BaseArchetype.GetCompositionBitSet() -> should contain only tags and fragments
-	bool bHasOnlyTagsAndElements = FMassElementBitSet(NewComposition.GetElementsBitSet() - BaseArchetype.GetCompositionBitSet()).ContainsOnlyTagsAndFragments();
-	checkfSlow(bHasOnlyTagsAndElements, TEXT("NewComposition is expected to differ from the BaseArchetype's composition just by tags and fragments"))
+	// BaseArchetype.GetCompositionBitSet() - NewComposition.GetElementsBitSet() -> should contain only tags and fragments
+	bool bAddsOnlyTagsAndElements = FMassElementBitSet(NewComposition.GetElementsBitSet() - BaseArchetype.GetCompositionBitSet()).ContainsOnlyTagsAndFragments();
+	bool bNothingElse = FMassElementBitSet(BaseArchetype.GetCompositionBitSet() - NewComposition.GetElementsBitSet()).ContainsOnlyTagsAndFragments();
+	checkf(bAddsOnlyTagsAndElements && bNothingElse, TEXT("NewComposition is expected to differ from the BaseArchetype's composition just by tags and fragments"))
 
 	CreatedArchetypeDataVersion = ArchetypeDataVersion;
 
