[MassEntity]
* Added missing include in MassElement.h

diff --git a/Engine/Source/Runtime/MassEntity/Public/MassElement.h b/Engine/Source/Runtime/MassEntity/Public/MassElement.h
--- a/Engine/Source/Runtime/MassEntity/Public/MassElement.h
+++ b/Engine/Source/Runtime/MassEntity/Public/MassElement.h
@@ -2,6 +2,7 @@
 
 #pragma once
 
+#include "Containers/StaticArray.h"
 #include "StructUtils/StructTypeBitSet.h"
 #include "MassEntityElementTypes.h"
 #include "Subsystems/Subsystem.h"
