Expose mass processing phase tick functions in public API to allow for setup of tick dependencies

diff --git a/Engine/Source/Runtime/MassEntity/Public/MassProcessingPhaseManager.h b/Engine/Source/Runtime/MassEntity/Public/MassProcessingPhaseManager.h
--- a/Engine/Source/Runtime/MassEntity/Public/MassProcessingPhaseManager.h
+++ b/Engine/Source/Runtime/MassEntity/Public/MassProcessingPhaseManager.h
@@ -215,6 +215,20 @@ public:
 		void Reset();
 	};
 
+	FTickFunction& GetProcessingPhaseTickFunction(EMassProcessingPhase Phase)
+	{
+		check(Phase < EMassProcessingPhase::MAX);
+		check((int)Phase >= 0 );
+		return ProcessingPhases[(int)Phase];
+	}
+	
+	const FTickFunction& GetProcessingPhaseTickFunction(EMassProcessingPhase Phase) const
+	{
+		check(Phase < EMassProcessingPhase::MAX);
+		check((int)Phase >= 0 );
+		return ProcessingPhases[(int)Phase];
+	}
+
 #if WITH_MASSENTITY_DEBUG
 	TConstArrayView<FMassProcessingPhase>  DebugGetProcessingPhases() const;
 	TConstArrayView<FPhaseGraphBuildState> DebugGetProcessingGraphBuildStates() const;
