[Mass] Updated Mass queries with the ability to create a topology hash.

This change adds a way to calculate a hash of the topology of the archetypes associated with a query. Whenever this hash changes it means that one or more rows were added, removed or reordered in one or more of the archetypes. In practice this is based on the entity order revision so the amount of work per query is very limited.

#Synced-CL 47889446

diff --git a/Engine/Source/Runtime/MassEntity/Private/MassEntityQuery.cpp b/Engine/Source/Runtime/MassEntity/Private/MassEntityQuery.cpp
--- a/Engine/Source/Runtime/MassEntity/Private/MassEntityQuery.cpp
+++ b/Engine/Source/Runtime/MassEntity/Private/MassEntityQuery.cpp
@@ -12,6 +12,7 @@
 #include "MassProcessor.h"
 #include "MassEntityTrace.h"
 #include "MassRequirementAccessDetector.h"
+#include "Hash/xxhash.h"
 #if WITH_MASSENTITY_DEBUG
 #include "MassDebugLogging.h"
 #endif // WITH_MASSENTITY_DEBUG
@@ -600,6 +601,27 @@ void FMassEntityQuery::ParallelForEachEntityChunk(FMassExecutionContext& Executi
 	}
 }
 
+uint64 FMassEntityQuery::CalculateArchetypeEntityOrderHash()
+{
+	CacheArchetypes();
+
+	// The hash is created as a cumulative hash to make sure that the order of archetypes doesn't matter.
+	uint64 Hash = 0;
+	for (const FMassArchetypeHandle& ArchetypeHandle : ValidArchetypes)
+	{
+		if (const FMassArchetypeData* Archetype = FMassArchetypeHelper::ArchetypeDataFromHandle(ArchetypeHandle))
+		{
+			int32 EntityOrderRevision = Archetype->GetEntityOrderVersion();
+			// Accept the possible overflow as a desired side effect.
+			Hash += FXxHash64::HashBuffer(&EntityOrderRevision, sizeof(int32)).Hash;
+		}
+	}
+	// Further reduce the risk of a hash collision between two different sized archetype lists.
+	int32 ArchetypeCount = ValidArchetypes.Num();
+	Hash += FXxHash64::HashBuffer(&ArchetypeCount, sizeof(int32)).Hash;
+	return Hash;
+}
+
 int32 FMassEntityQuery::GetNumMatchingEntities()
 {
 	CacheArchetypes();
diff --git a/Engine/Source/Runtime/MassEntity/Public/MassEntityQuery.h b/Engine/Source/Runtime/MassEntity/Public/MassEntityQuery.h
--- a/Engine/Source/Runtime/MassEntity/Public/MassEntityQuery.h
+++ b/Engine/Source/Runtime/MassEntity/Public/MassEntityQuery.h
@@ -157,6 +157,13 @@ public:
 		return ValidArchetypes; 
 	}
 
+	/**
+	 * Calculates a hash of the entity order revisions in all archetypes. If returned value ever differs from
+	 * a previous call it indicates that one or more archetypes associated with this query have been updated
+	 * by having entities added, removed or reordered.
+	 */
+	UE_API uint64 CalculateArchetypeEntityOrderHash();
+
 	/** 
 	 * Goes through ValidArchetypes and sums up the number of entities contained in them.
 	 * Note that the function is not const because calling it can result in re-caching of ValidArchetypes 
