[Mass] Fixed warning messages regarding undeclared subsystem access via `FMassSubsystemAccess`

diff --git a/Engine/Source/Runtime/MassEntity/Public/MassSubsystemAccess.h b/Engine/Source/Runtime/MassEntity/Public/MassSubsystemAccess.h
--- a/Engine/Source/Runtime/MassEntity/Public/MassSubsystemAccess.h
+++ b/Engine/Source/Runtime/MassEntity/Public/MassSubsystemAccess.h
@@ -29,9 +29,9 @@ struct FMassSubsystemAccess
 	template<typename T, typename = typename TEnableIf<TIsDerivedFrom<T, USubsystem>::IsDerived>::Type>
 	T* GetMutableSubsystem()
 	{
-		const uint32 SystemIndex = FMassExternalSubsystemBitSet::GetTypeIndex<T>();
+		const int32 SystemIndex = FMassExternalSubsystemBitSet::GetTypeIndex<T>();
 		if (ensureMsgf(MutableSubsystemsBitSet.IsBitSet(SystemIndex)
-			, TEXT("Missing 'template<> struct TMassExternalSubsystemTraits<U%s>' for this subsystem"), *GetNameSafe(T::StaticClass())))
+			, TEXT("Undeclared read/write access to subsystem %s"), *GetNameSafe(T::StaticClass())))
 		{
 			return GetSubsystemInternal<T>(SystemIndex);
 		}
@@ -50,9 +50,9 @@ struct FMassSubsystemAccess
 	template<typename T, typename = typename TEnableIf<TIsDerivedFrom<T, USubsystem>::IsDerived>::Type>
 	const T* GetSubsystem()
 	{
-		const uint32 SystemIndex = FMassExternalSubsystemBitSet::GetTypeIndex<T>();
+		const int32 SystemIndex = FMassExternalSubsystemBitSet::GetTypeIndex<T>();
 		if (ensureMsgf(ConstSubsystemsBitSet.IsBitSet(SystemIndex) || MutableSubsystemsBitSet.IsBitSet(SystemIndex)
-			, TEXT("Missing 'template<> struct TMassExternalSubsystemTraits<U%s>' for this subsystem"), *GetNameSafe(T::StaticClass())))
+			, TEXT("Undeclared read-only access to subsystem %s"), *GetNameSafe(T::StaticClass())))
 		{
 			return GetSubsystemInternal<T>(SystemIndex);
 		}
@@ -73,9 +73,10 @@ struct FMassSubsystemAccess
 	template<typename T, typename = typename TEnableIf<TIsDerivedFrom<T, USubsystem>::IsDerived>::Type>
 	T* GetMutableSubsystem(const TSubclassOf<USubsystem> SubsystemClass)
 	{
-		const uint32 SystemIndex = FMassExternalSubsystemBitSet::GetTypeIndex(**SubsystemClass);
+		check(SubsystemClass);
+		const int32 SystemIndex = FMassExternalSubsystemBitSet::GetTypeIndex(**SubsystemClass);
 		if (ensureMsgf(MutableSubsystemsBitSet.IsBitSet(SystemIndex)
-			, TEXT("Missing 'template<> struct TMassExternalSubsystemTraits<U%s>' for this subsystem"), *GetNameSafe(SubsystemClass)))
+			, TEXT("Undeclared read/write access to subsystem %s"), *GetNameSafe(SubsystemClass)))
 		{
 			return GetSubsystemInternal<T>(SystemIndex, SubsystemClass);
 		}
@@ -86,6 +87,7 @@ struct FMassSubsystemAccess
 	template<typename T, typename = typename TEnableIf<TIsDerivedFrom<T, USubsystem>::IsDerived>::Type>
 	T& GetMutableSubsystemChecked(const TSubclassOf<USubsystem> SubsystemClass)
 	{
+		check(SubsystemClass);
 		T* InstancePtr = GetMutableSubsystem<T>(SubsystemClass);
 		check(InstancePtr);
 		return *InstancePtr;
@@ -94,9 +96,10 @@ struct FMassSubsystemAccess
 	template<typename T, typename = typename TEnableIf<TIsDerivedFrom<T, USubsystem>::IsDerived>::Type>
 	const T* GetSubsystem(const TSubclassOf<USubsystem> SubsystemClass)
 	{
-		const uint32 SystemIndex = FMassExternalSubsystemBitSet::GetTypeIndex(**SubsystemClass);
+		check(SubsystemClass);
+		const int32 SystemIndex = FMassExternalSubsystemBitSet::GetTypeIndex(**SubsystemClass);
 		if (ensureMsgf(ConstSubsystemsBitSet.IsBitSet(SystemIndex) || MutableSubsystemsBitSet.IsBitSet(SystemIndex)
-			, TEXT("Missing 'template<> struct TMassExternalSubsystemTraits<U%s>' for this subsystem"), *GetNameSafe(SubsystemClass)))
+			, TEXT("Undeclared read-only access to subsystem %s"), *GetNameSafe(SubsystemClass)))
 		{
 			return GetSubsystemInternal<T>(SystemIndex, SubsystemClass);
 		}
@@ -106,6 +109,7 @@ struct FMassSubsystemAccess
 	template<typename T, typename = typename TEnableIf<TIsDerivedFrom<T, USubsystem>::IsDerived>::Type>
 	const T& GetSubsystemChecked(const TSubclassOf<USubsystem> SubsystemClass)
 	{
+		check(SubsystemClass);
 		const T* InstancePtr = GetSubsystem<T>(SubsystemClass);
 		check(InstancePtr);
 		return *InstancePtr;
@@ -186,11 +190,12 @@ struct FMassSubsystemAccess
 
 protected:
 	template<typename T>
-	T* GetSubsystemInternal(const uint32 SystemIndex)
+	T* GetSubsystemInternal(const int32 SystemIndex)
 	{
+		checkSlow(SystemIndex != INDEX_NONE);
 		if (UNLIKELY(Subsystems.IsValidIndex(SystemIndex) == false))
 		{
-			Subsystems.AddZeroed(Subsystems.Num() - SystemIndex + 1);
+			Subsystems.AddZeroed(SystemIndex - Subsystems.Num() + 1);
 		}
 
 		T* SystemInstance = (T*)Subsystems[SystemIndex];
@@ -210,11 +215,12 @@ protected:
 	}
 
 	template<typename T>
-	T* GetSubsystemInternal(const uint32 SystemIndex, const TSubclassOf<USubsystem> SubsystemClass)
+	T* GetSubsystemInternal(const int32 SystemIndex, const TSubclassOf<USubsystem> SubsystemClass)
 	{
+		checkSlow(SystemIndex != INDEX_NONE);
 		if (UNLIKELY(Subsystems.IsValidIndex(SystemIndex) == false))
 		{
-			Subsystems.AddZeroed(Subsystems.Num() - SystemIndex + 1);
+			Subsystems.AddZeroed(SystemIndex - Subsystems.Num() + 1);
 		}
 
 		USubsystem* SystemInstance = (T*)Subsystems[SystemIndex];
