[Mass] Issue 1155999: Compile errors in unknown
Added forward declaration for the missing type.

#rnx
#Synced-CL 48944341

diff --git a/Engine/Source/Runtime/MassEntity/Public/MassDebugger.h b/Engine/Source/Runtime/MassEntity/Public/MassDebugger.h
--- a/Engine/Source/Runtime/MassEntity/Public/MassDebugger.h
+++ b/Engine/Source/Runtime/MassEntity/Public/MassDebugger.h
@@ -449,6 +449,7 @@ struct FMassArchetypeHandle;
 struct FMassFragmentRequirements;
 struct FMassFragmentRequirementDescription;
 struct FMassArchetypeCompositionDescriptor;
+struct FMassExecutionContext;
 
 struct FMassDebugger
 {
