[Mass] Marked MassEntityHandle as a BlueprintType

diff --git a/Engine/Source/Runtime/MassEntity/Public/MassEntityHandle.h b/Engine/Source/Runtime/MassEntity/Public/MassEntityHandle.h
--- a/Engine/Source/Runtime/MassEntity/Public/MassEntityHandle.h
+++ b/Engine/Source/Runtime/MassEntity/Public/MassEntityHandle.h
@@ -8,7 +8,7 @@
  * A handle to a Mass entity. An entity is used in conjunction with the FMassEntityManager
  * for the current world and can contain lightweight fragments.
  */
-USTRUCT()
+USTRUCT(BlueprintType)
 struct alignas(8) FMassEntityHandle
 {
 	GENERATED_BODY()
