[Mass] Fixed TSharedFragmentsContainer::template FindOrAdd's compilation error (Make call was missing qualification)

#rnx

#Synced-CL 49028421

diff --git a/Engine/Source/Runtime/MassEntity/Public/MassEntityManager.h b/Engine/Source/Runtime/MassEntity/Public/MassEntityManager.h
--- a/Engine/Source/Runtime/MassEntity/Public/MassEntityManager.h
+++ b/Engine/Source/Runtime/MassEntity/Public/MassEntityManager.h
@@ -1272,7 +1272,7 @@ private:
 		TSharedStructType& FindOrAdd(const uint32 Hash, TArgs&&... InArgs)
 		{
 			// Need to actually construct the struct to make proper comparison to possible existing instance
-			TSharedStructType TempInstance = Make<T>(Forward<TArgs>(InArgs)...);
+			TSharedStructType TempInstance = TSharedStructType::template Make<T>(Forward<TArgs>(InArgs)...);
 
 			for (TMultiMap<uint32, int32>::TConstKeyIterator It = HashToInstanceIndexMap.CreateConstKeyIterator(Hash); It; ++It)
 			{
