PR #13077: Wide Mass task dispatch calls

#Synced-CL 49265344

diff --git a/Engine/Source/Runtime/MassEntity/Private/MassExecutor.cpp b/Engine/Source/Runtime/MassEntity/Private/MassExecutor.cpp
--- a/Engine/Source/Runtime/MassEntity/Private/MassExecutor.cpp
+++ b/Engine/Source/Runtime/MassEntity/Private/MassExecutor.cpp
@@ -152,6 +152,53 @@ private:
 	ENamedThreads::Type DesiredThread;
 };
 
+struct FMassExecutorDispatchTask
+{
+	FMassExecutorDispatchTask(const FMassExecutionContext& InExecutionContext, const TSharedPtr<FMassEntityManager>& InEntityManager, UMassProcessor& InProcessor
+		, TFunction<void()> InOnDoneNotification, const FString& InDebugName, ENamedThreads::Type InDesiredEndTaskThread)
+		: ExecutionContext(InExecutionContext)
+		, OnDoneNotification(InOnDoneNotification)
+		, DebugName(InDebugName)
+		, EntityManager(InEntityManager)
+		, Processor(InProcessor)
+		, DesiredEndTaskThread(InDesiredEndTaskThread)
+	{
+	}
+	static TStatId GetStatId()
+	{
+		RETURN_QUICK_DECLARE_CYCLE_STAT(FMassExecutorDispatchTask, STATGROUP_TaskGraphTasks);
+	}
+
+	static ENamedThreads::Type GetDesiredThread() { return ENamedThreads::AnyHiPriThreadHiPriTask; }
+	static ESubsequentsMode::Type GetSubsequentsMode() { return ESubsequentsMode::TrackSubsequents; }
+	void DoTask(ENamedThreads::Type /*CurrentThread*/, const FGraphEventRef& MyCompletionGraphEvent)
+	{
+		TRACE_CPUPROFILER_EVENT_SCOPE_STR("Dispatch Processors Async")
+		FGraphEventRef CompletionEvent;
+		FMassExecutionContext ExecutionContextThread = ExecutionContext;
+		CompletionEvent = Processor.DispatchProcessorTasks(EntityManager, ExecutionContextThread, {});
+		if (CompletionEvent.IsValid())
+		{
+			const FGraphEventArray Prerequisites = { CompletionEvent };
+
+			CompletionEvent = TGraphTask<FMassExecutorDoneTask>::CreateTask(&Prerequisites)
+				.ConstructAndDispatchWhenReady(MoveTemp(ExecutionContextThread), OnDoneNotification, Processor.GetName(), DesiredEndTaskThread);
+		}
+
+		if (CompletionEvent.IsValid())
+		{
+			MyCompletionGraphEvent->DontCompleteUntil(CompletionEvent);
+		}
+	}
+private:
+	FMassExecutionContext ExecutionContext;
+	TFunction<void()> OnDoneNotification;
+	FString DebugName;
+	TSharedPtr<FMassEntityManager> EntityManager;
+	UMassProcessor& Processor;
+	ENamedThreads::Type DesiredEndTaskThread;
+};
+
 FGraphEventRef TriggerParallelTasks(UMassProcessor& Processor, FProcessingContext&& ProcessingContext, TFunction<void()> OnDoneNotification, ENamedThreads::Type CurrentThread)
 {
 	TRACE_CPUPROFILER_EVENT_SCOPE(MassExecutor_RunParallel);
@@ -160,20 +207,16 @@ FGraphEventRef TriggerParallelTasks(UMassProcessor& Processor, FProcessingContex
 	// flushing stored commands. 
 	FMassExecutionContext ExecutionContext = MoveTemp(ProcessingContext).GetExecutionContext();
 
-	FGraphEventRef CompletionEvent;
+	FGraphEventRef DispatchEvent;
 	{
 		TRACE_CPUPROFILER_EVENT_SCOPE_STR("Dispatch Processors")
-		CompletionEvent = Processor.DispatchProcessorTasks(ProcessingContext.GetEntityManager(), ExecutionContext, {});
-	}
-
-	if (CompletionEvent.IsValid())
-	{
-		const FGraphEventArray Prerequisites = { CompletionEvent };
-		CompletionEvent = TGraphTask<FMassExecutorDoneTask>::CreateTask(&Prerequisites)
-			.ConstructAndDispatchWhenReady(MoveTemp(ExecutionContext), OnDoneNotification, Processor.GetName(), CurrentThread);
+		FGraphEventArray Prerequisites;
+		TSharedPtr<FMassEntityManager> EntityManager = ProcessingContext.EntityManager;
+		DispatchEvent  = TGraphTask<FMassExecutorDispatchTask>::CreateTask(&Prerequisites)
+			.ConstructAndDispatchWhenReady(ExecutionContext, EntityManager, Processor, OnDoneNotification, Processor.GetName(), CurrentThread);
 	}
 
-	return CompletionEvent;
+	return DispatchEvent;
 }
 
 //-----------------------------------------------------------------------------
