[Mass] de-modulo in perf critical Mass code

#rnx
#Synced-CL 49265344

diff --git a/Engine/Source/Runtime/MassEntity/Internal/MassArchetypeData.h b/Engine/Source/Runtime/MassEntity/Internal/MassArchetypeData.h
--- a/Engine/Source/Runtime/MassEntity/Internal/MassArchetypeData.h
+++ b/Engine/Source/Runtime/MassEntity/Internal/MassArchetypeData.h
@@ -419,7 +419,7 @@ public:
 		const int32 AbsoluteIndex = EntityMap.FindChecked(EntityIndex);
 		const int32 ChunkIndex = AbsoluteIndex / NumEntitiesPerChunk;
 	
-		return FMassRawEntityInChunkData(Chunks[ChunkIndex].GetRawMemory(), AbsoluteIndex % NumEntitiesPerChunk);
+		return FMassRawEntityInChunkData(Chunks[ChunkIndex].GetRawMemory(), AbsoluteIndex - (NumEntitiesPerChunk * ChunkIndex));
 	}
 
 	FORCEINLINE FMassRawEntityInChunkData MakeRawEntityHandle(const FMassEntityHandle Entity) const
@@ -433,7 +433,7 @@ public:
 		const int32 ChunkIndex = AbsoluteIndex / NumEntitiesPerChunk;
 		const FMassArchetypeChunk& Chunk = Chunks[ChunkIndex];
 
-		return FMassEntityInChunkDataHandle(Chunk.GetRawMemory(), AbsoluteIndex % NumEntitiesPerChunk
+		return FMassEntityInChunkDataHandle(Chunk.GetRawMemory(), AbsoluteIndex - (NumEntitiesPerChunk * ChunkIndex)
 			, ChunkIndex, Chunk.GetSerialModificationNumber());
 	}
 
diff --git a/Engine/Source/Runtime/MassEntity/Private/MassArchetypeData.cpp b/Engine/Source/Runtime/MassEntity/Private/MassArchetypeData.cpp
--- a/Engine/Source/Runtime/MassEntity/Private/MassArchetypeData.cpp
+++ b/Engine/Source/Runtime/MassEntity/Private/MassArchetypeData.cpp
@@ -191,7 +191,8 @@ void FMassArchetypeData::AddEntity(FMassEntityHandle Entity, const FMassArchetyp
 
 	// Initialize fragments
 	const int32 ChunkIndex = AbsoluteIndex / NumEntitiesPerChunk;
-	const int32 IndexWithinChunk = AbsoluteIndex % NumEntitiesPerChunk;
+	const int32 IndexWithinChunk = AbsoluteIndex - (NumEntitiesPerChunk * ChunkIndex);
+
 	FMassArchetypeChunk& Chunk = Chunks[ChunkIndex];
 	for (const FMassArchetypeFragmentConfig& FragmentConfig : FragmentConfigs)
 	{
@@ -288,7 +289,7 @@ void FMassArchetypeData::RemoveEntity(FMassEntityHandle Entity)
 
 	// Destroy fragments
 	const int32 ChunkIndex = AbsoluteIndex / NumEntitiesPerChunk;
-	const int32 IndexWithinChunk = AbsoluteIndex % NumEntitiesPerChunk;
+	const int32 IndexWithinChunk = AbsoluteIndex - (NumEntitiesPerChunk * ChunkIndex);
 	FMassArchetypeChunk& Chunk = Chunks[ChunkIndex];
 
 	for (const FMassArchetypeFragmentConfig& FragmentConfig : FragmentConfigs)
@@ -306,7 +307,7 @@ void FMassArchetypeData::RemoveEntityInternal(const int32 AbsoluteIndex)
 	++EntityOrderVersion;
 
 	const int32 ChunkIndex = AbsoluteIndex / NumEntitiesPerChunk;
-	const int32 IndexWithinChunk = AbsoluteIndex % NumEntitiesPerChunk;
+	const int32 IndexWithinChunk = AbsoluteIndex - (NumEntitiesPerChunk * ChunkIndex);
 
 	FMassArchetypeChunk& Chunk = Chunks[ChunkIndex];
 
@@ -413,7 +414,7 @@ void FMassArchetypeData::SetSharedFragmentsData(const FMassEntityHandle Entity,
 	// Gets the current chunk where the entity is located
 	const int32 OldAbsoluteIndex = EntityMap.FindChecked(Entity.Index);
 	const int32 OldChunkIndex = OldAbsoluteIndex / NumEntitiesPerChunk;
-	const int32 OldIndexWithinChunk = OldAbsoluteIndex % NumEntitiesPerChunk;
+	const int32 OldIndexWithinChunk = OldAbsoluteIndex - (NumEntitiesPerChunk * OldChunkIndex);
 	const FMassArchetypeChunk& OldChunk = Chunks[OldChunkIndex];
 
 	// Gets or adds a new chunk that will hold the new entity with the new shared values
@@ -491,12 +492,12 @@ void FMassArchetypeData::MoveEntityToAnotherArchetype(const FMassEntityHandle En
 
 	const int32 AbsoluteIndex = EntityMap.FindAndRemoveChecked(Entity.Index);
 	const int32 ChunkIndex = AbsoluteIndex / NumEntitiesPerChunk;
-	const int32 IndexWithinChunk = AbsoluteIndex % NumEntitiesPerChunk;
+	const int32 IndexWithinChunk = AbsoluteIndex - (NumEntitiesPerChunk * ChunkIndex);
 	FMassArchetypeChunk& Chunk = Chunks[ChunkIndex];
 
 	const int32 NewAbsoluteIndex = NewArchetype.AddEntityInternal(Entity, SharedFragmentValuesOverride ? *SharedFragmentValuesOverride : Chunk.GetSharedFragmentValues());
 	const int32 NewChunkIndex = NewAbsoluteIndex / NewArchetype.NumEntitiesPerChunk;
-	const int32 NewIndexWithinChunk = NewAbsoluteIndex % NewArchetype.NumEntitiesPerChunk;
+	const int32 NewIndexWithinChunk = NewAbsoluteIndex - (NewArchetype.NumEntitiesPerChunk * NewChunkIndex);
 	FMassArchetypeChunk& NewChunk = NewArchetype.Chunks[NewChunkIndex];
 
 	UE_TRACE_MASS_ARCHETYPE_CREATED(NewArchetype)
diff --git a/Engine/Source/Runtime/MassEntity/Private/MassArchetypeTypes.cpp b/Engine/Source/Runtime/MassEntity/Private/MassArchetypeTypes.cpp
--- a/Engine/Source/Runtime/MassEntity/Private/MassArchetypeTypes.cpp
+++ b/Engine/Source/Runtime/MassEntity/Private/MassArchetypeTypes.cpp
@@ -49,7 +49,7 @@ FORCEINLINE FMassArchetypeEntityCollection::FArchetypeEntityRange FMassArchetype
 
 		const int32 NumEntitiesPerChunk = ArchetypeData->GetNumEntitiesPerChunk();
 		const int32 ChunkIndex = TrueIndex / NumEntitiesPerChunk;
-		const int32 SubchunkStart = TrueIndex % NumEntitiesPerChunk;
+		const int32 SubchunkStart = TrueIndex - (NumEntitiesPerChunk * ChunkIndex);
 		return FArchetypeEntityRange(ChunkIndex, SubchunkStart, 1);
 	}
 
@@ -183,7 +183,7 @@ void FMassArchetypeEntityCollection::BuildEntityRanges(TStridedView<const int32>
 			ChunkEnd = (ChunkIndex + 1) * NumEntitiesPerChunk;
 			SubchunkLen = 0;
 			// new subchunk
-			const int32 SubchunkStart = Index % NumEntitiesPerChunk;
+			const int32 SubchunkStart = Index - (NumEntitiesPerChunk * ChunkIndex);
 			SubChunkPtr = &Ranges.Add_GetRef(FArchetypeEntityRange(ChunkIndex, SubchunkStart));
 		}
 		++SubchunkLen;
diff --git a/Engine/Source/Runtime/MassEntity/Private/MassRelationObservers.cpp b/Engine/Source/Runtime/MassEntity/Private/MassRelationObservers.cpp
--- a/Engine/Source/Runtime/MassEntity/Private/MassRelationObservers.cpp
+++ b/Engine/Source/Runtime/MassEntity/Private/MassRelationObservers.cpp
@@ -208,7 +208,7 @@ bool UMassRelationRoleDestruction::ConfigureRelationObserver(UE::Mass::FTypeHand
 	static_assert(static_cast<uint8>(UE::Mass::ERelationRole::MAX) == 2, "Current implementation relies on there being only two roles.");
 
 	const int32 RoleAsIndex = RelationRole != UE::Mass::ERelationRole::MAX ? static_cast<int32>(RelationRole) : 0;
-	const int32 OppositeRoleIndex = (RoleAsIndex + 1) % 2;
+	const int32 OppositeRoleIndex = static_cast<int32>((RelationRole == UE::Mass::ERelationRole::Object) ? UE::Mass::ERelationRole::Subject : UE::Mass::ERelationRole::Object);
 	bool bExecutionFunctionAssigned = false;
 
 	// this processor is specifically implemented for handling destruction of regular entities
