[Mass] Removed code deprecated in 5.5

#rnx

#Synced-CL 49265344

diff --git a/Engine/Plugins/Runtime/MassGameplay/Source/MassRepresentation/Private/MassRepresentationTypes.cpp b/Engine/Plugins/Runtime/MassGameplay/Source/MassRepresentation/Private/MassRepresentationTypes.cpp
--- a/Engine/Plugins/Runtime/MassGameplay/Source/MassRepresentation/Private/MassRepresentationTypes.cpp
+++ b/Engine/Plugins/Runtime/MassGameplay/Source/MassRepresentation/Private/MassRepresentationTypes.cpp
@@ -123,26 +123,3 @@ bool FStaticMeshInstanceVisualizationDesc::IsValid() const
 	}
 	return false;
 }
-
-//-----------------------------------------------------------------------------
-// FMassISMCSharedDataMap - DEPRECATED
-//-----------------------------------------------------------------------------
-FMassISMCSharedData& FMassISMCSharedDataMap::GetAndMarkDirtyChecked(const uint32 Hash)
-{
-	static FMassISMCSharedData Dummy;
-	return Dummy;
-}
-
-FMassISMCSharedData* FMassISMCSharedDataMap::GetAndMarkDirty(const uint32 Hash)
-{
-	return nullptr;
-}
-
-FMassISMCSharedData* FMassISMCSharedDataMap::Find(const uint32 Hash)
-{
-	return nullptr;
-}
-
-void FMassISMCSharedDataMap::Remove(const uint32 Hash)
-{
-}
diff --git a/Engine/Plugins/Runtime/MassGameplay/Source/MassRepresentation/Public/MassRepresentationTypes.h b/Engine/Plugins/Runtime/MassGameplay/Source/MassRepresentation/Public/MassRepresentationTypes.h
--- a/Engine/Plugins/Runtime/MassGameplay/Source/MassRepresentation/Public/MassRepresentationTypes.h
+++ b/Engine/Plugins/Runtime/MassGameplay/Source/MassRepresentation/Public/MassRepresentationTypes.h
@@ -4,9 +4,6 @@
 
 #include "MassLODTypes.h"
 #include "MassEntityHandle.h"
-#if UE_ENABLE_INCLUDE_ORDER_DEPRECATED_IN_5_5
-#include "MassEntityTypes.h"
-#endif // UE_ENABLE_INCLUDE_ORDER_DEPRECATED_IN_5_5
 #include "Engine/DataTable.h"
 #include "Misc/MTAccessDetector.h"
 #include "InstanceDataTypes.h"
@@ -519,22 +516,6 @@ protected:
 	TBitArray<> DirtyData;
 	/** Indices to Data that are available for reuse */
 	TArray<int32> FreeIndices;
-
-public:
-	UE_DEPRECATED(5.5, "Deprecated. Using hashes in Mass Visualization is being phased out. Use FISMCSharedDataKey instead.")
-	FMassISMCSharedData& GetAndMarkDirtyChecked(const uint32 Hash);
-	UE_DEPRECATED(5.5, "Deprecated. Using hashes in Mass Visualization is being phased out. Use FISMCSharedDataKey instead.")
-	FMassISMCSharedData* GetAndMarkDirty(const uint32 Hash);	
-	template<typename... TArgs>
-	UE_DEPRECATED(5.5, "Deprecated. Using hashes in Mass Visualization is being phased out. Use FISMCSharedDataKey instead.")
-	FMassISMCSharedData& FindOrAdd(const uint32 Hash, TArgs&&... InNewInstanceArgs);
-	UE_DEPRECATED(5.5, "Deprecated. Using hashes in Mass Visualization is being phased out. Use FISMCSharedDataKey instead.")
-	FMassISMCSharedData* Find(const uint32 Hash);
-	template<typename... TArgs>
-	UE_DEPRECATED(5.5, "Deprecated. Using hashes in Mass Visualization is being phased out. Use FISMCSharedDataKey instead.")
-	FMassISMCSharedData& Add(const uint32 Hash, TArgs&&... InNewInstanceArgs);
-	UE_DEPRECATED(5.5, "Deprecated. Using hashes in Mass Visualization is being phased out. Use FISMCSharedDataKey instead.")
-	void Remove(const uint32 Hash);
 };
 
 
@@ -803,21 +784,3 @@ typedef FMassInstancedStaticMeshInfoArrayViewAccessDetector FMassInstancedStatic
 typedef TArrayView<FMassInstancedStaticMeshInfo> FMassInstancedStaticMeshInfoArrayView;
 
 #endif // ENABLE_MT_DETECTOR
-
-//-----------------------------------------------------------------------------
-// DEPRECATED
-//-----------------------------------------------------------------------------
-
-template<typename... TArgs>
-FMassISMCSharedData& FMassISMCSharedDataMap::FindOrAdd(const uint32 Hash, TArgs&&... InNewInstanceArgs)
-{
-	static FMassISMCSharedData Dummy;
-	return Dummy;
-}
-
-template<typename... TArgs>
-FMassISMCSharedData& FMassISMCSharedDataMap::Add(const uint32 Hash, TArgs&&... InNewInstanceArgs)
-{
-	static FMassISMCSharedData Dummy;
-	return Dummy;
-}
diff --git a/Engine/Plugins/Runtime/MassGameplay/Source/MassRepresentation/Public/MassVisualizationComponent.h b/Engine/Plugins/Runtime/MassGameplay/Source/MassRepresentation/Public/MassVisualizationComponent.h
--- a/Engine/Plugins/Runtime/MassGameplay/Source/MassRepresentation/Public/MassVisualizationComponent.h
+++ b/Engine/Plugins/Runtime/MassGameplay/Source/MassRepresentation/Public/MassVisualizationComponent.h
@@ -129,7 +129,4 @@ protected:
 
 	/** Indices to InstancedStaticMeshInfos that need their SMComponent constructed */
 	TArray<FStaticMeshInstanceVisualizationDescHandle> InstancedSMComponentsRequiringConstructing;
-
-	UE_DEPRECATED(5.5, "This flavor of BuildLODSignificanceForInfo is no longer supported and is defunct.")
-	void BuildLODSignificanceForInfo(FMassInstancedStaticMeshInfo& Info) {}
 };
diff --git a/Engine/Plugins/Runtime/MassGameplay/Source/MassSpawner/Private/MassEntityTraitBase.cpp b/Engine/Plugins/Runtime/MassGameplay/Source/MassSpawner/Private/MassEntityTraitBase.cpp
--- a/Engine/Plugins/Runtime/MassGameplay/Source/MassSpawner/Private/MassEntityTraitBase.cpp
+++ b/Engine/Plugins/Runtime/MassGameplay/Source/MassSpawner/Private/MassEntityTraitBase.cpp
@@ -7,10 +7,6 @@
 //----------------------------------------------------------------------//
 
 #include UE_INLINE_GENERATED_CPP_BY_NAME(MassEntityTraitBase)
-void UMassEntityTraitBase::DestroyTemplate() const 
-{
-
-}
 
 bool UMassEntityTraitBase::ValidateTemplate(const FMassEntityTemplateBuildContext& BuildContext, const UWorld& World, FAdditionalTraitRequirements& OutTraitRequirements) const
 {
@@ -31,13 +27,3 @@ void UMassEntityTraitBase::PostInitProperties()
 }
 
 #endif // WITH_EDITOR
-
-//-----------------------------------------------------------------------------
-// DEPRECATED
-//-----------------------------------------------------------------------------
-bool UMassEntityTraitBase::ValidateTemplate(FMassEntityTemplateBuildContext& BuildContext, const UWorld& World) const
-{
-	TArray<const UStruct*> AdditionalRequirements;
-	FAdditionalTraitRequirements AdditionalTraitRequirementsWrapper(AdditionalRequirements);
-	return ValidateTemplate(BuildContext, World, AdditionalTraitRequirementsWrapper);
-}
diff --git a/Engine/Plugins/Runtime/MassGameplay/Source/MassSpawner/Public/MassEntityTraitBase.h b/Engine/Plugins/Runtime/MassGameplay/Source/MassSpawner/Public/MassEntityTraitBase.h
--- a/Engine/Plugins/Runtime/MassGameplay/Source/MassSpawner/Public/MassEntityTraitBase.h
+++ b/Engine/Plugins/Runtime/MassGameplay/Source/MassSpawner/Public/MassEntityTraitBase.h
@@ -2,9 +2,6 @@
 
 #pragma once 
 
-#if UE_ENABLE_INCLUDE_ORDER_DEPRECATED_IN_5_5
-#include "CoreMinimal.h"
-#endif // UE_ENABLE_INCLUDE_ORDER_DEPRECATED_IN_5_6
 #include "MassProcessingTypes.h"
 #include "MassEntityTraitBase.generated.h"
 
@@ -53,9 +50,6 @@ public:
 	/** Appends items into the entity template required for the trait. */
 	virtual void BuildTemplate(FMassEntityTemplateBuildContext& BuildContext, const UWorld& World) const PURE_VIRTUAL(UMassEntityTraitBase::BuildTemplate, return; );
 
-	UE_DEPRECATED(5.5, "This flavor of DestroyTemplate has been deprecated, use the version taking the World parameter")
-	MASSSPAWNER_API virtual void DestroyTemplate() const;
-
 	virtual void DestroyTemplate(const UWorld& World) const {}
 
 	/**
@@ -68,9 +62,6 @@ public:
 	 * @return whether the validation was successful
 	 */
 	MASSSPAWNER_API virtual bool ValidateTemplate(const FMassEntityTemplateBuildContext& BuildContext, const UWorld& World, FAdditionalTraitRequirements& OutTraitRequirements) const;
-	
-	UE_DEPRECATED(5.5, "This flavor of ValidateTemplate is deprecated. Use the three-parameter one instead.")
-	MASSSPAWNER_API virtual bool ValidateTemplate(FMassEntityTemplateBuildContext& BuildContext, const UWorld& World) const;
 
 #if WITH_EDITOR
 	DECLARE_MULTICAST_DELEGATE_OneParam(FOnNewTraitType, UMassEntityTraitBase&);
diff --git a/Engine/Plugins/Runtime/MassGameplay/Source/MassSpawner/Public/MassSpawnerSubsystem.h b/Engine/Plugins/Runtime/MassGameplay/Source/MassSpawner/Public/MassSpawnerSubsystem.h
--- a/Engine/Plugins/Runtime/MassGameplay/Source/MassSpawner/Public/MassSpawnerSubsystem.h
+++ b/Engine/Plugins/Runtime/MassGameplay/Source/MassSpawner/Public/MassSpawnerSubsystem.h
@@ -3,9 +3,6 @@
 #pragma once
 
 #include "MassEntityManager.h"
-#if UE_ENABLE_INCLUDE_ORDER_DEPRECATED_IN_5_5
-#include "StructUtils/InstancedStruct.h"
-#endif // UE_ENABLE_INCLUDE_ORDER_DEPRECATED_IN_5_6
 #include "MassSubsystemBase.h"
 #include "MassEntityTemplateRegistry.h"
 #include "MassSpawnerSubsystem.generated.h"
diff --git a/Engine/Source/Runtime/MassEntity/Private/MassExecutor.cpp b/Engine/Source/Runtime/MassEntity/Private/MassExecutor.cpp
--- a/Engine/Source/Runtime/MassEntity/Private/MassExecutor.cpp
+++ b/Engine/Source/Runtime/MassEntity/Private/MassExecutor.cpp
@@ -222,18 +222,6 @@ FGraphEventRef TriggerParallelTasks(UMassProcessor& Processor, FProcessingContex
 //-----------------------------------------------------------------------------
 // DEPRECATED
 //-----------------------------------------------------------------------------
-void RunProcessorsView(TArrayView<UMassProcessor* const> Processors, FProcessingContext& ProcessingContext, const FMassArchetypeEntityCollection* EntityCollection)
-{
-	if (EntityCollection)
-	{
-		RunProcessorsView(Processors, ProcessingContext, MakeArrayView(EntityCollection, 1));
-	}
-	else
-	{
-		RunProcessorsView(Processors, ProcessingContext);
-	}
-}
-
 inline FGraphEventRef TriggerParallelTasks(UMassProcessor& Processor, FProcessingContext& ProcessingContext, TFunction<void()> OnDoneNotification
 		, ENamedThreads::Type CurrentThread)
 {
diff --git a/Engine/Source/Runtime/MassEntity/Private/MassObserverManager.cpp b/Engine/Source/Runtime/MassEntity/Private/MassObserverManager.cpp
--- a/Engine/Source/Runtime/MassEntity/Private/MassObserverManager.cpp
+++ b/Engine/Source/Runtime/MassEntity/Private/MassObserverManager.cpp
@@ -731,11 +731,6 @@ void FMassObserverManager::ReleaseCreationHandle(FCreationNotificationHandle InC
 //----------------------------------------------------------------------//
 // DEPRECATED
 //----------------------------------------------------------------------//
-bool FMassObserverManager::OnPostEntitiesCreated(UE::Mass::FProcessingContext&, const FMassArchetypeEntityCollection& EntityCollection)
-{
-	return OnPostEntitiesCreated(EntityCollection);
-}
-
 bool FMassObserverManager::OnPostEntitiesCreated(UE::Mass::FProcessingContext&, TConstArrayView<FMassArchetypeEntityCollection> EntityCollections)
 {
 	TRACE_CPUPROFILER_EVENT_SCOPE_STR("OnPostEntitiesCreated")
@@ -800,16 +795,6 @@ void FMassObserverManager::FCreationContext::AppendEntities(const TConstArrayVie
 {
 }
 
-FMassObserverManager::FCreationContext::FCreationContext(const int32)
-	: FCreationContext()
-{}
-
-const FMassArchetypeEntityCollection& FMassObserverManager::FCreationContext::GetEntityCollection() const
-{
-	static FMassArchetypeEntityCollection DummyInstance;
-	return DummyInstance;
-}
-
 void FMassObserverManager::HandleElementsImpl(UE::Mass::FProcessingContext& ProcessingContext, TConstArrayView<FMassArchetypeEntityCollection> EntityCollections
 		, TArrayView<const UScriptStruct*> ObservedTypes, FMassObserversMap& HandlersContainer)
 {
diff --git a/Engine/Source/Runtime/MassEntity/Public/MassEntityManager.h b/Engine/Source/Runtime/MassEntity/Public/MassEntityManager.h
--- a/Engine/Source/Runtime/MassEntity/Public/MassEntityManager.h
+++ b/Engine/Source/Runtime/MassEntity/Public/MassEntityManager.h
@@ -9,10 +9,6 @@
 #if UE_ENABLE_INCLUDE_ORDER_DEPRECATED_IN_5_7
 #include "MassEntityQuery.h"
 #endif
-#if UE_ENABLE_INCLUDE_ORDER_DEPRECATED_IN_5_5
-#include "StructUtils/InstancedStruct.h"
-#include "StructUtils/StructUtilsTypes.h"
-#endif
 #include "MassObserverManager.h"
 #include "Containers/MpscQueue.h"
 #if UE_ENABLE_INCLUDE_ORDER_DEPRECATED_IN_5_7
@@ -722,15 +718,21 @@ public:
 	 */
 	UE_API void AppendCommands(const TSharedPtr<FMassCommandBuffer>& InOutCommandBuffer);
 
+private:
+	template<typename T, typename... TArgs>
+	const FSharedStruct& GetOrCreateSharedFragmentByHash(const uint32 Hash, TArgs&&... InArgs)
+	{
+		static_assert(UE::Mass::CSharedFragment<T>, "Given struct doesn't represent a valid shared fragment type. Make sure to inherit from FMassSharedFragment or one of its child-types.");
+		return SharedFragmentsContainer.FindOrAdd<T>(Hash, Forward<TArgs>(InArgs)...);
+	}
+
 	template<typename T>
-	UE_DEPRECATED(5.5, "This method will no longer be exposed. Use GetOrCreateConstSharedFragment instead.")
 	const FConstSharedStruct& GetOrCreateConstSharedFragmentByHash(const uint32 Hash, const T& Fragment)
 	{
 		static_assert(UE::Mass::CConstSharedFragment<T>, "Given struct doesn't represent a valid const shared fragment type. Make sure to inherit from FMassConstSharedFragment or one of its child-types.");
 		return ConstSharedFragmentsContainer.FindOrAdd(Hash, T::StaticStruct(), reinterpret_cast<const uint8*>(&Fragment));
 	}
 
-private:
 	template<typename T>
 	const FSharedStruct& GetOrCreateSharedFragmentByHash(const uint32 Hash, const T& Fragment)
 	{
@@ -776,14 +778,6 @@ public:
 	//-----------------------------------------------------------------------------
 	// Shared fragments
 	//-----------------------------------------------------------------------------
-	template<typename T, typename... TArgs>
-	UE_DEPRECATED(5.5, "This method will no longer be exposed. Use GetOrCreateSharedFragment instead.")
-	const FSharedStruct& GetOrCreateSharedFragmentByHash(const uint32 Hash, TArgs&&... InArgs)
-	{
-		static_assert(UE::Mass::CSharedFragment<T>, "Given struct doesn't represent a valid shared fragment type. Make sure to inherit from FMassSharedFragment or one of its child-types.");
-		return SharedFragmentsContainer.FindOrAdd<T>(Hash, Forward<TArgs>(InArgs)...);
-	}
-
 	/**
 	 * Returns or creates a shared struct associated to a given shared fragment set of values
 	 * identified internally by a CRC.
@@ -812,9 +806,7 @@ public:
 		static_assert(UE::Mass::CConstSharedFragment<T>,
 			"Given struct doesn't represent a valid const shared fragment type. Make sure to inherit from FMassConstSharedFragment or one of its child-types.");
 		const uint32 Hash = UE::StructUtils::GetStructInstanceCrc32(FConstStructView::Make(Fragment));
-PRAGMA_DISABLE_DEPRECATION_WARNINGS
 		return GetOrCreateConstSharedFragmentByHash(Hash, Fragment);
-PRAGMA_ENABLE_DEPRECATION_WARNINGS
 	}
 
 	/**
@@ -877,9 +869,7 @@ PRAGMA_ENABLE_DEPRECATION_WARNINGS
 			"Given struct doesn't represent a valid const shared fragment type. Make sure to inherit from FMassConstSharedFragment or one of its child-types.");
 		T Struct(Forward<TArgs>(InArgs)...);
 		const uint32 Hash = UE::StructUtils::GetStructInstanceCrc32(FConstStructView::Make(Struct));
-PRAGMA_DISABLE_DEPRECATION_WARNINGS
 		return GetOrCreateConstSharedFragmentByHash(Hash, MoveTemp(Struct));
-PRAGMA_ENABLE_DEPRECATION_WARNINGS
 	}
 
 	/**
@@ -996,9 +986,7 @@ PRAGMA_ENABLE_DEPRECATION_WARNINGS
 			"Given struct doesn't represent a valid const shared fragment type. Make sure to inherit from FMassConstSharedFragment or one of its child-types.");
 		T Fragment(Forward<TArgs>(InArgs)...);
 		const uint32 Hash = UE::StructUtils::GetStructInstanceCrc32(HashingHelperStruct);
-PRAGMA_DISABLE_DEPRECATION_WARNINGS
 		return GetOrCreateConstSharedFragmentByHash(Hash, MoveTemp(Fragment));
-PRAGMA_ENABLE_DEPRECATION_WARNINGS
 	}
 
 	/**
diff --git a/Engine/Source/Runtime/MassEntity/Public/MassEntityTypes.h b/Engine/Source/Runtime/MassEntity/Public/MassEntityTypes.h
--- a/Engine/Source/Runtime/MassEntity/Public/MassEntityTypes.h
+++ b/Engine/Source/Runtime/MassEntity/Public/MassEntityTypes.h
@@ -304,27 +304,6 @@ public:
 	{
 		return ElementsBitSet.HasAll<TElementType>();
 	}
-
-	UE_DEPRECATED(5.5, "This FMassArchetypeCompositionDescriptor constructor is deprecated. Please explicitly provide FConstSharedFragmentBitSet.")
-	FMassArchetypeCompositionDescriptor(const FMassFragmentBitSet& InFragments, const FMassTagBitSet& InTags, const FMassChunkFragmentBitSet& InChunkFragments, const FMassSharedFragmentBitSet& InSharedFragments)
-		: FMassArchetypeCompositionDescriptor(InFragments, InTags, InChunkFragments, InSharedFragments, FMassConstSharedFragmentBitSet())
-	{}
-
-	UE_DEPRECATED(5.5, "This FMassArchetypeCompositionDescriptor constructor is deprecated. Please explicitly provide FConstSharedFragmentBitSet.")
-	FMassArchetypeCompositionDescriptor(TConstArrayView<const UScriptStruct*> InFragments, const FMassTagBitSet& InTags, const FMassChunkFragmentBitSet& InChunkFragments, const FMassSharedFragmentBitSet& InSharedFragments)
-		: FMassArchetypeCompositionDescriptor(FMassFragmentBitSet(InFragments), InTags, InChunkFragments, InSharedFragments, FMassConstSharedFragmentBitSet())
-	{}
-
-	UE_DEPRECATED(5.5, "This FMassArchetypeCompositionDescriptor constructor is deprecated. Please explicitly provide FConstSharedFragmentBitSet.")
-	FMassArchetypeCompositionDescriptor(TConstArrayView<FInstancedStruct> InFragmentInstances, const FMassTagBitSet& InTags, const FMassChunkFragmentBitSet& InChunkFragments, const FMassSharedFragmentBitSet& InSharedFragments)
-		: FMassArchetypeCompositionDescriptor(FMassFragmentBitSet(InFragmentInstances), InTags, InChunkFragments, InSharedFragments, FMassConstSharedFragmentBitSet())
-	{}
-
-	UE_DEPRECATED(5.5, "This FMassArchetypeCompositionDescriptor constructor is deprecated. Please explicitly provide FConstSharedFragmentBitSet.")
-	FMassArchetypeCompositionDescriptor(FMassFragmentBitSet&& InFragments, FMassTagBitSet&& InTags, FMassChunkFragmentBitSet&& InChunkFragments, FMassSharedFragmentBitSet&& InSharedFragments)
-	{
-		ensureMsgf(false, TEXT("This constructor is defunct. Please update your implementation based on deprecation warning."));
-	}
 };
 PRAGMA_ENABLE_DEPRECATION_WARNINGS
 
@@ -650,16 +629,6 @@ protected:
 
 private:
 	mutable FTransactionallySafeRWLock HashLock;
-
-public:
-	//-----------------------------------------------------------------------------
-	// DEPRECATED
-	//-----------------------------------------------------------------------------
-	UE_DEPRECATED(5.5, "HasExactFragmentTypesMatch is deprecated. Use HasExactSharedFragmentTypesMatch or the two-parameter version of HasExactFragmentTypesMatch.")
-	inline bool HasExactFragmentTypesMatch(const FMassSharedFragmentBitSet& InSharedFragmentBitSet) const
-	{
-		return (StoredElementsBitSet & InSharedFragmentBitSet).IsEquivalent(InSharedFragmentBitSet);
-	}
 };
 
 /**
diff --git a/Engine/Source/Runtime/MassEntity/Public/MassEntityView.h b/Engine/Source/Runtime/MassEntity/Public/MassEntityView.h
--- a/Engine/Source/Runtime/MassEntity/Public/MassEntityView.h
+++ b/Engine/Source/Runtime/MassEntity/Public/MassEntityView.h
@@ -116,7 +116,7 @@ struct FMassEntityView
 		return (T*)GetSharedFragmentPtr(*T::StaticStruct());
 	}
 
-	template<UE::Mass::CConstSharedFragment T>
+	/*template<UE::Mass::CConstSharedFragment T>
 	UE_DEPRECATED(5.5, "Using GetSharedFragmentDataPtr with const shared fragments is deprecated. Use GetConstSharedFragmentDataPtr instead")
 	T* GetSharedFragmentDataPtr() const
 	{
@@ -128,7 +128,7 @@ struct FMassEntityView
 	T& GetSharedFragmentData() const
 	{
 		return const_cast<T&>(GetConstSharedFragmentData<T>());
-	}
+	}*/
 
 	FStructView GetSharedFragmentDataStruct(const UScriptStruct* FragmentType) const
 	{
diff --git a/Engine/Source/Runtime/MassEntity/Public/MassExecutor.h b/Engine/Source/Runtime/MassEntity/Public/MassExecutor.h
--- a/Engine/Source/Runtime/MassEntity/Public/MassExecutor.h
+++ b/Engine/Source/Runtime/MassEntity/Public/MassExecutor.h
@@ -49,9 +49,6 @@ namespace UE::Mass::Executor
 	MASSENTITY_API FGraphEventRef TriggerParallelTasks(UMassProcessor& Processor, FProcessingContext&& ProcessingContext, TFunction<void()> OnDoneNotification
 		, ENamedThreads::Type CurrentThread = ENamedThreads::GameThread);
 
-	UE_DEPRECATED(5.5, "This flavor of RunProcessorsView is deprecated. Use the one with TConstArrayView<FMassArchetypeEntityCollection> parameter instead.")
-	MASSENTITY_API void RunProcessorsView(TArrayView<UMassProcessor* const> Processors, FProcessingContext& ProcessingContext, const FMassArchetypeEntityCollection* EntityCollection);
-
 	UE_DEPRECATED(5.6, "lvalue flavor of TriggerParallelTasks has been deprevate. Use the rvalue version.")
 	MASSENTITY_API FGraphEventRef TriggerParallelTasks(UMassProcessor& Processor, FProcessingContext& ProcessingContext, TFunction<void()> OnDoneNotification
 		, ENamedThreads::Type CurrentThread = ENamedThreads::GameThread);
diff --git a/Engine/Source/Runtime/MassEntity/Public/MassObserverManager.h b/Engine/Source/Runtime/MassEntity/Public/MassObserverManager.h
--- a/Engine/Source/Runtime/MassEntity/Public/MassObserverManager.h
+++ b/Engine/Source/Runtime/MassEntity/Public/MassObserverManager.h
@@ -252,8 +252,6 @@ protected:
 #endif // WITH_MASSENTITY_DEBUG
 
 public:
-	UE_DEPRECATED(5.5, "This flavor of OnPostEntitiesCreated is deprecated. Please use the one taking a TConstArrayView<FMassArchetypeEntityCollection> parameter instead")
-	MASSENTITY_API bool OnPostEntitiesCreated(UE::Mass::FProcessingContext& InProcessingContext, const FMassArchetypeEntityCollection& EntityCollection);
 	UE_DEPRECATED(5.6, "*FragmentOrTag* functions have been deprecated. Use OnCompositionChanged instead.")
 	MASSENTITY_API void OnPostFragmentOrTagAdded(const UScriptStruct& FragmentOrTagType, const FMassArchetypeEntityCollection& EntityCollection);
 	UE_DEPRECATED(5.6, "*FragmentOrTag* functions have been deprecated. Use OnCompositionChanged instead.")
diff --git a/Engine/Source/Runtime/MassEntity/Public/MassObserverNotificationTypes.h b/Engine/Source/Runtime/MassEntity/Public/MassObserverNotificationTypes.h
--- a/Engine/Source/Runtime/MassEntity/Public/MassObserverNotificationTypes.h
+++ b/Engine/Source/Runtime/MassEntity/Public/MassObserverNotificationTypes.h
@@ -390,10 +390,6 @@ namespace UE::Mass
 			UE_API void AppendEntities(const TConstArrayView<FMassEntityHandle>);
 			UE_DEPRECATED(5.6, "Manually adding entities directly to the creation context is not longer supported and is not taking place automatically")
 			UE_API void AppendEntities(const TConstArrayView<FMassEntityHandle>, FMassArchetypeEntityCollection&&);
-			UE_DEPRECATED(5.5, "This constructor is now deprecated and defunct. Use one of the others instead.")
-			UE_API explicit FCreationContext(const int32);
-			UE_DEPRECATED(5.5, "This function is now deprecated since FEntityCreationContext can contain more than a single collection now. Use GetEntityCollections instead.")
-			UE_API const FMassArchetypeEntityCollection& GetEntityCollection() const;
 			
 			/**
 			 *	Called in response to composition mutating operation - these operations invalidate stored collections
