[Mass] removed some (remaining) Mass code deprecated in 5.5

#rnx
#Synced-CL 49647462

diff --git a/Engine/Plugins/Runtime/MassGameplay/Source/MassRepresentation/Public/MassRepresentationTypes.h b/Engine/Plugins/Runtime/MassGameplay/Source/MassRepresentation/Public/MassRepresentationTypes.h
--- a/Engine/Plugins/Runtime/MassGameplay/Source/MassRepresentation/Public/MassRepresentationTypes.h
+++ b/Engine/Plugins/Runtime/MassGameplay/Source/MassRepresentation/Public/MassRepresentationTypes.h
@@ -718,12 +718,6 @@ protected:
 	TArray<FMassLODSignificanceRange> LODSignificanceRanges;
 
 	friend class UMassVisualizationComponent;
-
-	//-----------------------------------------------------------------------------
-	// DEPRECATED
-	//-----------------------------------------------------------------------------
-	UE_DEPRECATED(5.5, "Deprecated in flavor of the function taking the ISMComponent parameter. This version is not defunct")
-	void ClearVisualInstance(FMassISMCSharedDataMap& ISMCSharedData) {}
 };
 
 #if ENABLE_MT_DETECTOR
diff --git a/Engine/Source/Runtime/MassEntity/Public/MassEntityView.h b/Engine/Source/Runtime/MassEntity/Public/MassEntityView.h
--- a/Engine/Source/Runtime/MassEntity/Public/MassEntityView.h
+++ b/Engine/Source/Runtime/MassEntity/Public/MassEntityView.h
@@ -116,20 +116,6 @@ struct FMassEntityView
 		return (T*)GetSharedFragmentPtr(*T::StaticStruct());
 	}
 
-	/*template<UE::Mass::CConstSharedFragment T>
-	UE_DEPRECATED(5.5, "Using GetSharedFragmentDataPtr with const shared fragments is deprecated. Use GetConstSharedFragmentDataPtr instead")
-	T* GetSharedFragmentDataPtr() const
-	{
-		return const_cast<T*>(GetConstSharedFragmentDataPtr<T>());
-	}
-
-	template<UE::Mass::CConstSharedFragment T>
-	UE_DEPRECATED(5.5, "Using GetSharedFragmentDataPtr with const shared fragments is deprecated. Use GetConstSharedFragmentData instead")
-	T& GetSharedFragmentData() const
-	{
-		return const_cast<T&>(GetConstSharedFragmentData<T>());
-	}*/
-
 	FStructView GetSharedFragmentDataStruct(const UScriptStruct* FragmentType) const
 	{
 		check(UE::Mass::IsA<FMassSharedFragment>(FragmentType));
diff --git a/Engine/Source/Runtime/MassEntity/Public/MassProcessor.h b/Engine/Source/Runtime/MassEntity/Public/MassProcessor.h
--- a/Engine/Source/Runtime/MassEntity/Public/MassProcessor.h
+++ b/Engine/Source/Runtime/MassEntity/Public/MassProcessor.h
@@ -425,10 +425,6 @@ protected:
 	// DEPRECATED
 	//-----------------------------------------------------------------------------
 public:
-
-	/*UE_DEPRECATED(5.6, "This flavor of Initialize is deprecated. Please use the one requiring a FMassEntityManager parameter")
-	virtual void Initialize(UObject& Owner) final;*/
-
 	UE_DEPRECATED(5.6, "This flavor of SetChildProcessors is deprecated. Please use one of the others.")
 	UE_API void SetChildProcessors(TArray<UMassProcessor*>&& InProcessors);
 };
