[Mass] Fixed checkSlow in FMassExecutionContext.GetIndirectConstFragmentPtr to use FMassEntityQuery::HasConstIndirectEntityRequirement rather than FMassEntityQuery::HasMutableIndirectEntityRequirement

#rnx

#Synced-CL 49647462

diff --git a/Engine/Source/Runtime/MassEntity/Public/MassExecutionContext.h b/Engine/Source/Runtime/MassEntity/Public/MassExecutionContext.h
--- a/Engine/Source/Runtime/MassEntity/Public/MassExecutionContext.h
+++ b/Engine/Source/Runtime/MassEntity/Public/MassExecutionContext.h
@@ -532,7 +532,7 @@ public:
 	const T* GetIndirectConstFragmentPtr(const FMassEntityHandle& FragmentOwner)
 	{
 		checkf(QueriesStack.Num(), TEXT("There must be an active query to access indirect fragments, otherwise just use direct access with a query"));
-		checkfSlow(QueriesStack.Last().Query->HasMutableIndirectEntityRequirement<T>()
+		checkfSlow(QueriesStack.Last().Query->HasConstIndirectEntityRequirement<T>()
 			, TEXT("Attempting indirect fragment access without first adding and IndirectFragment requirement to the query"));
 
 		if (CachedIndirectEntityView.GetEntity() != FragmentOwner)
