Added ability to check fragments or tags in one call in a Mass context.

#Synced-CL 49834081

diff --git a/Engine/Source/Runtime/MassEntity/Public/MassExecutionContext.h b/Engine/Source/Runtime/MassEntity/Public/MassExecutionContext.h
--- a/Engine/Source/Runtime/MassEntity/Public/MassExecutionContext.h
+++ b/Engine/Source/Runtime/MassEntity/Public/MassExecutionContext.h
@@ -309,15 +309,25 @@ public:
 		}
 	}
 
+	bool DoesArchetypeHaveElement(const UScriptStruct& ElementType) const
+	{
+		return CurrentArchetypeCompositionBitSet.Contains(ElementType);
+	}
+
+	template<UE::Mass::CElement T>
+	bool DoesArchetypeHaveElement() const
+	{
+		return CurrentArchetypeCompositionBitSet.Contains<T>();
+	}
+
 	bool DoesArchetypeHaveFragment(const UScriptStruct& FragmentType) const
 	{
 		return CurrentArchetypeCompositionBitSet.Contains(FragmentType);
 	}
 
-	template<typename T>
+	template<UE::Mass::CFragment T>
 	bool DoesArchetypeHaveFragment() const
 	{
-		static_assert(UE::Mass::CFragment<T>, MASS_INVALID_FRAGMENT_MSG);
 		return CurrentArchetypeCompositionBitSet.Contains<T>();
 	}
 
@@ -326,10 +336,9 @@ public:
 		return CurrentArchetypeCompositionBitSet.Contains(TagType);
 	}
 
-	template<typename T>
+	template<UE::Mass::CTag T>
 	bool DoesArchetypeHaveTag() const
 	{
-		static_assert(UE::Mass::CTag<T>, "Given struct is not of a valid tag type.");
 		return CurrentArchetypeCompositionBitSet.Contains<T>();
 	}
 
