[Mass]
Thread-safe observer lock notifications.
Lock observers during processing phases

#Synced-CL 49747586

diff --git a/Engine/Source/Runtime/MassEntity/Private/MassExecutor.cpp b/Engine/Source/Runtime/MassEntity/Private/MassExecutor.cpp
--- a/Engine/Source/Runtime/MassEntity/Private/MassExecutor.cpp
+++ b/Engine/Source/Runtime/MassEntity/Private/MassExecutor.cpp
@@ -109,9 +109,10 @@ void RunProcessorsView(TArrayView<UMassProcessor* const> Processors, FProcessing
 
 struct FMassExecutorDoneTask
 {
-	FMassExecutorDoneTask(FMassExecutionContext&& InExecutionContext, TFunction<void()> InOnDoneNotification, const FString& InDebugName, ENamedThreads::Type InDesiredThread)
+	FMassExecutorDoneTask(FMassExecutionContext&& InExecutionContext, TFunction<void()> InOnDoneNotification, TFunction<void()> InOnPreCommandFlushNotification, const FString& InDebugName, ENamedThreads::Type InDesiredThread)
 		: ExecutionContext(InExecutionContext)
 		, OnDoneNotification(InOnDoneNotification)
+		, OnPreCommandFlushNotification(InOnPreCommandFlushNotification)
 		, DebugName(InDebugName)
 		, DesiredThread(InDesiredThread)
 	{
@@ -140,14 +141,20 @@ struct FMassExecutorDoneTask
 		}
 
 		UE_LOG(LogMass, Verbose, TEXT("MassExecutor %s tasks DONE"), *DebugName);
+
+		if (OnPreCommandFlushNotification)
+		{
+			OnPreCommandFlushNotification();
+		}
+		
 		ExecutionContext.SetFlushDeferredCommands(true);
 		ExecutionContext.FlushDeferred();
-
 		OnDoneNotification();
 	}
 private:
 	FMassExecutionContext ExecutionContext;
 	TFunction<void()> OnDoneNotification;
+	TFunction<void()> OnPreCommandFlushNotification;
 	FString DebugName;
 	ENamedThreads::Type DesiredThread;
 };
@@ -155,9 +162,10 @@ private:
 struct FMassExecutorDispatchTask
 {
 	FMassExecutorDispatchTask(const FMassExecutionContext& InExecutionContext, const TSharedPtr<FMassEntityManager>& InEntityManager, UMassProcessor& InProcessor
-		, TFunction<void()> InOnDoneNotification, const FString& InDebugName, ENamedThreads::Type InDesiredEndTaskThread)
+		, TFunction<void()> InOnDoneNotification, TFunction<void()> InOnPreCommandFlushNotification, const FString& InDebugName, ENamedThreads::Type InDesiredEndTaskThread)
 		: ExecutionContext(InExecutionContext)
 		, OnDoneNotification(InOnDoneNotification)
+		, OnPreCommandFlushNotification(InOnPreCommandFlushNotification)
 		, DebugName(InDebugName)
 		, EntityManager(InEntityManager)
 		, Processor(InProcessor)
@@ -182,7 +190,7 @@ struct FMassExecutorDispatchTask
 			const FGraphEventArray Prerequisites = { CompletionEvent };
 
 			CompletionEvent = TGraphTask<FMassExecutorDoneTask>::CreateTask(&Prerequisites)
-				.ConstructAndDispatchWhenReady(MoveTemp(ExecutionContextThread), OnDoneNotification, Processor.GetName(), DesiredEndTaskThread);
+				.ConstructAndDispatchWhenReady(MoveTemp(ExecutionContextThread), OnDoneNotification, OnPreCommandFlushNotification, Processor.GetName(), DesiredEndTaskThread);
 		}
 
 		if (CompletionEvent.IsValid())
@@ -193,16 +201,19 @@ struct FMassExecutorDispatchTask
 private:
 	FMassExecutionContext ExecutionContext;
 	TFunction<void()> OnDoneNotification;
+	TFunction<void()> OnPreCommandFlushNotification;
 	FString DebugName;
 	TSharedPtr<FMassEntityManager> EntityManager;
 	UMassProcessor& Processor;
 	ENamedThreads::Type DesiredEndTaskThread;
 };
 
-FGraphEventRef TriggerParallelTasks(UMassProcessor& Processor, FProcessingContext&& ProcessingContext, TFunction<void()> OnDoneNotification, ENamedThreads::Type CurrentThread)
+FGraphEventRef TriggerParallelTasks(UMassProcessor& Processor, FProcessingContext&& ProcessingContext, TFunction<void()> OnDoneNotification, ENamedThreads::Type CurrentThread, TFunction<void()> OnPreCommandFlushNotification)
 {
 	TRACE_CPUPROFILER_EVENT_SCOPE(MassExecutor_RunParallel);
 
+	TSharedPtr<FMassEntityManager> EntityManager = ProcessingContext.EntityManager;
+
 	// We need to transfer ProcessingContext's ExecutionContext - otherwise ProcessingContext's destructor will attempt
 	// flushing stored commands. 
 	FMassExecutionContext ExecutionContext = MoveTemp(ProcessingContext).GetExecutionContext();
@@ -211,9 +222,8 @@ FGraphEventRef TriggerParallelTasks(UMassProcessor& Processor, FProcessingContex
 	{
 		TRACE_CPUPROFILER_EVENT_SCOPE_STR("Dispatch Processors")
 		FGraphEventArray Prerequisites;
-		TSharedPtr<FMassEntityManager> EntityManager = ProcessingContext.EntityManager;
 		DispatchEvent  = TGraphTask<FMassExecutorDispatchTask>::CreateTask(&Prerequisites)
-			.ConstructAndDispatchWhenReady(ExecutionContext, EntityManager, Processor, OnDoneNotification, Processor.GetName(), CurrentThread);
+			.ConstructAndDispatchWhenReady(ExecutionContext, EntityManager, Processor, OnDoneNotification, OnPreCommandFlushNotification, Processor.GetName(), CurrentThread);
 	}
 
 	return DispatchEvent;
@@ -226,7 +236,7 @@ inline FGraphEventRef TriggerParallelTasks(UMassProcessor& Processor, FProcessin
 		, ENamedThreads::Type CurrentThread)
 {
 	FProcessingContext LocalContext = ProcessingContext;
-	return TriggerParallelTasks(Processor, MoveTemp(ProcessingContext), OnDoneNotification, CurrentThread);
+	return TriggerParallelTasks(Processor, MoveTemp(LocalContext), OnDoneNotification, CurrentThread);
 }
 
 } // namespace UE::Mass::Executor
diff --git a/Engine/Source/Runtime/MassEntity/Private/MassObserverManager.cpp b/Engine/Source/Runtime/MassEntity/Private/MassObserverManager.cpp
--- a/Engine/Source/Runtime/MassEntity/Private/MassObserverManager.cpp
+++ b/Engine/Source/Runtime/MassEntity/Private/MassObserverManager.cpp
@@ -565,9 +565,9 @@ void FMassObserverManager::RemoveObserverInstance(const UScriptStruct& ElementTy
 
 TSharedRef<FMassObserverManager::FObserverLock> FMassObserverManager::GetOrMakeObserverLock()
 {
-	if (ActiveObserverLock.IsValid())
+	if (TSharedPtr<FMassObserverManager::FObserverLock> SharedActiveLock = ActiveObserverLock.Pin())
 	{
-		return ActiveObserverLock.Pin().ToSharedRef();
+		return SharedActiveLock.ToSharedRef();
 	}
 	else
 	{
diff --git a/Engine/Source/Runtime/MassEntity/Private/MassObserverNotificationTypes.cpp b/Engine/Source/Runtime/MassEntity/Private/MassObserverNotificationTypes.cpp
--- a/Engine/Source/Runtime/MassEntity/Private/MassObserverNotificationTypes.cpp
+++ b/Engine/Source/Runtime/MassEntity/Private/MassObserverNotificationTypes.cpp
@@ -41,6 +41,7 @@ namespace UE::Mass::ObserverManager
 		checkf(SharedEntityManager->GetObserverManager().LocksCount >= 0
 			, TEXT("%hs: the lock count has become unbalanced."), __FUNCTION__);
 		SharedEntityManager->GetObserverManager().ResumeExecution(*this);
+
 	}
 
 	void FObserverLock::ForceUpdateCurrentThreadID()
@@ -79,6 +80,8 @@ namespace UE::Mass::ObserverManager
 		// if the creation handle isn't set there are no creation ops we know about
 		if (CreationHandle.IsSet())
 		{
+			UE::TScopeLock ScopeLock(Lock->GetNotificationLock());
+
 			const FBufferedNotification& Notification = Lock->GetCreationNotification(CreationHandle);
 
 			if (const UE::Mass::FEntityCollection* CreatedEntities = Notification.AffectedEntities.TryGet<FEntityCollection>())
diff --git a/Engine/Source/Runtime/MassEntity/Private/MassProcessingPhaseManager.cpp b/Engine/Source/Runtime/MassEntity/Private/MassProcessingPhaseManager.cpp
--- a/Engine/Source/Runtime/MassEntity/Private/MassProcessingPhaseManager.cpp
+++ b/Engine/Source/Runtime/MassEntity/Private/MassProcessingPhaseManager.cpp
@@ -79,19 +79,28 @@ void FMassProcessingPhase::ExecuteTick(float DeltaTime, ELevelTick TickType, ENa
 	
 	FMassEntityManager& EntityManager = PhaseManager->GetEntityManagerRef();
 
-	bIsDuringMassProcessing = true;
+	bIsDuringMassProcessing = true;	
 
 	if (bRunInParallelMode && PhaseManager->IsPaused() == false)
 	{
 		bool bWorkRequested = false;
 		if (PhaseProcessor->IsEmpty() == false)
 		{
+			ObserversLockPtr = EntityManager.GetOrMakeObserversLock().ToSharedPtr();
 			FMassProcessingContext Context(EntityManager, DeltaTime);
-			const FGraphEventRef PipelineCompletionEvent = UE::Mass::Executor::TriggerParallelTasks(*PhaseProcessor, MoveTemp(Context), [this, DeltaTime]()
+
+			auto OnProcessingDone = [this, DeltaTime]()
 				{
 					OnParallelExecutionDone(DeltaTime);
-				}
-				, CurrentThread);
+				};
+
+			auto OnPreCommandFlush = [this]()
+				{
+					ObserversLockPtr = nullptr;
+				};
+
+			const FGraphEventRef PipelineCompletionEvent = UE::Mass::Executor::TriggerParallelTasks(*PhaseProcessor, MoveTemp(Context), 
+				OnProcessingDone, CurrentThread, OnPreCommandFlush);
 
 			if (PipelineCompletionEvent.IsValid())
 			{
@@ -101,6 +110,7 @@ void FMassProcessingPhase::ExecuteTick(float DeltaTime, ELevelTick TickType, ENa
 		}
 		if (bWorkRequested == false)
 		{
+			ObserversLockPtr = nullptr;
 			OnParallelExecutionDone(DeltaTime);
 		}
 	}
@@ -108,25 +118,29 @@ void FMassProcessingPhase::ExecuteTick(float DeltaTime, ELevelTick TickType, ENa
 	{
 		if (PhaseManager->IsPaused() == false)
 		{
+			ObserversLockPtr = EntityManager.GetOrMakeObserversLock().ToSharedPtr();
 			// note that it's important to create the processing context in this scope
 			// so that it wraps up its destruction before we call OnPhaseEnd, which in turn will cause
 			// the main EntityManager's command buffer to flush
 			FMassProcessingContext Context(EntityManager, DeltaTime);
 			UE::Mass::Executor::Run(*PhaseProcessor, Context);
+
 		}
 
+		bIsDuringMassProcessing = false;
+		ObserversLockPtr = nullptr;
 		{
 			LLM_SCOPE_BYNAME(TEXT("Mass/PhaseEndDelegate"));
 			OnPhaseEnd.Broadcast(DeltaTime);
 		}
 		PhaseManager->OnPhaseEnd(*this);
-		bIsDuringMassProcessing = false;
 	}
 }
 
 void FMassProcessingPhase::OnParallelExecutionDone(const float DeltaTime)
 {
 	bIsDuringMassProcessing = false;
+	ObserversLockPtr = nullptr;
 	{
 		LLM_SCOPE_BYNAME(TEXT("Mass/PhaseEndDelegate"));
 		OnPhaseEnd.Broadcast(DeltaTime);
diff --git a/Engine/Source/Runtime/MassEntity/Public/MassExecutor.h b/Engine/Source/Runtime/MassEntity/Public/MassExecutor.h
--- a/Engine/Source/Runtime/MassEntity/Public/MassExecutor.h
+++ b/Engine/Source/Runtime/MassEntity/Public/MassExecutor.h
@@ -44,10 +44,11 @@ namespace UE::Mass::Executor
 	 *  Triggers tasks executing Processor (and potentially it's children) and returns the task graph event representing 
 	 *  the task (the event will be "completed" once all the processors finish running). 
 	 *  @param OnDoneNotification will be called after all the processors are done, just after flushing the command buffer.
-	 *    Note that OnDoneNotification will be executed on GameThread.
+	 *  @param OnPreCommandFlushNotification will be called after all the processors are done, just before flushing the command buffer.
 	 */
 	MASSENTITY_API FGraphEventRef TriggerParallelTasks(UMassProcessor& Processor, FProcessingContext&& ProcessingContext, TFunction<void()> OnDoneNotification
-		, ENamedThreads::Type CurrentThread = ENamedThreads::GameThread);
+		, ENamedThreads::Type CurrentThread = ENamedThreads::GameThread
+		, TFunction<void()> OnPreCommandFlushNotification = TFunction<void()>());
 
 	UE_DEPRECATED(5.6, "lvalue flavor of TriggerParallelTasks has been deprevate. Use the rvalue version.")
 	MASSENTITY_API FGraphEventRef TriggerParallelTasks(UMassProcessor& Processor, FProcessingContext& ProcessingContext, TFunction<void()> OnDoneNotification
diff --git a/Engine/Source/Runtime/MassEntity/Public/MassObserverManager.h b/Engine/Source/Runtime/MassEntity/Public/MassObserverManager.h
--- a/Engine/Source/Runtime/MassEntity/Public/MassObserverManager.h
+++ b/Engine/Source/Runtime/MassEntity/Public/MassObserverManager.h
@@ -226,13 +226,13 @@ protected:
 
 	/**
 	 * Never access directly, use GetOrMakeObserverLock or GetOrMakeCreationContext instead.
-	 * Note: current lock is single-threaded. There's a path towards making it multithreaded, we'll work on it once we have a use-case
+	 * The Mass Processing Phase synchronously creates a lock before dispatching parallel mass tasks
+	 * Lock should be thread-safe for creation, fetching, and adding items.
 	 */
 	TWeakPtr<FObserverLock> ActiveObserverLock;
 	int32 LocksCount = 0;
 	TWeakPtr<FCreationContext> ActiveCreationContext;
 
-
 	FMassElementBitSet ObservedElements[static_cast<uint8>(EMassObservedOperation::MAX)];
 
 	UPROPERTY()
diff --git a/Engine/Source/Runtime/MassEntity/Public/MassObserverNotificationTypes.h b/Engine/Source/Runtime/MassEntity/Public/MassObserverNotificationTypes.h
--- a/Engine/Source/Runtime/MassEntity/Public/MassObserverNotificationTypes.h
+++ b/Engine/Source/Runtime/MassEntity/Public/MassObserverNotificationTypes.h
@@ -229,13 +229,18 @@ namespace UE::Mass
 				return BufferedNotifications[CreationHandle];
 			}
 
+			FTransactionallySafeCriticalSection& GetNotificationLock()
+			{
+				return NotificationLock;
+			}
+			
 		private:
 			UE_API explicit FObserverLock(FMassObserverManager& ObserverManager);
 
 			int32 GetOrCreateCreationNotification()
 			{
-				UE_CHECK_OWNER_THREADID();
-				//if (CreationNotificationIndex.compare_exchange_weak(LocalCreationNotificationIndex, BufferedNotifications.Num()))
+				UE::TScopeLock ScopeLock(NotificationLock);
+
 				if (CreationNotificationIndex == INDEX_NONE)
 				{
 					CreationNotificationIndex = BufferedNotifications.Num();
@@ -246,7 +251,8 @@ namespace UE::Mass
 
 			bool ReleaseCreationNotification(FCreationNotificationHandle CreationHandle)
 			{
-				UE_CHECK_OWNER_THREADID();
+				UE::TScopeLock ScopeLock(NotificationLock);
+
 				checkf(BufferedNotifications.IsValidIndex(CreationHandle), TEXT("Given CreationHandle doesn't match stored notifications"));
 
 				int32 CreationHandleOpIndex = CreationHandle;
@@ -262,7 +268,8 @@ namespace UE::Mass
 
 			int32 AddCreatedEntity(FMassEntityHandle CreatedEntity)
 			{
-				UE_CHECK_OWNER_THREADID();
+				UE::TScopeLock ScopeLock(NotificationLock);
+
 				if (CreationNotificationIndex == INDEX_NONE)
 				{
 					CreationNotificationIndex = BufferedNotifications.Num();
@@ -278,7 +285,8 @@ namespace UE::Mass
 
 			int32 AddCreatedEntities(const TConstArrayView<FMassEntityHandle> InCreatedEntities)
 			{
-				UE_CHECK_OWNER_THREADID();
+				UE::TScopeLock ScopeLock(NotificationLock);
+
 				if (CreationNotificationIndex == INDEX_NONE)
 				{
 					CreationNotificationIndex = BufferedNotifications.Num();
@@ -294,7 +302,8 @@ namespace UE::Mass
 
 			int32 AddCreatedEntities(const TConstArrayView<FMassEntityHandle> InCreatedEntities, FMassArchetypeEntityCollection&& InEntityCollection)
 			{
-				UE_CHECK_OWNER_THREADID();
+				UE::TScopeLock ScopeLock(NotificationLock);
+
 				if (CreationNotificationIndex == INDEX_NONE)
 				{
 					CreationNotificationIndex = BufferedNotifications.Num();
@@ -312,7 +321,8 @@ namespace UE::Mass
 			template<typename T> requires std::is_same_v<typename TDecay<T>::Type, FMassArchetypeEntityCollection>
 			int32 AddCreatedEntitiesCollection(T&& InEntityCollection)
 			{
-				UE_CHECK_OWNER_THREADID();
+				UE::TScopeLock ScopeLock(NotificationLock);
+
 				if (CreationNotificationIndex == INDEX_NONE)
 				{
 					CreationNotificationIndex = BufferedNotifications.Num();
@@ -329,6 +339,8 @@ namespace UE::Mass
 			template<typename TEntities>
 			void AddNotification(const EMassObservedOperation OperationType, TEntities&& Entities, FMassElementBitSet&& OverlapBitSet)
 			{
+				UE::TScopeLock ScopeLock(NotificationLock);
+
 				if (BufferedNotifications.Num()
 					&& BufferedNotifications.Last().Type == OperationType
 					&& FBufferedNotification::AreCompositionsEqual(BufferedNotifications.Last().CompositionChange, OverlapBitSet))
@@ -363,6 +375,7 @@ namespace UE::Mass
 	#if WITH_MASSENTITY_DEBUG
 			uint32 LockSerialNumber = 0;
 	#endif // WITH_MASSENTITY_DEBUG
+			FTransactionallySafeCriticalSection NotificationLock;
 		};
 
 		/**
diff --git a/Engine/Source/Runtime/MassEntity/Public/MassProcessingPhaseManager.h b/Engine/Source/Runtime/MassEntity/Public/MassProcessingPhaseManager.h
--- a/Engine/Source/Runtime/MassEntity/Public/MassProcessingPhaseManager.h
+++ b/Engine/Source/Runtime/MassEntity/Public/MassProcessingPhaseManager.h
@@ -95,6 +95,8 @@ protected:
 private:
 	FMassProcessingPhaseManager* PhaseManager = nullptr;
 	std::atomic<bool> bIsDuringMassProcessing = false;
+	TSharedPtr<FMassObserverManager::FObserverLock> ObserversLockPtr;
+
 	bool bRunInParallelMode = true;
 	uint8 SupportedTickTypes = 0;
 };
