[Mass] changes to Mass EntityManager and Archetype API to support use of struct views

#Synced-CL 49801936

diff --git a/Engine/Source/Runtime/MassEntity/Internal/MassArchetypeData.h b/Engine/Source/Runtime/MassEntity/Internal/MassArchetypeData.h
--- a/Engine/Source/Runtime/MassEntity/Internal/MassArchetypeData.h
+++ b/Engine/Source/Runtime/MassEntity/Internal/MassArchetypeData.h
@@ -327,11 +327,15 @@ public:
 	 * @param Entity is the entity to set the data of all fragments
 	 * @param FragmentSources are the fragments to copy the data from
 	 */
-	void SetFragmentsData(const FMassEntityHandle Entity, TArrayView<const FInstancedStruct> FragmentSources);
+	template<typename TStruct>
+	void SetFragmentsData(const FMassEntityHandle Entity, TConstArrayView<TStruct> FragmentSources);
 
-	/** For all entities indicated by EntityCollection the function sets the value of fragment of type
-	 *  FragmentSource.GetScriptStruct to the value represented by FragmentSource.GetMemory */
-	void SetFragmentData(FMassArchetypeEntityCollection::FConstEntityRangeArrayView EntityRangeContainer, const FInstancedStruct& FragmentSource);
+	/**
+	 * For all entities indicated by EntityCollection the function sets the value of fragment of type
+	 *  FragmentSource.GetScriptStruct to the value represented by FragmentSource.GetMemory
+	 */
+	template<typename TStruct>
+	void SetFragmentData(FMassArchetypeEntityCollection::FConstEntityRangeArrayView EntityRangeContainer, const TStruct& FragmentSource);
 
 	/** Returns conversion from given Requirements to archetype's fragment indices */
 	void GetRequirementsFragmentMapping(TConstArrayView<FMassFragmentRequirementDescription> Requirements, FMassFragmentIndicesMapping& OutFragmentIndices) const;
@@ -570,3 +574,48 @@ inline FMassArchetypeCompositionDescriptor FMassArchetypeData::GetCompositionDes
 {
 	return FMassArchetypeCompositionDescriptor(CompositionBitSet);
 }
+
+template<typename TStruct>
+void FMassArchetypeData::SetFragmentsData(const FMassEntityHandle Entity, TConstArrayView<TStruct> FragmentSources)
+{
+	FMassRawEntityInChunkData InternalIndex = MakeEntityHandle(Entity);
+
+	for (const TStruct& Instance : FragmentSources)
+	{
+		const UScriptStruct* FragmentType = Instance.GetScriptStruct();
+		check(FragmentType);
+		const int32 FragmentIndex = FragmentIndexMap.FindChecked(FragmentType);
+		void* FragmentMemory = GetFragmentData(FragmentIndex, InternalIndex);
+		FragmentType->CopyScriptStruct(FragmentMemory, Instance.GetMemory());
+	}
+}
+
+template<typename TStruct>
+void FMassArchetypeData::SetFragmentData(FMassArchetypeEntityCollection::FConstEntityRangeArrayView EntityRangeContainer, const TStruct& FragmentSource)
+{
+	check(FragmentSource.IsValid());
+	const UScriptStruct* FragmentType = FragmentSource.GetScriptStruct();
+	check(FragmentType);
+	const int32* FragmentIndex = FragmentIndexMap.Find(FragmentType);
+	if (LIKELY(FragmentIndex))
+	{
+		const int32 FragmentTypeSize = FragmentType->GetStructureSize();
+		const uint8* FragmentSourceMemory = FragmentSource.GetMemory();
+		check(FragmentSourceMemory);
+	
+		for (FMassArchetypeChunkIterator ChunkIterator(EntityRangeContainer); ChunkIterator; ++ChunkIterator)
+		{
+			uint8* FragmentMemory = (uint8*)FragmentConfigs[*FragmentIndex].GetFragmentData(Chunks[ChunkIterator->ChunkIndex].GetRawMemory(), ChunkIterator->SubchunkStart);
+			for (int i = ChunkIterator->Length; i; --i, FragmentMemory += FragmentTypeSize)
+			{
+				FragmentType->CopyScriptStruct(FragmentMemory, FragmentSourceMemory);
+			}
+		}
+	}
+	else
+	{
+		UE_LOG(LogMass, Warning
+			, TEXT("Attempting to set value of fragment of type %s, while it's not part of the archetype's composition")
+			, *FragmentType->GetName());
+	}
+}
\ No newline at end of file
diff --git a/Engine/Source/Runtime/MassEntity/Private/MassArchetypeData.cpp b/Engine/Source/Runtime/MassEntity/Private/MassArchetypeData.cpp
--- a/Engine/Source/Runtime/MassEntity/Private/MassArchetypeData.cpp
+++ b/Engine/Source/Runtime/MassEntity/Private/MassArchetypeData.cpp
@@ -443,49 +443,6 @@ void FMassArchetypeData::SetSharedFragmentsData(const FMassEntityHandle Entity,
 	}
 }
 
-void FMassArchetypeData::SetFragmentsData(const FMassEntityHandle Entity, TArrayView<const FInstancedStruct> FragmentInstances)
-{
-	FMassRawEntityInChunkData InternalIndex = MakeEntityHandle(Entity);
-
-	for (const FInstancedStruct& Instance : FragmentInstances)
-	{
-		const UScriptStruct* FragmentType = Instance.GetScriptStruct();
-		check(FragmentType);
-		const int32 FragmentIndex = FragmentIndexMap.FindChecked(FragmentType);
-		void* FragmentMemory = GetFragmentData(FragmentIndex, InternalIndex);
-		FragmentType->CopyScriptStruct(FragmentMemory, Instance.GetMemory());
-	}
-}
-
-void FMassArchetypeData::SetFragmentData(FMassArchetypeEntityCollection::FConstEntityRangeArrayView EntityRangeContainer, const FInstancedStruct& FragmentSource)
-{
-	check(FragmentSource.IsValid());
-	const UScriptStruct* FragmentType = FragmentSource.GetScriptStruct();
-	check(FragmentType);
-	const int32* FragmentIndex = FragmentIndexMap.Find(FragmentType);
-	if (LIKELY(FragmentIndex))
-	{
-		const int32 FragmentTypeSize = FragmentType->GetStructureSize();
-		const uint8* FragmentSourceMemory = FragmentSource.GetMemory();
-		check(FragmentSourceMemory);
-	
-		for (FMassArchetypeChunkIterator ChunkIterator(EntityRangeContainer); ChunkIterator; ++ChunkIterator)
-		{
-			uint8* FragmentMemory = (uint8*)FragmentConfigs[*FragmentIndex].GetFragmentData(Chunks[ChunkIterator->ChunkIndex].GetRawMemory(), ChunkIterator->SubchunkStart);
-			for (int i = ChunkIterator->Length; i; --i, FragmentMemory += FragmentTypeSize)
-			{
-				FragmentType->CopyScriptStruct(FragmentMemory, FragmentSourceMemory);
-			}
-		}
-	}
-	else
-	{
-		UE_LOG(LogMass, Warning
-			, TEXT("Attempting to set value of fragment of type %s, while it's not part of the archetype's composition")
-			, *FragmentType->GetName());
-	}
-}
-
 void FMassArchetypeData::MoveEntityToAnotherArchetype(const FMassEntityHandle Entity, FMassArchetypeData& NewArchetype, const FMassArchetypeSharedFragmentValues* SharedFragmentValuesOverride)
 {
 	check(&NewArchetype != this);
@@ -1527,4 +1484,4 @@ bool FMassArchetypeHelper::DoesArchetypeMatchRequirements(const FMassArchetypeDa
 
 	return false;
 }
-#endif // WITH_MASSENTITY_DEBUG
+#endif // WITH_MASSENTITY_DEBUG
\ No newline at end of file
diff --git a/Engine/Source/Runtime/MassEntity/Private/MassEntityManager.cpp b/Engine/Source/Runtime/MassEntity/Private/MassEntityManager.cpp
--- a/Engine/Source/Runtime/MassEntity/Private/MassEntityManager.cpp
+++ b/Engine/Source/Runtime/MassEntity/Private/MassEntityManager.cpp
@@ -59,6 +59,16 @@ namespace UE::Mass::Private
 //-----------------------------------------------------------------------------
 // FMassEntityManager
 //-----------------------------------------------------------------------------
+template void FMassEntityManager::SetEntityFragmentValues<FInstancedStruct>(FMassEntityHandle EntityHandle, TConstArrayView<FInstancedStruct> FragmentInstanceList);
+template void FMassEntityManager::SetEntityFragmentValues<FStructView>(FMassEntityHandle EntityHandle, TConstArrayView<FStructView> FragmentInstanceList);
+template void FMassEntityManager::SetEntityFragmentValues<FConstStructView>(FMassEntityHandle EntityHandle, TConstArrayView<FConstStructView> FragmentInstanceList);
+template void FMassEntityManager::BatchSetEntityFragmentValues<FInstancedStruct>(const FMassArchetypeEntityCollection& SparseEntities, TConstArrayView<FInstancedStruct> FragmentInstanceList);
+template void FMassEntityManager::BatchSetEntityFragmentValues<FStructView>(const FMassArchetypeEntityCollection& SparseEntities, TConstArrayView<FStructView> FragmentInstanceList);
+template void FMassEntityManager::BatchSetEntityFragmentValues<FConstStructView>(const FMassArchetypeEntityCollection& SparseEntities, TConstArrayView<FConstStructView> FragmentInstanceList);
+template void FMassEntityManager::BatchSetEntityFragmentValues<FInstancedStruct>(TConstArrayView<FMassArchetypeEntityCollection> EntityCollections, TConstArrayView<FInstancedStruct> FragmentInstanceList);
+template void FMassEntityManager::BatchSetEntityFragmentValues<FStructView>(TConstArrayView<FMassArchetypeEntityCollection> EntityCollections, TConstArrayView<FStructView> FragmentInstanceList);
+template void FMassEntityManager::BatchSetEntityFragmentValues<FConstStructView>(TConstArrayView<FMassArchetypeEntityCollection> EntityCollections, TConstArrayView<FConstStructView> FragmentInstanceList);
+
 FMassEntityManager::FMassEntityManager(UObject* InOwner)
 	: ObserverManager(*this)
 	, TypeManager(new UE::Mass::FTypeManager(*this))
@@ -1793,7 +1803,8 @@ void FMassEntityManager::MoveEntityToAnotherArchetype(FMassEntityHandle Entity,
 	ObserverManager.OnCompositionChanged(Entity, MoveTemp(CompositionAdded), EMassObservedOperation::AddElement);
 }
 
-void FMassEntityManager::SetEntityFragmentValues(FMassEntityHandle Entity, TArrayView<const FInstancedStruct> FragmentInstanceList)
+template<UE::Mass::CStructType TStruct>
+void FMassEntityManager::SetEntityFragmentValues(FMassEntityHandle Entity, TConstArrayView<TStruct> FragmentInstanceList)
 {
 	CheckIfEntityIsActive(Entity);
 
@@ -1802,7 +1813,8 @@ void FMassEntityManager::SetEntityFragmentValues(FMassEntityHandle Entity, TArra
 	CurrentArchetype->SetFragmentsData(Entity, FragmentInstanceList);
 }
 
-void FMassEntityManager::BatchSetEntityFragmentValues(const FMassArchetypeEntityCollection& SparseEntities, TArrayView<const FInstancedStruct> FragmentInstanceList)
+template<UE::Mass::CStructType TStruct>
+void FMassEntityManager::BatchSetEntityFragmentValues(const FMassArchetypeEntityCollection& SparseEntities, TConstArrayView<TStruct> FragmentInstanceList)
 {
 	if (FragmentInstanceList.Num())
 	{
@@ -1810,7 +1822,8 @@ void FMassEntityManager::BatchSetEntityFragmentValues(const FMassArchetypeEntity
 	}
 }
 
-void FMassEntityManager::BatchSetEntityFragmentValues(TConstArrayView<FMassArchetypeEntityCollection> EntityCollections, TArrayView<const FInstancedStruct> FragmentInstanceList)
+template<UE::Mass::CStructType TStruct>
+void FMassEntityManager::BatchSetEntityFragmentValues(TConstArrayView<FMassArchetypeEntityCollection> EntityCollections, TConstArrayView<TStruct> FragmentInstanceList)
 {
 	CHECK_SYNC_API();
 
@@ -1826,7 +1839,7 @@ void FMassEntityManager::BatchSetEntityFragmentValues(TConstArrayView<FMassArche
 		FMassArchetypeData* Archetype = SparseEntities.GetArchetype().DataPtr.Get();
 		check(Archetype);
 
-		for (const FInstancedStruct& FragmentTemplate : FragmentInstanceList)
+		for (const TStruct& FragmentTemplate : FragmentInstanceList)
 		{
 			Archetype->SetFragmentData(SparseEntities.GetRanges(), FragmentTemplate);
 		}
diff --git a/Engine/Source/Runtime/MassEntity/Public/MassEntityManager.h b/Engine/Source/Runtime/MassEntity/Public/MassEntityManager.h
--- a/Engine/Source/Runtime/MassEntity/Public/MassEntityManager.h
+++ b/Engine/Source/Runtime/MassEntity/Public/MassEntityManager.h
@@ -54,6 +54,11 @@ namespace UE::Mass
 		struct FObserverLock;
 		struct FCreationContext;
 	}
+
+	template<typename T>
+	concept CStructType = TIsDerivedFrom<Clean<T>, FInstancedStruct>::Value
+		|| TIsDerivedFrom<Clean<T>, FStructView>::Value
+		|| TIsDerivedFrom<Clean<T>, FConstStructView>::Value;
 }
 
 // use REQUESTED_MASS_CONCURRENT_RESERVE=0 in your project's Build.cs file to disable concurrent storage
@@ -531,14 +536,23 @@ public:
 	 *  the given entity does have given fragments. Failing this assumption will cause a check-fail.
 	 *  @param EntityHandle idicates the target entity
 	 */
-	UE_API void SetEntityFragmentValues(FMassEntityHandle EntityHandle, TArrayView<const FInstancedStruct> FragmentInstanceList);
+	void SetEntityFragmentValues(FMassEntityHandle EntityHandle, TArrayView<const FInstancedStruct> FragmentInstanceList);
+
+	template<UE::Mass::CStructType TStruct>
+	void SetEntityFragmentValues(FMassEntityHandle EntityHandle, TConstArrayView<TStruct> FragmentInstanceList);
 
 	/** Copies values from FragmentInstanceList over to fragments of given entities collection. The caller is responsible 
 	 *  for ensuring that the given entity archetype (FMassArchetypeEntityCollection .Archetype) does have given fragments. 
 	 *  Failing this assumption will cause a check-fail. */
-	UE_API void BatchSetEntityFragmentValues(const FMassArchetypeEntityCollection& SparseEntities, TArrayView<const FInstancedStruct> FragmentInstanceList);
+	void BatchSetEntityFragmentValues(const FMassArchetypeEntityCollection& SparseEntities, TArrayView<const FInstancedStruct> FragmentInstanceList);
+
+	template<UE::Mass::CStructType TStruct>
+	void BatchSetEntityFragmentValues(const FMassArchetypeEntityCollection& SparseEntities, TConstArrayView<TStruct> FragmentInstanceList);
 
-	UE_API void BatchSetEntityFragmentValues(TConstArrayView<FMassArchetypeEntityCollection> EntityCollections, TArrayView<const FInstancedStruct> FragmentInstanceList);
+	void BatchSetEntityFragmentValues(TConstArrayView<FMassArchetypeEntityCollection> EntityCollections, TArrayView<const FInstancedStruct> FragmentInstanceList);
+
+	template<UE::Mass::CStructType TStruct>
+	void BatchSetEntityFragmentValues(TConstArrayView<FMassArchetypeEntityCollection> EntityCollections, TConstArrayView<TStruct> FragmentInstanceList);
 
 	/**
 	 * @return whether the given handle represents a valid and built entity
@@ -1476,4 +1490,29 @@ bool FMassEntityManager::BatchCreateRelations(TArrayView<FMassEntityHandle> Subj
 	return BatchCreateRelations(UE::Mass::FTypeManager::MakeTypeHandle(T::StaticStruct()), Subjects, Objects);
 }
 
+inline void FMassEntityManager::SetEntityFragmentValues(FMassEntityHandle EntityHandle, TArrayView<const FInstancedStruct> FragmentInstanceList)
+{
+	SetEntityFragmentValues<FInstancedStruct>(EntityHandle, FragmentInstanceList);
+}
+
+inline void FMassEntityManager::BatchSetEntityFragmentValues(const FMassArchetypeEntityCollection& SparseEntities, TArrayView<const FInstancedStruct> FragmentInstanceList)
+{
+	BatchSetEntityFragmentValues<FInstancedStruct>(SparseEntities, FragmentInstanceList);
+}
+
+inline void FMassEntityManager::BatchSetEntityFragmentValues(TConstArrayView<FMassArchetypeEntityCollection> EntityCollections, TArrayView<const FInstancedStruct> FragmentInstanceList)
+{
+	BatchSetEntityFragmentValues<FInstancedStruct>(EntityCollections, FragmentInstanceList);
+}
+
+extern template UE_API void FMassEntityManager::SetEntityFragmentValues<FInstancedStruct>(FMassEntityHandle EntityHandle, TConstArrayView<FInstancedStruct> FragmentInstanceList);
+extern template UE_API void FMassEntityManager::SetEntityFragmentValues<FStructView>(FMassEntityHandle EntityHandle, TConstArrayView<FStructView> FragmentInstanceList);
+extern template UE_API void FMassEntityManager::SetEntityFragmentValues<FConstStructView>(FMassEntityHandle EntityHandle, TConstArrayView<FConstStructView> FragmentInstanceList);
+extern template UE_API void FMassEntityManager::BatchSetEntityFragmentValues<FInstancedStruct>(const FMassArchetypeEntityCollection& SparseEntities, TConstArrayView<FInstancedStruct> FragmentInstanceList);
+extern template UE_API void FMassEntityManager::BatchSetEntityFragmentValues<FStructView>(const FMassArchetypeEntityCollection& SparseEntities, TConstArrayView<FStructView> FragmentInstanceList);
+extern template UE_API void FMassEntityManager::BatchSetEntityFragmentValues<FConstStructView>(const FMassArchetypeEntityCollection& SparseEntities, TConstArrayView<FConstStructView> FragmentInstanceList);
+extern template UE_API void FMassEntityManager::BatchSetEntityFragmentValues<FInstancedStruct>(TConstArrayView<FMassArchetypeEntityCollection> EntityCollections, TConstArrayView<FInstancedStruct> FragmentInstanceList);
+extern template UE_API void FMassEntityManager::BatchSetEntityFragmentValues<FStructView>(TConstArrayView<FMassArchetypeEntityCollection> EntityCollections, TConstArrayView<FStructView> FragmentInstanceList);
+extern template UE_API void FMassEntityManager::BatchSetEntityFragmentValues<FConstStructView>(TConstArrayView<FMassArchetypeEntityCollection> EntityCollections, TConstArrayView<FConstStructView> FragmentInstanceList);
+
 #undef UE_API
