[Backout] - CL49943571
#rnx
Original CL Desc
-----------------------------------------------------------------
[Mass]
Thread-safe observer lock notifications.
Lock observers during processing phases

#Synced-CL 49747586

diff --git a/Engine/Source/Runtime/MassEntity/Private/MassExecutor.cpp b/Engine/Source/Runtime/MassEntity/Private/MassExecutor.cpp
--- a/Engine/Source/Runtime/MassEntity/Private/MassExecutor.cpp
+++ b/Engine/Source/Runtime/MassEntity/Private/MassExecutor.cpp
@@ -109,10 +109,9 @@ void RunProcessorsView(TArrayView<UMassProcessor* const> Processors, FProcessing
 
 struct FMassExecutorDoneTask
 {
-	FMassExecutorDoneTask(FMassExecutionContext&& InExecutionContext, TFunction<void()> InOnDoneNotification, TFunction<void()> InOnPreCommandFlushNotification, const FString& InDebugName, ENamedThreads::Type InDesiredThread)
+	FMassExecutorDoneTask(FMassExecutionContext&& InExecutionContext, TFunction<void()> InOnDoneNotification, const FString& InDebugName, ENamedThreads::Type InDesiredThread)
 		: ExecutionContext(InExecutionContext)
 		, OnDoneNotification(InOnDoneNotification)
-		, OnPreCommandFlushNotification(InOnPreCommandFlushNotification)
 		, DebugName(InDebugName)
 		, DesiredThread(InDesiredThread)
 	{
@@ -141,20 +140,14 @@ struct FMassExecutorDoneTask
 		}
 
 		UE_LOG(LogMass, Verbose, TEXT("MassExecutor %s tasks DONE"), *DebugName);
-
-		if (OnPreCommandFlushNotification)
-		{
-			OnPreCommandFlushNotification();
-		}
-		
 		ExecutionContext.SetFlushDeferredCommands(true);
 		ExecutionContext.FlushDeferred();
+
 		OnDoneNotification();
 	}
 private:
 	FMassExecutionContext ExecutionContext;
 	TFunction<void()> OnDoneNotification;
-	TFunction<void()> OnPreCommandFlushNotification;
 	FString DebugName;
 	ENamedThreads::Type DesiredThread;
 };
@@ -162,10 +155,9 @@ private:
 struct FMassExecutorDispatchTask
 {
 	FMassExecutorDispatchTask(const FMassExecutionContext& InExecutionContext, const TSharedPtr<FMassEntityManager>& InEntityManager, UMassProcessor& InProcessor
-		, TFunction<void()> InOnDoneNotification, TFunction<void()> InOnPreCommandFlushNotification, const FString& InDebugName, ENamedThreads::Type InDesiredEndTaskThread)
+		, TFunction<void()> InOnDoneNotification, const FString& InDebugName, ENamedThreads::Type InDesiredEndTaskThread)
 		: ExecutionContext(InExecutionContext)
 		, OnDoneNotification(InOnDoneNotification)
-		, OnPreCommandFlushNotification(InOnPreCommandFlushNotification)
 		, DebugName(InDebugName)
 		, EntityManager(InEntityManager)
 		, Processor(InProcessor)
@@ -190,7 +182,7 @@ struct FMassExecutorDispatchTask
 			const FGraphEventArray Prerequisites = { CompletionEvent };
 
 			CompletionEvent = TGraphTask<FMassExecutorDoneTask>::CreateTask(&Prerequisites)
-				.ConstructAndDispatchWhenReady(MoveTemp(ExecutionContextThread), OnDoneNotification, OnPreCommandFlushNotification, Processor.GetName(), DesiredEndTaskThread);
+				.ConstructAndDispatchWhenReady(MoveTemp(ExecutionContextThread), OnDoneNotification, Processor.GetName(), DesiredEndTaskThread);
 		}
 
 		if (CompletionEvent.IsValid())
@@ -201,19 +193,16 @@ struct FMassExecutorDispatchTask
 private:
 	FMassExecutionContext ExecutionContext;
 	TFunction<void()> OnDoneNotification;
-	TFunction<void()> OnPreCommandFlushNotification;
 	FString DebugName;
 	TSharedPtr<FMassEntityManager> EntityManager;
 	UMassProcessor& Processor;
 	ENamedThreads::Type DesiredEndTaskThread;
 };
 
-FGraphEventRef TriggerParallelTasks(UMassProcessor& Processor, FProcessingContext&& ProcessingContext, TFunction<void()> OnDoneNotification, ENamedThreads::Type CurrentThread, TFunction<void()> OnPreCommandFlushNotification)
+FGraphEventRef TriggerParallelTasks(UMassProcessor& Processor, FProcessingContext&& ProcessingContext, TFunction<void()> OnDoneNotification, ENamedThreads::Type CurrentThread)
 {
 	TRACE_CPUPROFILER_EVENT_SCOPE(MassExecutor_RunParallel);
 
-	TSharedPtr<FMassEntityManager> EntityManager = ProcessingContext.EntityManager;
-
 	// We need to transfer ProcessingContext's ExecutionContext - otherwise ProcessingContext's destructor will attempt
 	// flushing stored commands. 
 	FMassExecutionContext ExecutionContext = MoveTemp(ProcessingContext).GetExecutionContext();
@@ -222,8 +211,9 @@ FGraphEventRef TriggerParallelTasks(UMassProcessor& Processor, FProcessingContex
 	{
 		TRACE_CPUPROFILER_EVENT_SCOPE_STR("Dispatch Processors")
 		FGraphEventArray Prerequisites;
+		TSharedPtr<FMassEntityManager> EntityManager = ProcessingContext.EntityManager;
 		DispatchEvent  = TGraphTask<FMassExecutorDispatchTask>::CreateTask(&Prerequisites)
-			.ConstructAndDispatchWhenReady(ExecutionContext, EntityManager, Processor, OnDoneNotification, OnPreCommandFlushNotification, Processor.GetName(), CurrentThread);
+			.ConstructAndDispatchWhenReady(ExecutionContext, EntityManager, Processor, OnDoneNotification, Processor.GetName(), CurrentThread);
 	}
 
 	return DispatchEvent;
@@ -236,7 +226,7 @@ inline FGraphEventRef TriggerParallelTasks(UMassProcessor& Processor, FProcessin
 		, ENamedThreads::Type CurrentThread)
 {
 	FProcessingContext LocalContext = ProcessingContext;
-	return TriggerParallelTasks(Processor, MoveTemp(LocalContext), OnDoneNotification, CurrentThread);
+	return TriggerParallelTasks(Processor, MoveTemp(ProcessingContext), OnDoneNotification, CurrentThread);
 }
 
 } // namespace UE::Mass::Executor
diff --git a/Engine/Source/Runtime/MassEntity/Private/MassObserverManager.cpp b/Engine/Source/Runtime/MassEntity/Private/MassObserverManager.cpp
--- a/Engine/Source/Runtime/MassEntity/Private/MassObserverManager.cpp
+++ b/Engine/Source/Runtime/MassEntity/Private/MassObserverManager.cpp
@@ -565,9 +565,9 @@ void FMassObserverManager::RemoveObserverInstance(const UScriptStruct& ElementTy
 
 TSharedRef<FMassObserverManager::FObserverLock> FMassObserverManager::GetOrMakeObserverLock()
 {
-	if (TSharedPtr<FMassObserverManager::FObserverLock> SharedActiveLock = ActiveObserverLock.Pin())
+	if (ActiveObserverLock.IsValid())
 	{
-		return SharedActiveLock.ToSharedRef();
+		return ActiveObserverLock.Pin().ToSharedRef();
 	}
 	else
 	{
diff --git a/Engine/Source/Runtime/MassEntity/Private/MassObserverNotificationTypes.cpp b/Engine/Source/Runtime/MassEntity/Private/MassObserverNotificationTypes.cpp
--- a/Engine/Source/Runtime/MassEntity/Private/MassObserverNotificationTypes.cpp
+++ b/Engine/Source/Runtime/MassEntity/Private/MassObserverNotificationTypes.cpp
@@ -41,7 +41,6 @@ namespace UE::Mass::ObserverManager
 		checkf(SharedEntityManager->GetObserverManager().LocksCount >= 0
 			, TEXT("%hs: the lock count has become unbalanced."), __FUNCTION__);
 		SharedEntityManager->GetObserverManager().ResumeExecution(*this);
-
 	}
 
 	void FObserverLock::ForceUpdateCurrentThreadID()
@@ -80,8 +79,6 @@ namespace UE::Mass::ObserverManager
 		// if the creation handle isn't set there are no creation ops we know about
 		if (CreationHandle.IsSet())
 		{
-			UE::TScopeLock ScopeLock(Lock->GetNotificationLock());
-
 			const FBufferedNotification& Notification = Lock->GetCreationNotification(CreationHandle);
 
 			if (const UE::Mass::FEntityCollection* CreatedEntities = Notification.AffectedEntities.TryGet<FEntityCollection>())
diff --git a/Engine/Source/Runtime/MassEntity/Private/MassProcessingPhaseManager.cpp b/Engine/Source/Runtime/MassEntity/Private/MassProcessingPhaseManager.cpp
--- a/Engine/Source/Runtime/MassEntity/Private/MassProcessingPhaseManager.cpp
+++ b/Engine/Source/Runtime/MassEntity/Private/MassProcessingPhaseManager.cpp
@@ -79,28 +79,19 @@ void FMassProcessingPhase::ExecuteTick(float DeltaTime, ELevelTick TickType, ENa
 	
 	FMassEntityManager& EntityManager = PhaseManager->GetEntityManagerRef();
 
-	bIsDuringMassProcessing = true;	
+	bIsDuringMassProcessing = true;
 
 	if (bRunInParallelMode && PhaseManager->IsPaused() == false)
 	{
 		bool bWorkRequested = false;
 		if (PhaseProcessor->IsEmpty() == false)
 		{
-			ObserversLockPtr = EntityManager.GetOrMakeObserversLock().ToSharedPtr();
 			FMassProcessingContext Context(EntityManager, DeltaTime);
-
-			auto OnProcessingDone = [this, DeltaTime]()
+			const FGraphEventRef PipelineCompletionEvent = UE::Mass::Executor::TriggerParallelTasks(*PhaseProcessor, MoveTemp(Context), [this, DeltaTime]()
 				{
 					OnParallelExecutionDone(DeltaTime);
-				};
-
-			auto OnPreCommandFlush = [this]()
-				{
-					ObserversLockPtr = nullptr;
-				};
-
-			const FGraphEventRef PipelineCompletionEvent = UE::Mass::Executor::TriggerParallelTasks(*PhaseProcessor, MoveTemp(Context), 
-				OnProcessingDone, CurrentThread, OnPreCommandFlush);
+				}
+				, CurrentThread);
 
 			if (PipelineCompletionEvent.IsValid())
 			{
@@ -110,7 +101,6 @@ void FMassProcessingPhase::ExecuteTick(float DeltaTime, ELevelTick TickType, ENa
 		}
 		if (bWorkRequested == false)
 		{
-			ObserversLockPtr = nullptr;
 			OnParallelExecutionDone(DeltaTime);
 		}
 	}
@@ -118,29 +108,25 @@ void FMassProcessingPhase::ExecuteTick(float DeltaTime, ELevelTick TickType, ENa
 	{
 		if (PhaseManager->IsPaused() == false)
 		{
-			ObserversLockPtr = EntityManager.GetOrMakeObserversLock().ToSharedPtr();
 			// note that it's important to create the processing context in this scope
 			// so that it wraps up its destruction before we call OnPhaseEnd, which in turn will cause
 			// the main EntityManager's command buffer to flush
 			FMassProcessingContext Context(EntityManager, DeltaTime);
 			UE::Mass::Executor::Run(*PhaseProcessor, Context);
-
 		}
 
-		bIsDuringMassProcessing = false;
-		ObserversLockPtr = nullptr;
 		{
 			LLM_SCOPE_BYNAME(TEXT("Mass/PhaseEndDelegate"));
 			OnPhaseEnd.Broadcast(DeltaTime);
 		}
 		PhaseManager->OnPhaseEnd(*this);
+		bIsDuringMassProcessing = false;
 	}
 }
 
 void FMassProcessingPhase::OnParallelExecutionDone(const float DeltaTime)
 {
 	bIsDuringMassProcessing = false;
-	ObserversLockPtr = nullptr;
 	{
 		LLM_SCOPE_BYNAME(TEXT("Mass/PhaseEndDelegate"));
 		OnPhaseEnd.Broadcast(DeltaTime);
diff --git a/Engine/Source/Runtime/MassEntity/Public/MassExecutor.h b/Engine/Source/Runtime/MassEntity/Public/MassExecutor.h
--- a/Engine/Source/Runtime/MassEntity/Public/MassExecutor.h
+++ b/Engine/Source/Runtime/MassEntity/Public/MassExecutor.h
@@ -44,11 +44,10 @@ namespace UE::Mass::Executor
 	 *  Triggers tasks executing Processor (and potentially it's children) and returns the task graph event representing 
 	 *  the task (the event will be "completed" once all the processors finish running). 
 	 *  @param OnDoneNotification will be called after all the processors are done, just after flushing the command buffer.
-	 *  @param OnPreCommandFlushNotification will be called after all the processors are done, just before flushing the command buffer.
+	 *    Note that OnDoneNotification will be executed on GameThread.
 	 */
 	MASSENTITY_API FGraphEventRef TriggerParallelTasks(UMassProcessor& Processor, FProcessingContext&& ProcessingContext, TFunction<void()> OnDoneNotification
-		, ENamedThreads::Type CurrentThread = ENamedThreads::GameThread
-		, TFunction<void()> OnPreCommandFlushNotification = TFunction<void()>());
+		, ENamedThreads::Type CurrentThread = ENamedThreads::GameThread);
 
 	UE_DEPRECATED(5.6, "lvalue flavor of TriggerParallelTasks has been deprevate. Use the rvalue version.")
 	MASSENTITY_API FGraphEventRef TriggerParallelTasks(UMassProcessor& Processor, FProcessingContext& ProcessingContext, TFunction<void()> OnDoneNotification
diff --git a/Engine/Source/Runtime/MassEntity/Public/MassObserverManager.h b/Engine/Source/Runtime/MassEntity/Public/MassObserverManager.h
--- a/Engine/Source/Runtime/MassEntity/Public/MassObserverManager.h
+++ b/Engine/Source/Runtime/MassEntity/Public/MassObserverManager.h
@@ -226,13 +226,13 @@ protected:
 
 	/**
 	 * Never access directly, use GetOrMakeObserverLock or GetOrMakeCreationContext instead.
-	 * The Mass Processing Phase synchronously creates a lock before dispatching parallel mass tasks
-	 * Lock should be thread-safe for creation, fetching, and adding items.
+	 * Note: current lock is single-threaded. There's a path towards making it multithreaded, we'll work on it once we have a use-case
 	 */
 	TWeakPtr<FObserverLock> ActiveObserverLock;
 	int32 LocksCount = 0;
 	TWeakPtr<FCreationContext> ActiveCreationContext;
 
+
 	FMassElementBitSet ObservedElements[static_cast<uint8>(EMassObservedOperation::MAX)];
 
 	UPROPERTY()
diff --git a/Engine/Source/Runtime/MassEntity/Public/MassObserverNotificationTypes.h b/Engine/Source/Runtime/MassEntity/Public/MassObserverNotificationTypes.h
--- a/Engine/Source/Runtime/MassEntity/Public/MassObserverNotificationTypes.h
+++ b/Engine/Source/Runtime/MassEntity/Public/MassObserverNotificationTypes.h
@@ -229,18 +229,13 @@ namespace UE::Mass
 				return BufferedNotifications[CreationHandle];
 			}
 
-			FTransactionallySafeCriticalSection& GetNotificationLock()
-			{
-				return NotificationLock;
-			}
-			
 		private:
 			UE_API explicit FObserverLock(FMassObserverManager& ObserverManager);
 
 			int32 GetOrCreateCreationNotification()
 			{
-				UE::TScopeLock ScopeLock(NotificationLock);
-
+				UE_CHECK_OWNER_THREADID();
+				//if (CreationNotificationIndex.compare_exchange_weak(LocalCreationNotificationIndex, BufferedNotifications.Num()))
 				if (CreationNotificationIndex == INDEX_NONE)
 				{
 					CreationNotificationIndex = BufferedNotifications.Num();
@@ -251,8 +246,7 @@ namespace UE::Mass
 
 			bool ReleaseCreationNotification(FCreationNotificationHandle CreationHandle)
 			{
-				UE::TScopeLock ScopeLock(NotificationLock);
-
+				UE_CHECK_OWNER_THREADID();
 				checkf(BufferedNotifications.IsValidIndex(CreationHandle), TEXT("Given CreationHandle doesn't match stored notifications"));
 
 				int32 CreationHandleOpIndex = CreationHandle;
@@ -268,8 +262,7 @@ namespace UE::Mass
 
 			int32 AddCreatedEntity(FMassEntityHandle CreatedEntity)
 			{
-				UE::TScopeLock ScopeLock(NotificationLock);
-
+				UE_CHECK_OWNER_THREADID();
 				if (CreationNotificationIndex == INDEX_NONE)
 				{
 					CreationNotificationIndex = BufferedNotifications.Num();
@@ -285,8 +278,7 @@ namespace UE::Mass
 
 			int32 AddCreatedEntities(const TConstArrayView<FMassEntityHandle> InCreatedEntities)
 			{
-				UE::TScopeLock ScopeLock(NotificationLock);
-
+				UE_CHECK_OWNER_THREADID();
 				if (CreationNotificationIndex == INDEX_NONE)
 				{
 					CreationNotificationIndex = BufferedNotifications.Num();
@@ -302,8 +294,7 @@ namespace UE::Mass
 
 			int32 AddCreatedEntities(const TConstArrayView<FMassEntityHandle> InCreatedEntities, FMassArchetypeEntityCollection&& InEntityCollection)
 			{
-				UE::TScopeLock ScopeLock(NotificationLock);
-
+				UE_CHECK_OWNER_THREADID();
 				if (CreationNotificationIndex == INDEX_NONE)
 				{
 					CreationNotificationIndex = BufferedNotifications.Num();
@@ -321,8 +312,7 @@ namespace UE::Mass
 			template<typename T> requires std::is_same_v<typename TDecay<T>::Type, FMassArchetypeEntityCollection>
 			int32 AddCreatedEntitiesCollection(T&& InEntityCollection)
 			{
-				UE::TScopeLock ScopeLock(NotificationLock);
-
+				UE_CHECK_OWNER_THREADID();
 				if (CreationNotificationIndex == INDEX_NONE)
 				{
 					CreationNotificationIndex = BufferedNotifications.Num();
@@ -339,8 +329,6 @@ namespace UE::Mass
 			template<typename TEntities>
 			void AddNotification(const EMassObservedOperation OperationType, TEntities&& Entities, FMassElementBitSet&& OverlapBitSet)
 			{
-				UE::TScopeLock ScopeLock(NotificationLock);
-
 				if (BufferedNotifications.Num()
 					&& BufferedNotifications.Last().Type == OperationType
 					&& FBufferedNotification::AreCompositionsEqual(BufferedNotifications.Last().CompositionChange, OverlapBitSet))
@@ -375,7 +363,6 @@ namespace UE::Mass
 	#if WITH_MASSENTITY_DEBUG
 			uint32 LockSerialNumber = 0;
 	#endif // WITH_MASSENTITY_DEBUG
-			FTransactionallySafeCriticalSection NotificationLock;
 		};
 
 		/**
diff --git a/Engine/Source/Runtime/MassEntity/Public/MassProcessingPhaseManager.h b/Engine/Source/Runtime/MassEntity/Public/MassProcessingPhaseManager.h
--- a/Engine/Source/Runtime/MassEntity/Public/MassProcessingPhaseManager.h
+++ b/Engine/Source/Runtime/MassEntity/Public/MassProcessingPhaseManager.h
@@ -95,8 +95,6 @@ protected:
 private:
 	FMassProcessingPhaseManager* PhaseManager = nullptr;
 	std::atomic<bool> bIsDuringMassProcessing = false;
-	TSharedPtr<FMassObserverManager::FObserverLock> ObserversLockPtr;
-
 	bool bRunInParallelMode = true;
 	uint8 SupportedTickTypes = 0;
 };
