[Mass] fixed typos and removed deprecated methods

#rnx

#Synced-CL 49939927

diff --git a/Engine/Source/Runtime/MassEntity/Private/MassEntityManager.cpp b/Engine/Source/Runtime/MassEntity/Private/MassEntityManager.cpp
--- a/Engine/Source/Runtime/MassEntity/Private/MassEntityManager.cpp
+++ b/Engine/Source/Runtime/MassEntity/Private/MassEntityManager.cpp
@@ -2260,29 +2260,6 @@ bool FMassEntityManager::DebugHasCommandsToFlush() const
 }
 #endif // WITH_MASSENTITY_DEBUG
 
-//-----------------------------------------------------------------------------
-// DEPRECATED
-//-----------------------------------------------------------------------------
-FMassArchetypeHandle FMassEntityManager::InternalCreateSimilarArchetype(const FMassArchetypeData& SourceArchetypeRef, FMassArchetypeCompositionDescriptor&& NewComposition)
-{
-	return InternalCreateSimilarArchetype(SourceArchetypeRef, Forward<FMassArchetypeCompositionDescriptor>(NewComposition), {});
-}
-
-void FMassEntityManager::SetEntityFragmentsValues(FMassEntityHandle Entity, TArrayView<const FInstancedStruct> FragmentInstanceList)
-{
-	SetEntityFragmentValues(Entity, FragmentInstanceList);
-}
-
-void FMassEntityManager::BatchSetEntityFragmentsValues(const FMassArchetypeEntityCollection& SparseEntities, TArrayView<const FInstancedStruct> FragmentInstanceList)
-{
-	ensureMsgf(false, TEXT("The static BatchSetEntityFragmentsValues is not expected to be called anymore. There's no way to deduce the EntityManager instance related to the call"));
-}
-
-void FMassEntityManager::BatchSetEntityFragmentsValues(TConstArrayView<FMassArchetypeEntityCollection> EntityCollections, TArrayView<const FInstancedStruct> FragmentInstanceList)
-{
-	ensureMsgf(false, TEXT("The static BatchSetEntityFragmentsValues is not expected to be called anymore. There's no way to deduce the EntityManager instance related to the call"));
-}
-
 #undef CHECK_SYNC_API
 #undef CHECK_SYNC_API_RETURN
 #undef CHECK_ELEMENT
\ No newline at end of file
diff --git a/Engine/Source/Runtime/MassEntity/Public/MassEntityManager.h b/Engine/Source/Runtime/MassEntity/Public/MassEntityManager.h
--- a/Engine/Source/Runtime/MassEntity/Public/MassEntityManager.h
+++ b/Engine/Source/Runtime/MassEntity/Public/MassEntityManager.h
@@ -525,7 +525,7 @@ public:
 
 	/** 
 	 * Moves an entity over to a new archetype by copying over fragments common to both archetypes
-	 * @param EntityHandle idicates the entity to move 
+	 * @param EntityHandle indicates the entity to move 
 	 * @param NewArchetypeHandle the handle to the new archetype
 	 * @param SharedFragmentValuesOverride if provided will override all given entity's shared fragment values
 	 */
@@ -534,7 +534,7 @@ public:
 	/** 
 	 *  Copies values from FragmentInstanceList over to target entity's fragment. Caller is responsible for ensuring that 
 	 *  the given entity does have given fragments. Failing this assumption will cause a check-fail.
-	 *  @param EntityHandle idicates the target entity
+	 *  @param EntityHandle indicates the target entity
 	 */
 	void SetEntityFragmentValues(FMassEntityHandle EntityHandle, TArrayView<const FInstancedStruct> FragmentInstanceList);
 
@@ -1365,34 +1365,6 @@ private:
 	};
 
 	EInitializationState InitializationState = EInitializationState::Uninitialized;
-
-	//-----------------------------------------------------------------------------
-	// DEPRECATED
-	//-----------------------------------------------------------------------------
-protected:
-	UE_DEPRECATED(5.6, "This flavor of InternalCreateSimilarArchetype is deprecated due to the introduction of archetype grouping. Use InternalCreateSimilarArchetype with a FArchetypeGroups parameter instead")
-	UE_API FMassArchetypeHandle InternalCreateSimilarArchetype(const FMassArchetypeData& SourceArchetypeRef, FMassArchetypeCompositionDescriptor&& NewComposition);
-public:
-	UE_DEPRECATED(5.6, "SetEntityFragmentsValues is deprecated. Use SetEntityFragmentValues instead (note the slight change in name).")
-	UE_API void SetEntityFragmentsValues(FMassEntityHandle EntityHandle, TArrayView<const FInstancedStruct> FragmentInstanceList);
-
-	UE_DEPRECATED(5.6, "Static BatchSetEntityFragmentsValues is deprecated. Use EntityManager's member function BatchSetEntityFragmentValues (note the slight change in name).")
-	static UE_API void BatchSetEntityFragmentsValues(const FMassArchetypeEntityCollection& SparseEntities, TArrayView<const FInstancedStruct> FragmentInstanceList);
-
-	UE_DEPRECATED(5.6, "Static BatchSetEntityFragmentsValues is deprecated. Use EntityManager's member function BatchSetEntityFragmentValues (note the slight change in name).")
-	static UE_API void BatchSetEntityFragmentsValues(TConstArrayView<FMassArchetypeEntityCollection> EntityCollections, TArrayView<const FInstancedStruct> FragmentInstanceList);
-
-	template<UE::Mass::CConstSharedFragment T>
-	UE_DEPRECATED(5.6, "Using ForEachSharedFragment for Const Shared Fragments has been deprecated. Use ForEachConstSharedFragment instead.")
-	void ForEachSharedFragment(TFunctionRef< void(T& /*SharedFragment*/) > ExecuteFunction)
-	{
-	}
-
-	template<UE::Mass::CConstSharedFragment T>
-	UE_DEPRECATED(5.6, "Using ForEachSharedFragmentConditional for Const Shared Fragments has been deprecated. Use ForEachConstSharedFragmentConditional instead.")
-	void ForEachSharedFragmentConditional(TFunctionRef< bool(T& /*SharedFragment*/) > ConditionFunction, TFunctionRef< void(T& /*SharedFragment*/) > ExecuteFunction)
-	{
-	}
 };
 
 //-----------------------------------------------------------------------------
