[Mass] Fixed non-unity compilation issue (missing header) in MassEntityBuilder.cpp

#rnx

#Synced-CL 50154785

diff --git a/Engine/Source/Runtime/MassEntity/Private/MassEntityBuilder.cpp b/Engine/Source/Runtime/MassEntity/Private/MassEntityBuilder.cpp
--- a/Engine/Source/Runtime/MassEntity/Private/MassEntityBuilder.cpp
+++ b/Engine/Source/Runtime/MassEntity/Private/MassEntityBuilder.cpp
@@ -5,9 +5,7 @@
 #include "MassCommands.h"
 #include "MassEntityManager.h"
 #include "VisualLogger/VisualLogger.h"
-#if WITH_MASSENTITY_DEBUG
 #include "HAL/IConsoleManager.h"
-#endif // WITH_MASSENTITY_DEBUG
 
 namespace UE::Mass
 {
