[Mass] Added templated FMassCommandBuffer::PushUniqueCommand

#rnx

#Synced-CL 50149493

diff --git a/Engine/Source/Runtime/MassEntity/Public/MassCommandBuffer.h b/Engine/Source/Runtime/MassEntity/Public/MassCommandBuffer.h
--- a/Engine/Source/Runtime/MassEntity/Public/MassCommandBuffer.h
+++ b/Engine/Source/Runtime/MassEntity/Public/MassCommandBuffer.h
@@ -156,7 +156,7 @@ struct FMassCommandBuffer
 	 * For example, when you implement a command type that has member properties that 
 	 * control how the given command works or what exactly it does.
 	 */
-	void PushUniqueCommand(TUniquePtr<FMassBatchedCommand>&& CommandInstance)
+	FMassBatchedCommand* PushUniqueCommand(TUniquePtr<FMassBatchedCommand>&& CommandInstance)
 	{
 		COMMAND_PUSHING_CHECK();
 
@@ -166,6 +166,17 @@ struct FMassCommandBuffer
 		AppendedCommandInstances.Add(MoveTemp(CommandInstance));
 
 		++ActiveCommandsCounter;
+
+		return AppendedCommandInstances.Last().Get();
+	}
+
+	template<typename TCommand, typename... ArgsType>
+	requires TIsDerivedFrom<UE::Mass::Clean<TCommand>, FMassBatchedCommand>::Value
+		&& (UE::Mass::Command::TCommandTraits<TCommand>::RequiresUniqueHandling == true)
+	TCommand& PushUniqueCommand(ArgsType&&... Args)
+	{
+		TUniquePtr<TCommand> CommandInstance = MakeUnique<TCommand>(Forward<ArgsType>(Args)...);
+		return *static_cast<TCommand*>(PushUniqueCommand(MoveTemp(CommandInstance)));
 	}
 
 	template<typename T>
