[Mass] Renamed FMassEntityManager::AddElementToEntities RemoveElementFromEntities to BatchAddElementToEntities and BatchRemoveElementFromEntities respectively, for consistency with the rest of the API.

Also added a missing FMassCommandRemoveElement command.

#Synced-CL 50154785

diff --git a/Engine/Source/Runtime/MassEntity/Private/MassEntityManager.cpp b/Engine/Source/Runtime/MassEntity/Private/MassEntityManager.cpp
--- a/Engine/Source/Runtime/MassEntity/Private/MassEntityManager.cpp
+++ b/Engine/Source/Runtime/MassEntity/Private/MassEntityManager.cpp
@@ -848,10 +848,10 @@ TSharedPtr<FMassEntityManager::FEntityCreationContext> FMassEntityManager::Inter
 
 	UE_TRACE_MASS_ENTITIES_CREATED(ReservedEntities, *ArchetypeData)
 
-		if (bCreateContext)
-		{
-			return ObserverManager.GetOrMakeCreationContext(ReservedEntities, FMassArchetypeEntityCollection(ArchetypeHandle, MoveTemp(TargetArchetypeEntityRanges)));
-		}
+	if (bCreateContext)
+	{
+		return ObserverManager.GetOrMakeCreationContext(ReservedEntities, FMassArchetypeEntityCollection(ArchetypeHandle, MoveTemp(TargetArchetypeEntityRanges)));
+	}
 
 	return nullptr;
 }
@@ -1390,7 +1390,7 @@ bool FMassEntityManager::DoesEntityHaveElement(FMassEntityHandle Entity, TNotNul
 	return CurrentArchetype->GetCompositionDescriptor().Contains(ElementType);
 }
 
-void FMassEntityManager::AddElementToEntities(TConstArrayView<FMassEntityHandle> Entities, TNotNull<const UScriptStruct*> ElementType)
+void FMassEntityManager::BatchAddElementToEntities(TConstArrayView<FMassEntityHandle> Entities, TNotNull<const UScriptStruct*> ElementType)
 {
 	CHECK_ELEMENT(ElementType);
 
@@ -1407,21 +1407,21 @@ void FMassEntityManager::AddElementToEntities(TConstArrayView<FMassEntityHandle>
 	}
 }
 
-void FMassEntityManager::AddElementToEntity(FMassEntityHandle Entity, TNotNull<const UScriptStruct*> ElementType)
+void FMassEntityManager::AddElementToEntity(FMassEntityHandle EntityHandle, TNotNull<const UScriptStruct*> ElementType)
 {
 	CHECK_ELEMENT(ElementType);
 
 	if (UE::Mass::IsA<FMassTag>(ElementType))
 	{
-		AddTagToEntity(Entity, ElementType);
+		AddTagToEntity(EntityHandle, ElementType);
 	}
 	else
 	{
-		AddFragmentToEntity(Entity, ElementType);
+		AddFragmentToEntity(EntityHandle, ElementType);
 	}
 }
 
-void FMassEntityManager::RemoveElementFromEntities(TConstArrayView<FMassEntityHandle> Entities, TNotNull<const UScriptStruct*> ElementType)
+void FMassEntityManager::BatchRemoveElementFromEntities(TConstArrayView<FMassEntityHandle> Entities, TNotNull<const UScriptStruct*> ElementType)
 {
 	CHECK_ELEMENT(ElementType);
 
diff --git a/Engine/Source/Runtime/MassEntity/Private/MassRelationManager.cpp b/Engine/Source/Runtime/MassEntity/Private/MassRelationManager.cpp
--- a/Engine/Source/Runtime/MassEntity/Private/MassRelationManager.cpp
+++ b/Engine/Source/Runtime/MassEntity/Private/MassRelationManager.cpp
@@ -271,8 +271,8 @@ namespace UE::Mass
 		}
 		else
 		{
-			EntityManager.AddElementToEntities(Objects, ObjectElement);
-			EntityManager.AddElementToEntities(Subjects, SubjectElement);
+			EntityManager.BatchAddElementToEntities(Objects, ObjectElement);
+			EntityManager.BatchAddElementToEntities(Subjects, SubjectElement);
 		}
 
 		if (RelationData.Traits.bHierarchical)
diff --git a/Engine/Source/Runtime/MassEntity/Public/MassCommands.h b/Engine/Source/Runtime/MassEntity/Public/MassCommands.h
--- a/Engine/Source/Runtime/MassEntity/Public/MassCommands.h
+++ b/Engine/Source/Runtime/MassEntity/Public/MassCommands.h
@@ -456,7 +456,7 @@ protected:
 };
 
 /**
- * Command capable of adding any element type, be it a fragment or a tag.
+ * Command adds any element type to entities, be it a fragment or a tag, including sparse element types
  * Note that this type of command can only be added via PushUniqueCommand
  */
 struct FMassCommandAddElement : public FMassBatchedEntityCommand
@@ -472,7 +472,7 @@ protected:
 	virtual void Run(FMassEntityManager& EntityManager) override
 	{
 		TRACE_CPUPROFILER_EVENT_SCOPE(MassCommandAddElement_Execute);
-		EntityManager.AddElementToEntities(TargetEntities, ElementType);
+		EntityManager.BatchAddElementToEntities(TargetEntities, ElementType);
 	}
 
 	TNotNull<const UScriptStruct*> ElementType;
@@ -485,7 +485,39 @@ struct UE::Mass::Command::TCommandTraits<FMassCommandAddElement> final
 	{
 		RequiresUniqueHandling = true
 	};
-};;
+};
+
+/**
+ * Command removes any element type from entities, be it a fragment or a tag, including sparse element types
+ * Note that this type of command can only be added via PushUniqueCommand
+ */
+struct FMassCommandRemoveElement : public FMassBatchedEntityCommand
+{
+	using Super = FMassBatchedEntityCommand;
+
+	FMassCommandRemoveElement(const TNotNull<const UScriptStruct*> InElementType)
+		: Super(EMassCommandOperationType::Remove DEBUG_NAME("RemoveElement"))
+		, ElementType(InElementType)
+	{}
+
+protected:
+	virtual void Run(FMassEntityManager& EntityManager) override
+	{
+		TRACE_CPUPROFILER_EVENT_SCOPE(MassCommandAddElement_Execute);
+		EntityManager.BatchRemoveElementFromEntities(TargetEntities, ElementType);
+	}
+
+	TNotNull<const UScriptStruct*> ElementType;
+};
+
+template<>
+struct UE::Mass::Command::TCommandTraits<FMassCommandRemoveElement> final
+{
+	enum
+	{
+		RequiresUniqueHandling = true
+	};
+};
 
 template<typename... TOthers>
 struct FMassCommandBuildEntity : public FMassCommandAddFragmentInstances<TOthers...>
diff --git a/Engine/Source/Runtime/MassEntity/Public/MassEntityManager.h b/Engine/Source/Runtime/MassEntity/Public/MassEntityManager.h
--- a/Engine/Source/Runtime/MassEntity/Public/MassEntityManager.h
+++ b/Engine/Source/Runtime/MassEntity/Public/MassEntityManager.h
@@ -407,7 +407,7 @@ public:
 	 * Adds ElementType to the entities, treating it accordingly based on which element type it represents (i.e. Fragment or Tag).
 	 * The function will assert on unhandled types.
 	 */
-	void AddElementToEntities(TConstArrayView<FMassEntityHandle> Entities, TNotNull<const UScriptStruct*> ElementType);
+	UE_API void BatchAddElementToEntities(TConstArrayView<FMassEntityHandle> Entities, TNotNull<const UScriptStruct*> ElementType);
 
 	/**
 	 * Adds ElementType to the entity, treating it accordingly based on which element type it represents (i.e. Fragment or Tag).
@@ -419,7 +419,7 @@ public:
 	 * Removes ElementType from the entities, treating it accordingly based on which element type it represents (i.e. Fragment or Tag).
 	 * The function will assert on unhandled types.
 	 */
-	void RemoveElementFromEntities(TConstArrayView<FMassEntityHandle> Entities, TNotNull<const UScriptStruct*> ElementType);
+	UE_API void BatchRemoveElementFromEntities(TConstArrayView<FMassEntityHandle> Entities, TNotNull<const UScriptStruct*> ElementType);
 
 	/*
 	 * Removes ElementType from the entity, treating it accordingly based on which element type it represents (i.e. Fragment or Tag).
@@ -428,7 +428,7 @@ public:
 	void RemoveElementFromEntity(FMassEntityHandle Entity, TNotNull<const UScriptStruct*> ElementType);
 
 	/** @return whether Entity has an element of ElementType */
-	bool DoesEntityHaveElement(FMassEntityHandle Entity, TNotNull<const UScriptStruct*> ElementType) const;
+	UE_API bool DoesEntityHaveElement(FMassEntityHandle Entity, TNotNull<const UScriptStruct*> ElementType) const;
 
 	/** 
 	 * Adds a new const shared fragment to the given entity. Note that it only works if the given entity doesn't have
@@ -1402,6 +1402,21 @@ private:
 	};
 
 	EInitializationState InitializationState = EInitializationState::Uninitialized;
+
+public:
+	//-----------------------------------------------------------------------------
+	// DEPRECATED
+	//-----------------------------------------------------------------------------
+	UE_DEPRECATED(5.8, "Use BatchAddElementToEntities instead")
+	void AddElementToEntities(TConstArrayView<FMassEntityHandle> Entities, TNotNull<const UScriptStruct*> ElementType)
+	{
+		BatchAddElementToEntities(Entities, ElementType);
+	}
+	UE_DEPRECATED(5.8, "Use BatchRemoveElementFromEntities instead")
+	void RemoveElementFromEntities(TConstArrayView<FMassEntityHandle> Entities, TNotNull<const UScriptStruct*> ElementType)
+	{
+		BatchRemoveElementFromEntities(Entities, ElementType);
+	}
 };
 
 //-----------------------------------------------------------------------------
