[Mass] Tiny update to recent FMassCommandBuffer::PushUniqueCommand change, change the return type from a raw pointer to TNotNull

#rnx

#Synced-CL 50246518

diff --git a/Engine/Source/Runtime/MassEntity/Public/MassCommandBuffer.h b/Engine/Source/Runtime/MassEntity/Public/MassCommandBuffer.h
--- a/Engine/Source/Runtime/MassEntity/Public/MassCommandBuffer.h
+++ b/Engine/Source/Runtime/MassEntity/Public/MassCommandBuffer.h
@@ -156,7 +156,7 @@ struct FMassCommandBuffer
 	 * For example, when you implement a command type that has member properties that 
 	 * control how the given command works or what exactly it does.
 	 */
-	FMassBatchedCommand* PushUniqueCommand(TUniquePtr<FMassBatchedCommand>&& CommandInstance)
+	TNotNull<FMassBatchedCommand*> PushUniqueCommand(TUniquePtr<FMassBatchedCommand>&& CommandInstance)
 	{
 		COMMAND_PUSHING_CHECK();
 
