[Mass] compilation fix for templated FMassCommandBuffer::PushUniqueCommand

#rnx

#Synced-CL 50246518

diff --git a/Engine/Source/Runtime/MassEntity/Public/MassCommandBuffer.h b/Engine/Source/Runtime/MassEntity/Public/MassCommandBuffer.h
--- a/Engine/Source/Runtime/MassEntity/Public/MassCommandBuffer.h
+++ b/Engine/Source/Runtime/MassEntity/Public/MassCommandBuffer.h
@@ -176,7 +176,7 @@ struct FMassCommandBuffer
 	TCommand& PushUniqueCommand(ArgsType&&... Args)
 	{
 		TUniquePtr<TCommand> CommandInstance = MakeUnique<TCommand>(Forward<ArgsType>(Args)...);
-		return *static_cast<TCommand*>(PushUniqueCommand(MoveTemp(CommandInstance)));
+		return static_cast<TCommand&>(*PushUniqueCommand(MoveTemp(CommandInstance)));
 	}
 
 	template<typename T>
