[Mass] mass entity world subsystem now flushes remaining commands on deinitialization. This is required to cleanup any elements registered to other subsystems or rendering scene.

#rnx

#Synced-CL 50318118

diff --git a/Engine/Source/Runtime/MassEntity/Private/MassEntitySubsystem.cpp b/Engine/Source/Runtime/MassEntity/Private/MassEntitySubsystem.cpp
--- a/Engine/Source/Runtime/MassEntity/Private/MassEntitySubsystem.cpp
+++ b/Engine/Source/Runtime/MassEntity/Private/MassEntitySubsystem.cpp
@@ -89,6 +89,15 @@ void UMassEntitySubsystem::PostInitialize()
 
 void UMassEntitySubsystem::Deinitialize()
 {
+	// Frameworks like SceneGraph use Mass entities and dedicated fragments to register in multiple
+	// other frameworks for the lifetime of their world objects (e.g., Render proxies, physics objects, navigation elements, etc.).
+	// Deferred commands to destroy those Mass entities will get pushed when removing all these objects
+	// from the scene/world (i.e., UWorld::CleanupWorldInternal), and we need these commands to be flushed to perform proper cleanup.
+	EntityManager->FlushCommands();
+
+	// @todo: Similar scenario exists for entities with a lifetime tied to the world simulation and gated by BeginPlay/EndPlay.
+	// We should consider flushing commands there too (i.e, WorldEndPlay) before all world subsystems and other frameworks handle their endplay callbacks.
+
 	EntityManager->Deinitialize();
 	EntityManager.Reset();
 	Super::Deinitialize();
