[Mass] Changes the return type of FMassCommandBuffer::PushUniqueCommand

#rnx

#Synced-CL 50246518

diff --git a/Engine/Source/Runtime/MassEntity/Public/MassCommandBuffer.h b/Engine/Source/Runtime/MassEntity/Public/MassCommandBuffer.h
--- a/Engine/Source/Runtime/MassEntity/Public/MassCommandBuffer.h
+++ b/Engine/Source/Runtime/MassEntity/Public/MassCommandBuffer.h
@@ -156,7 +156,7 @@ struct FMassCommandBuffer
 	 * For example, when you implement a command type that has member properties that 
 	 * control how the given command works or what exactly it does.
 	 */
-	TNotNull<FMassBatchedCommand*> PushUniqueCommand(TUniquePtr<FMassBatchedCommand>&& CommandInstance)
+	FMassBatchedCommand& PushUniqueCommand(TUniquePtr<FMassBatchedCommand>&& CommandInstance)
 	{
 		COMMAND_PUSHING_CHECK();
 
@@ -167,7 +167,7 @@ struct FMassCommandBuffer
 
 		++ActiveCommandsCounter;
 
-		return AppendedCommandInstances.Last().Get();
+		return *AppendedCommandInstances.Last().Get();
 	}
 
 	template<typename TCommand, typename... ArgsType>
@@ -176,7 +176,7 @@ struct FMassCommandBuffer
 	TCommand& PushUniqueCommand(ArgsType&&... Args)
 	{
 		TUniquePtr<TCommand> CommandInstance = MakeUnique<TCommand>(Forward<ArgsType>(Args)...);
-		return static_cast<TCommand&>(*PushUniqueCommand(MoveTemp(CommandInstance)));
+		return static_cast<TCommand&>(PushUniqueCommand(MoveTemp(CommandInstance)));
 	}
 
 	template<typename T>
