[Mass] Refactored FMassEntityManager's AddTagToEntity, RemoveTagFromEntity and SwapTagsForEntity to use TNotNull for the script struct parameters - the implementation already assumed that was the case, so this change simply formalizes the assumption.

#rnx

#Synced-CL 50154785

diff --git a/Engine/Source/Runtime/MassEntity/Private/MassEntityManager.cpp b/Engine/Source/Runtime/MassEntity/Private/MassEntityManager.cpp
--- a/Engine/Source/Runtime/MassEntity/Private/MassEntityManager.cpp
+++ b/Engine/Source/Runtime/MassEntity/Private/MassEntityManager.cpp
@@ -1292,20 +1292,20 @@ void FMassEntityManager::RemoveFragmentListFromEntity(FMassEntityHandle Entity,
 	}
 }
 
-void FMassEntityManager::SwapTagsForEntity(FMassEntityHandle Entity, const UScriptStruct* OldTagType, const UScriptStruct* NewTagType)
+void FMassEntityManager::SwapTagsForEntity(FMassEntityHandle Entity, TNotNull<const UScriptStruct*> OldTagType, TNotNull<const UScriptStruct*> NewTagType)
 {
 	CHECK_SYNC_API();
 
+	checkf(UE::Mass::IsA<FMassTag>(OldTagType), TEXT("%hs works only with tags while '%s' is not one."), __FUNCTION__, *GetPathNameSafe(OldTagType));
+	checkf(UE::Mass::IsA<FMassTag>(NewTagType), TEXT("%hs works only with tags while '%s' is not one."), __FUNCTION__, *GetPathNameSafe(NewTagType));
+
 	if (OldTagType == NewTagType)
 	{
 		return;
 	}
 
 	CheckIfEntityIsActive(Entity);
-
-	checkf(UE::Mass::IsA<FMassTag>(OldTagType), TEXT("%hs works only with tags while '%s' is not one."), __FUNCTION__, *GetPathNameSafe(OldTagType));
-	checkf(UE::Mass::IsA<FMassTag>(NewTagType), TEXT("%hs works only with tags while '%s' is not one."), __FUNCTION__, *GetPathNameSafe(NewTagType));
-	
+		
 	TSharedPtr<FMassArchetypeData>& CurrentArchetype = GetEntityStorageInterface().GetArchetypeAsShared(Entity.Index);
 	check(CurrentArchetype);
 
@@ -1324,7 +1324,7 @@ void FMassEntityManager::SwapTagsForEntity(FMassEntityHandle Entity, const UScri
 	GetEntityStorageInterface().SetArchetypeFromShared(Entity.Index, NewArchetypeHandle.DataPtr);
 }
 
-void FMassEntityManager::AddTagToEntity(FMassEntityHandle Entity, const UScriptStruct* TagType)
+void FMassEntityManager::AddTagToEntity(FMassEntityHandle Entity, TNotNull<const UScriptStruct*> TagType)
 {
 	CHECK_SYNC_API();
 	checkf(UE::Mass::IsA<FMassTag>(TagType), TEXT("%hs works only with tags while '%s' is not one."), __FUNCTION__, *GetPathNameSafe(TagType));
@@ -1349,7 +1349,7 @@ void FMassEntityManager::AddTagToEntity(FMassEntityHandle Entity, const UScriptS
 	}
 }
 	
-void FMassEntityManager::RemoveTagFromEntity(FMassEntityHandle Entity, const UScriptStruct* TagType)
+void FMassEntityManager::RemoveTagFromEntity(FMassEntityHandle Entity, TNotNull<const UScriptStruct*> TagType)
 {
 	CHECK_SYNC_API();
 	checkf(UE::Mass::IsA<FMassTag>(TagType), TEXT("%hs works only with tags while '%s' is not one."), __FUNCTION__, *GetPathNameSafe(TagType));
diff --git a/Engine/Source/Runtime/MassEntity/Public/MassEntityManager.h b/Engine/Source/Runtime/MassEntity/Public/MassEntityManager.h
--- a/Engine/Source/Runtime/MassEntity/Public/MassEntityManager.h
+++ b/Engine/Source/Runtime/MassEntity/Public/MassEntityManager.h
@@ -399,9 +399,9 @@ public:
 	UE_API void RemoveFragmentFromEntity(FMassEntityHandle EntityHandle, const UScriptStruct* FragmentType);
 	UE_API void RemoveFragmentListFromEntity(FMassEntityHandle EntityHandle, TConstArrayView<const UScriptStruct*> FragmentList);
 
-	UE_API void AddTagToEntity(FMassEntityHandle EntityHandle, const UScriptStruct* TagType);
-	UE_API void RemoveTagFromEntity(FMassEntityHandle EntityHandle, const UScriptStruct* TagType);
-	UE_API void SwapTagsForEntity(FMassEntityHandle EntityHandle, const UScriptStruct* FromFragmentType, const UScriptStruct* ToFragmentType);
+	UE_API void AddTagToEntity(FMassEntityHandle EntityHandle, TNotNull<const UScriptStruct*> TagType);
+	UE_API void RemoveTagFromEntity(FMassEntityHandle EntityHandle, TNotNull<const UScriptStruct*> TagType);
+	UE_API void SwapTagsForEntity(FMassEntityHandle EntityHandle, TNotNull<const UScriptStruct*> FromFragmentType, TNotNull<const UScriptStruct*> ToFragmentType);
 
 	/**
 	 * Adds ElementType to the entities, treating it accordingly based on which element type it represents (i.e. Fragment or Tag).
