[Mass] reducing FMassArchetypeData's reliance on FMassArchetypeCompositionDescriptor (it's a wrapper for FMassElementBitSet now)

#rnx

#Synced-CL 50154785

diff --git a/Engine/Source/Runtime/MassEntity/Internal/MassArchetypeData.h b/Engine/Source/Runtime/MassEntity/Internal/MassArchetypeData.h
--- a/Engine/Source/Runtime/MassEntity/Internal/MassArchetypeData.h
+++ b/Engine/Source/Runtime/MassEntity/Internal/MassArchetypeData.h
@@ -253,16 +253,16 @@ public:
 	bool HasFragmentType(const UScriptStruct* FragmentType) const;
 	bool HasTagType(const UScriptStruct* TagType) const;
 
-	bool IsEquivalent(const FMassArchetypeCompositionDescriptor& OtherCompositionDescriptor, const UE::Mass::FArchetypeGroups& OtherGroups) const;
+	bool IsEquivalent(const FMassElementBitSet& OtherCompositionBitSet, const UE::Mass::FArchetypeGroups& OtherGroups) const;
 
-	void Initialize(const FMassEntityManager& EntityManager, const FMassArchetypeCompositionDescriptor& InCompositionDescriptor, const uint32 ArchetypeDataVersion);
+	void Initialize(const FMassEntityManager& EntityManager, const FMassElementBitSet& InCompositionBitSet, const uint32 ArchetypeDataVersion);
 
 	/** 
 	 * A special way of initializing an archetype resulting in a copy of BaseArchetype's setup with OverrideTags
 	 * replacing original tags of BaseArchetype
 	 */
 	void InitializeWithSimilar(const FMassEntityManager& EntityManager, const FMassArchetypeData& BaseArchetype
-		, FMassArchetypeCompositionDescriptor&& NewComposition, const UE::Mass::FArchetypeGroups& InGroups, const uint32 ArchetypeDataVersion);
+		, FMassElementBitSet&& InCompositionBitSet, const UE::Mass::FArchetypeGroups& InGroups, const uint32 ArchetypeDataVersion);
 
 	void AddEntity(FMassEntityHandle Entity, const FMassArchetypeSharedFragmentValues& InSharedFragmentValues);
 	void RemoveEntity(FMassEntityHandle Entity);
diff --git a/Engine/Source/Runtime/MassEntity/Private/MassArchetypeData.cpp b/Engine/Source/Runtime/MassEntity/Private/MassArchetypeData.cpp
--- a/Engine/Source/Runtime/MassEntity/Private/MassArchetypeData.cpp
+++ b/Engine/Source/Runtime/MassEntity/Private/MassArchetypeData.cpp
@@ -64,7 +64,7 @@ bool FMassArchetypeData::HasTagType(const UScriptStruct* TagType) const
 		&& CompositionBitSet.Contains(*TagType);
 }
 
-void FMassArchetypeData::Initialize(const FMassEntityManager& EntityManager, const FMassArchetypeCompositionDescriptor& InCompositionDescriptor, const uint32 ArchetypeDataVersion)
+void FMassArchetypeData::Initialize(const FMassEntityManager& EntityManager, const FMassElementBitSet& InCompositionBitSet, const uint32 ArchetypeDataVersion)
 {
 	if (!ensureMsgf(Chunks.Num() == 0, TEXT("Trying to re-initialize non-empty Mass Archetype is not supported")))
 	{
@@ -77,7 +77,7 @@ void FMassArchetypeData::Initialize(const FMassEntityManager& EntityManager, con
 
 	CreatedArchetypeDataVersion = ArchetypeDataVersion;
 
-	CompositionBitSet = InCompositionDescriptor.GetElementsBitSet();
+	CompositionBitSet = InCompositionBitSet;
 
 	ConfigureFragments(EntityManager);
 
@@ -98,21 +98,20 @@ void FMassArchetypeData::Initialize(const FMassEntityManager& EntityManager, con
 #endif // WITH_MASSENTITY_DEBUG
 }
 
-// @todo moving FMassArchetypeCompositionDescriptor doesn't make sense anymore, deprecated this constructor and optionally introduce a FMassElementBitSet&& one
 void FMassArchetypeData::InitializeWithSimilar(const FMassEntityManager& EntityManager, const FMassArchetypeData& BaseArchetype
-	, FMassArchetypeCompositionDescriptor&& NewComposition, const UE::Mass::FArchetypeGroups& InGroups, const uint32 ArchetypeDataVersion)
+	, FMassElementBitSet&& InCompositionBitSet, const UE::Mass::FArchetypeGroups& InGroups, const uint32 ArchetypeDataVersion)
 {
 	checkf(IsInitialized() == false, TEXT("Trying to InitializeWithSimilar but this archetype has already been initialized"));
 
-	// NewComposition.GetElementsBitSet() - BaseArchetype.GetCompositionBitSet() -> should contain only tags and fragments
-	// BaseArchetype.GetCompositionBitSet() - NewComposition.GetElementsBitSet() -> should contain only tags and fragments
-	bool bAddsOnlyTagsAndElements = FMassElementBitSet(NewComposition.GetElementsBitSet() - BaseArchetype.GetCompositionBitSet()).ContainsOnlyTagsAndFragments();
-	bool bNothingElse = FMassElementBitSet(BaseArchetype.GetCompositionBitSet() - NewComposition.GetElementsBitSet()).ContainsOnlyTagsAndFragments();
+	// InCompositionBitSet - BaseArchetype.GetCompositionBitSet() -> should contain only tags and fragments
+	// BaseArchetype.GetCompositionBitSet() - InCompositionBitSet -> should contain only tags and fragments
+	bool bAddsOnlyTagsAndElements = FMassElementBitSet(InCompositionBitSet - BaseArchetype.GetCompositionBitSet()).ContainsOnlyTagsAndFragments();
+	bool bNothingElse = FMassElementBitSet(BaseArchetype.GetCompositionBitSet() - InCompositionBitSet).ContainsOnlyTagsAndFragments();
 	checkf(bAddsOnlyTagsAndElements && bNothingElse, TEXT("NewComposition is expected to differ from the BaseArchetype's composition just by tags and fragments"))
 
 	CreatedArchetypeDataVersion = ArchetypeDataVersion;
 
-	CompositionBitSet = NewComposition.GetElementsBitSet();
+	CompositionBitSet = MoveTemp(InCompositionBitSet);
 	if (CompositionBitSet.Get<FMassFragmentBitSet>() != BaseArchetype.CompositionBitSet.Get<FMassFragmentBitSet>())
 	{
 		ConfigureFragments(EntityManager);
@@ -1443,9 +1442,9 @@ void FMassArchetypeData::BatchSetFragmentValues(TConstArrayView<FMassArchetypeEn
 	}
 }
 
-bool FMassArchetypeData::IsEquivalent(const FMassArchetypeCompositionDescriptor& OtherCompositionDescriptor, const UE::Mass::FArchetypeGroups& OtherGroups) const
+bool FMassArchetypeData::IsEquivalent(const FMassElementBitSet& OtherCompositionBitSet, const UE::Mass::FArchetypeGroups& OtherGroups) const
 {
-	return CompositionBitSet.IsEquivalent(OtherCompositionDescriptor.GetElementsBitSet()) && Groups == OtherGroups;
+	return CompositionBitSet.IsEquivalent(OtherCompositionBitSet) && Groups == OtherGroups;
 }
 
 //-----------------------------------------------------------------------------
diff --git a/Engine/Source/Runtime/MassEntity/Private/MassEntityManager.cpp b/Engine/Source/Runtime/MassEntity/Private/MassEntityManager.cpp
--- a/Engine/Source/Runtime/MassEntity/Private/MassEntityManager.cpp
+++ b/Engine/Source/Runtime/MassEntity/Private/MassEntityManager.cpp
@@ -395,7 +395,7 @@ FMassArchetypeHandle FMassEntityManager::CreateArchetype(const FMassArchetypeCom
 	TSharedPtr<FMassArchetypeData> ArchetypeDataPtr;
 	for (const TSharedPtr<FMassArchetypeData>& Ptr : HashRow)
 	{
-		if (Ptr->IsEquivalent(Composition, /*Groups=*/{}))
+		if (Ptr->IsEquivalent(Composition.GetElementsBitSet(), /*Groups=*/{}))
 		{
 #if WITH_MASSENTITY_DEBUG
 			// Keep track of all names for this archetype.
@@ -421,7 +421,7 @@ FMassArchetypeHandle FMassEntityManager::CreateArchetype(const FMassArchetypeCom
 
 		// Create a new archetype
 		FMassArchetypeData* NewArchetype = new FMassArchetypeData(CreationParams);
-		NewArchetype->Initialize(*this, Composition, ArchetypeDataVersion);
+		NewArchetype->Initialize(*this, Composition.GetElementsBitSet(), ArchetypeDataVersion);
 		ArchetypeDataPtr = HashRow.Add_GetRef(MakeShareable(NewArchetype));
 		AllArchetypes.Add(ArchetypeDataPtr);
 		ensure(AllArchetypes.Num() == ArchetypeDataVersion);
@@ -468,7 +468,7 @@ FMassArchetypeHandle FMassEntityManager::InternalCreateSimilarArchetype(const FM
 	TSharedPtr<FMassArchetypeData> ArchetypeDataPtr;
 	for (const TSharedPtr<FMassArchetypeData>& Ptr : HashRow)
 	{
-		if (Ptr->IsEquivalent(NewComposition, Groups))
+		if (Ptr->IsEquivalent(NewComposition.GetElementsBitSet(), Groups))
 		{
 			ArchetypeDataPtr = Ptr;
 			break;
@@ -482,7 +482,7 @@ FMassArchetypeHandle FMassEntityManager::InternalCreateSimilarArchetype(const FM
 
 		// Create a new archetype
 		FMassArchetypeData* NewArchetype = new FMassArchetypeData(FMassArchetypeCreationParams(SourceArchetypeRef));
-		NewArchetype->InitializeWithSimilar(*this, SourceArchetypeRef, MoveTemp(NewComposition), Groups, ArchetypeDataVersion);
+		NewArchetype->InitializeWithSimilar(*this, SourceArchetypeRef, MoveTemp(NewComposition.GetElementsBitSet()), Groups, ArchetypeDataVersion);
 		NewArchetype->CopyDebugNamesFrom(SourceArchetypeRef);
 
 		ArchetypeDataPtr = HashRow.Add_GetRef(MakeShareable(NewArchetype));
