[Mass] code cleanup following sparse elements submission

#rnx

#Synced-CL 50439895

diff --git a/Engine/Source/Developer/MassEntityTestSuite/Private/MassSparseElementsTest.cpp b/Engine/Source/Developer/MassEntityTestSuite/Private/MassSparseElementsTest.cpp
--- a/Engine/Source/Developer/MassEntityTestSuite/Private/MassSparseElementsTest.cpp
+++ b/Engine/Source/Developer/MassEntityTestSuite/Private/MassSparseElementsTest.cpp
@@ -49,12 +49,12 @@ namespace UE::Mass::Test::SparseElements
 	{
 		static bool HasElement(const FMassArchetypeHandle& ArchetypeHandle, TNotNull<UScriptStruct*> SparseElementType, bool bIncludeTransient)
 		{
-			return FMassArchetypeHelper::ArchetypeDataFromHandleChecked(ArchetypeHandle).HasTagType(SparseElementType, bIncludeTransient);
+			return FMassArchetypeHelper::ArchetypeDataFromHandleChecked(ArchetypeHandle).HasElement(SparseElementType, bIncludeTransient ? EIncludeSparseElements::Yes : EIncludeSparseElements::No);
 		}
 
 		static bool HasElement(FMassEntityView EntityView, TNotNull<UScriptStruct*> SparseElementType)
 		{
-			return EntityView.HasTag(SparseElementType);
+			return EntityView.HasElement(SparseElementType);
 		}
 
 		TNotNull<UScriptStruct*> MakeSparseElement()
@@ -74,12 +74,12 @@ namespace UE::Mass::Test::SparseElements
 	{
 		static bool HasElement(const FMassArchetypeHandle& ArchetypeHandle, TNotNull<UScriptStruct*> SparseElementType, bool bIncludeTransient)
 		{
-			return FMassArchetypeHelper::ArchetypeDataFromHandleChecked(ArchetypeHandle).HasFragmentType(SparseElementType, bIncludeTransient);
+			return FMassArchetypeHelper::ArchetypeDataFromHandleChecked(ArchetypeHandle).HasElement(SparseElementType, bIncludeTransient ? EIncludeSparseElements::Yes : EIncludeSparseElements::No);
 		}
 
 		static bool HasElement(FMassEntityView EntityView, TNotNull<UScriptStruct*> SparseElementType)
 		{
-			return EntityView.HasFragment(SparseElementType);
+			return EntityView.HasElement(SparseElementType);
 		}
 
 		TNotNull<UScriptStruct*> MakeSparseElement()
diff --git a/Engine/Source/Runtime/MassEntity/Internal/MassArchetypeData.h b/Engine/Source/Runtime/MassEntity/Internal/MassArchetypeData.h
--- a/Engine/Source/Runtime/MassEntity/Internal/MassArchetypeData.h
+++ b/Engine/Source/Runtime/MassEntity/Internal/MassArchetypeData.h
@@ -377,37 +377,32 @@ public:
 
 	/**
 	 * checks whether a given fragment type is part of this archetype's entities' composition
-	 * @param bIncludeSparse if true and FragmentType is a sparse element, the function will
-	 *	check if it hosts any entities owning such a fragment
 	 */
-	bool HasFragmentType(const UScriptStruct* FragmentType, const bool bIncludeSparse = false) const
+	bool HasFragmentType(const UScriptStruct* FragmentType) const
 	{
-		return (UE::Mass::IsA<FMassFragment>(FragmentType) 
-			&& (CompositionBitSet.Contains(*FragmentType)
-				|| (bIncludeSparse && UE::Mass::IsSparse(FragmentType) && SparseElementsBitSet.Contains(*FragmentType))));
+		return UE::Mass::IsA<FMassFragment>(FragmentType) 
+			&& CompositionBitSet.Contains(*FragmentType);
 	}
 
 	/**
 	 * checks whether a given tag type is part of this archetype's entities' composition
-	 * @param bIncludeSparse if true and TagType is a sparse element, the function will
-	 *	check if it hosts any entities owning such a tag
 	 */
-	bool HasTagType(const UScriptStruct* TagType, const bool bIncludeSparse = false) const
+	bool HasTagType(const UScriptStruct* TagType) const
 	{
-		return (UE::Mass::IsA<FMassTag>(TagType)
-			&& (CompositionBitSet.Contains(*TagType)
-				|| (bIncludeSparse && UE::Mass::IsSparse(TagType) && SparseElementsBitSet.Contains(*TagType))));
+		return UE::Mass::IsA<FMassTag>(TagType)
+			&& CompositionBitSet.Contains(*TagType);
 	}
 
 	/**
 	 * checks whether a given element type is part of this archetype's entities' composition
-	 * @param bIncludeSparse if true and ElementType is a sparse element, the function will
-	 *	check if it hosts any entities owning such an element
+	 * @param IncludeSparseElements whether the function should check if the archetype hosts any entities
+	 *	owning such a fragment, provided FragmentType represents a sparse element type
 	 */
-	bool HasElement(const UScriptStruct* ElementType, const bool bIncludeSparse = false) const
+	bool HasElement(const UScriptStruct* ElementType, const UE::Mass::EIncludeSparseElements IncludeSparseElements/* = UE::Mass::EIncludeSparseElements::No*/) const
 	{
 		return (CompositionBitSet.Contains(*ElementType)
-				|| (bIncludeSparse && UE::Mass::IsSparse(ElementType) && SparseElementsBitSet.Contains(*ElementType)));
+				|| (IncludeSparseElements == UE::Mass::EIncludeSparseElements::Yes
+					&& UE::Mass::IsSparse(ElementType) && SparseElementsBitSet.Contains(*ElementType)));
 	}
 
 	bool IsEquivalent(const FMassElementBitSet& OtherCompositionBitSet, const UE::Mass::FArchetypeGroups& OtherGroups) const;
@@ -426,7 +421,10 @@ public:
 	void AddEntity(FMassEntityHandle Entity, const FMassArchetypeSharedFragmentValues& InSharedFragmentValues);
 	void RemoveEntity(FMassEntityHandle Entity);
 
-	/** @return if ElementType is not a tag the function will return FInstancedStruct ptr pointing at the appropriate data */
+	/**
+	 * @return if ElementType is a tag the function will return an empty struct view. Otherwise, the function will
+	 *	return a valid view to the element owned by the entity, provided it does own one, or an empty view otherwise.
+	 */
 	FStructView AddSparseElementToEntity(FMassEntityHandle EntityHandle, TNotNull<const UScriptStruct*> ElementType);
 	/** Works with all sparse element types (fragments and tags). */
 	bool RemoveSparseElementFromEntity(FMassEntityHandle EntityHandle, TNotNull<const UScriptStruct*> ElementType);
@@ -442,17 +440,6 @@ public:
 
 	bool HasSparseElementForEntity(TNotNull<const UScriptStruct*> ElementType, const FMassEntityHandle EntityHandle) const;
 	bool HasSparseElementForEntity(TNotNull<const UScriptStruct*> ElementType, const FMassEntityInChunkDataHandle& EntityInChunkHandle) const;
-	bool DoesChunkContainSparseElement(TNotNull<const UScriptStruct*> ElementType, const int32 ChunkIndex) const
-	{
-		checkf(Chunks.IsValidIndex(ChunkIndex), TEXT("Invalid chunk index (%d out of %d)"), ChunkIndex, Chunks.Num());
-
-		if (SparseElementsBitSet.Contains(*ElementType) && Chunks[ChunkIndex].HasSparseElements())
-		{
-			const UE::Mass::FChunkSparseElements& SparseElements = Chunks[ChunkIndex].GetSparseElementsUnsafe();
-			return (SparseElements.IsEmpty() == false) && SparseElements.Contains(ElementType);
-		}
-		return false;
-	}
 	bool DoesEntityMatchSparseComposition(const FMassEntityInChunkDataHandle& EntityInChunkHandle
 		, const FMassElementBitSet& RequiredAllSparseElements
 		, const FMassElementBitSet& RequiredAnySparseElements
diff --git a/Engine/Source/Runtime/MassEntity/Private/MassEntityView.cpp b/Engine/Source/Runtime/MassEntity/Private/MassEntityView.cpp
--- a/Engine/Source/Runtime/MassEntity/Private/MassEntityView.cpp
+++ b/Engine/Source/Runtime/MassEntity/Private/MassEntityView.cpp
@@ -111,20 +111,22 @@ void* FMassEntityView::GetSharedFragmentPtrChecked(const UScriptStruct& Fragment
 	return SharedFragment->GetMemory();
 }
 
-bool FMassEntityView::HasTag(TNotNull<const UScriptStruct*> TagType, const bool bIncludeSparse) const
+bool FMassEntityView::HasTag(TNotNull<const UScriptStruct*> TagType) const
 {
 	checkSlow(EntityDataHandle.IsValid(Archetype));
-	return Archetype && (Archetype->HasTagType(TagType) || (bIncludeSparse && Archetype->HasSparseElementForEntity(TagType, EntityDataHandle)));
+	return Archetype && Archetype->HasTagType(TagType);
 }
 
-bool FMassEntityView::HasFragment(TNotNull<const UScriptStruct*> FragmentType, const bool bIncludeSparse) const
+bool FMassEntityView::HasFragment(TNotNull<const UScriptStruct*> FragmentType) const
 {
 	checkSlow(EntityDataHandle.IsValid(Archetype));
-	return Archetype && (Archetype->HasFragmentType(FragmentType) || (bIncludeSparse && Archetype->HasSparseElementForEntity(FragmentType, EntityDataHandle)));
+	return Archetype && Archetype->HasFragmentType(FragmentType);
 }
 
-bool FMassEntityView::HasElement(TNotNull<const UScriptStruct*> ElementType, const bool bIncludeSparse) const
+bool FMassEntityView::HasElement(TNotNull<const UScriptStruct*> ElementType, const UE::Mass::EIncludeSparseElements IncludeSparseElements) const
 {
 	checkSlow(EntityDataHandle.IsValid(Archetype));
-	return Archetype && (Archetype->HasElement(ElementType) || (bIncludeSparse && Archetype->HasSparseElementForEntity(ElementType, EntityDataHandle)));
+	return Archetype
+		&& (Archetype->HasElement(ElementType, UE::Mass::EIncludeSparseElements::No) 
+			|| (IncludeSparseElements == UE::Mass::EIncludeSparseElements::Yes && Archetype->HasSparseElementForEntity(ElementType, EntityDataHandle)));
 }
diff --git a/Engine/Source/Runtime/MassEntity/Private/MassExecutionContext.cpp b/Engine/Source/Runtime/MassEntity/Private/MassExecutionContext.cpp
--- a/Engine/Source/Runtime/MassEntity/Private/MassExecutionContext.cpp
+++ b/Engine/Source/Runtime/MassEntity/Private/MassExecutionContext.cpp
@@ -316,12 +316,6 @@ bool FMassExecutionContext::HasSparseElement(TNotNull<const UScriptStruct*> Elem
 	return ArchetypeDataPtr && ArchetypeDataPtr->HasSparseElementForEntity(ElementType, EntityListView[EntityIterator]);
 }
 
-bool FMassExecutionContext::DoesChunkContainSparseElement(TNotNull<const UScriptStruct*> ElementType) const
-{
-	// we're not checking if ChunkEntityIndex.ChunkIndex is valid, ArchetypeDataPtr->DoesChunkContainSparseElement will do that.
-	return ArchetypeDataPtr && ArchetypeDataPtr->DoesChunkContainSparseElement(ElementType, ChunkEntityIndex.ChunkIndex);
-}
-
 bool FMassExecutionContext::ExecuteSparseElementsFilterImpl(const FSparseEntityIterator& EntityIt) const
 {
 	// @todo we could gain some perf if we cached required sparse elements bitsets as member variables. Needs perf measuring if worth it.
diff --git a/Engine/Source/Runtime/MassEntity/Public/MassEntityBuilder.h b/Engine/Source/Runtime/MassEntity/Public/MassEntityBuilder.h
--- a/Engine/Source/Runtime/MassEntity/Public/MassEntityBuilder.h
+++ b/Engine/Source/Runtime/MassEntity/Public/MassEntityBuilder.h
@@ -21,12 +21,6 @@ namespace UE::Mass
 		struct FEntityBuilderHelper;
 	}
 
-	enum class EIncludeSparseElements : uint8
-	{
-		No,
-		Yes
-	};
-
 /**
  * FEntityBuilder is a utility struct that provides a convenient way to create and configure entities
  * in the Mass framework. It bridges multiple APIs from FMassEntityManager, MassSpawnerSubsystem,
diff --git a/Engine/Source/Runtime/MassEntity/Public/MassEntityTypes.h b/Engine/Source/Runtime/MassEntity/Public/MassEntityTypes.h
--- a/Engine/Source/Runtime/MassEntity/Public/MassEntityTypes.h
+++ b/Engine/Source/Runtime/MassEntity/Public/MassEntityTypes.h
@@ -62,6 +62,12 @@ namespace UE::Mass
 		int32 MaxChunkIndex;
 		int32 EntityCountRemaining;
 	};
+
+	enum class EIncludeSparseElements : uint8
+	{
+		No,
+		Yes
+	};
 }
 
 /** The type summarily describing a composition of an entity or an archetype. It contains information on both the
diff --git a/Engine/Source/Runtime/MassEntity/Public/MassEntityView.h b/Engine/Source/Runtime/MassEntity/Public/MassEntityView.h
--- a/Engine/Source/Runtime/MassEntity/Public/MassEntityView.h
+++ b/Engine/Source/Runtime/MassEntity/Public/MassEntityView.h
@@ -137,29 +137,29 @@ struct FMassEntityView
 	}
 
 	template<typename T>
-	bool HasTag(bool bIncludeSparse = true) const
+	bool HasTag() const
 	{
-		return HasTag(*T::StaticStruct(), bIncludeSparse);
+		return HasTag(*T::StaticStruct());
 	}
 
-	UE_API bool HasTag(TNotNull<const UScriptStruct*> TagType, bool bIncludeSparse = true) const;
-	bool HasTag(const UScriptStruct& TagType, bool bIncludeSparse = true) const
+	UE_API bool HasTag(TNotNull<const UScriptStruct*> TagType) const;
+	bool HasTag(const UScriptStruct& TagType) const
 	{
-		return HasTag(&TagType, bIncludeSparse);
+		return HasTag(&TagType);
 	}
 
-	UE_API bool HasFragment(TNotNull<const UScriptStruct*> FragmentType, bool bIncludeSparse = true) const;
+	UE_API bool HasFragment(TNotNull<const UScriptStruct*> FragmentType) const;
 	template<UE::Mass::CFragment T>
-	bool HasFragment(bool bIncludeSparse = true) const
+	bool HasFragment() const
 	{
-		return HasFragment(T::StaticStruct(), bIncludeSparse);
+		return HasFragment(T::StaticStruct());
 	}
 
-	UE_API bool HasElement(TNotNull<const UScriptStruct*> ElementType, bool bIncludeSparse = true) const;
+	UE_API bool HasElement(TNotNull<const UScriptStruct*> ElementType, const UE::Mass::EIncludeSparseElements IncludeSparseElements = UE::Mass::EIncludeSparseElements::Yes) const;
 	template<UE::Mass::CSparse T>
-	bool HasElement(bool bIncludeSparse = true) const
+	bool HasElement(const UE::Mass::EIncludeSparseElements IncludeSparseElements = UE::Mass::EIncludeSparseElements::Yes) const
 	{
-		return HasElement(T::StaticStruct(), bIncludeSparse);
+		return HasElement(T::StaticStruct(), IncludeSparseElements);
 	}
 
 	bool IsSet() const;
diff --git a/Engine/Source/Runtime/MassEntity/Public/MassExecutionContext.h b/Engine/Source/Runtime/MassEntity/Public/MassExecutionContext.h
--- a/Engine/Source/Runtime/MassEntity/Public/MassExecutionContext.h
+++ b/Engine/Source/Runtime/MassEntity/Public/MassExecutionContext.h
@@ -769,7 +769,6 @@ public:
 	}
 
 	MASSENTITY_API bool HasSparseElement(TNotNull<const UScriptStruct*> ElementType, const FEntityIterator& EntityIterator) const;
-	MASSENTITY_API bool DoesChunkContainSparseElement(TNotNull<const UScriptStruct*> ElementType) const;
 
 	/** Sparse chunk related operation */
 	const FMassArchetypeEntityCollection& GetEntityCollection() const { return EntityCollection; }
