[Mass] sparse elements further cleanup

#rnx

#Synced-CL 50439895

#Synced-CL 50439895

diff --git a/Engine/Source/Runtime/MassEntity/Private/MassEntityManager.cpp b/Engine/Source/Runtime/MassEntity/Private/MassEntityManager.cpp
--- a/Engine/Source/Runtime/MassEntity/Private/MassEntityManager.cpp
+++ b/Engine/Source/Runtime/MassEntity/Private/MassEntityManager.cpp
@@ -398,8 +398,6 @@ FMassArchetypeHandle FMassEntityManager::GetOrCreateSuitableArchetype(const FMas
 FMassArchetypeHandle FMassEntityManager::CreateArchetype(const FMassArchetypeCompositionDescriptor& Composition, const FMassArchetypeCreationParams& CreationParams)
 {
 	LLM_SCOPE_BYNAME(TEXT("Mass/EntityManager"));
-	/*FMassElementBitSet AndToo = Composition.GetElementsBitSet() & FMassElementBitSet::GetAllSparseElements();
-	check(AndToo.IsEmpty());*/
 	FMassElementBitSet NonSparseComposition = Composition.GetElementsBitSet().AndNot(FMassElementBitSet::GetAllSparseElements());
 	const uint32 TypeHash = HashCombine(GetTypeHash(NonSparseComposition), GetTypeHash(UE::Mass::FArchetypeGroups()));
 
@@ -474,8 +472,6 @@ FMassArchetypeHandle FMassEntityManager::InternalCreateSimilarArchetype(const FM
 	// be there if IDs were added and removed to this specific Groups container instance
 	checkf(Groups.IsShrunk(), TEXT("A group container with invalid trailing IDs has been passed to archetype creation - this is not expected and will cause issues. Make sure to Shrink your Groups before passing to %hs"), __FUNCTION__);
 
-	/*FMassElementBitSet AndToo = NewComposition.GetElementsBitSet() & FMassElementBitSet::GetAllSparseElements();
-	check(AndToo.IsEmpty());*/
 	FMassElementBitSet NonSparseComposition = NewComposition.GetElementsBitSet().AndNot(FMassElementBitSet::GetAllSparseElements());
 	const uint32 TypeHash = HashCombine(GetTypeHash(NonSparseComposition), GetTypeHash(UE::Mass::FArchetypeGroups()));
 
diff --git a/Engine/Source/Runtime/MassEntity/Public/MassExecutionContext.h b/Engine/Source/Runtime/MassEntity/Public/MassExecutionContext.h
--- a/Engine/Source/Runtime/MassEntity/Public/MassExecutionContext.h
+++ b/Engine/Source/Runtime/MassEntity/Public/MassExecutionContext.h
@@ -308,7 +308,7 @@ public:
 	/**
 	 * A flavor of FEntityIterator that skips all the entities that don't match current query's sparse element requirements.
 	 * This means that if there are no sparse requirements then FSparseEntityIterator will behave like a regular FEntityIterator
-	 * but a less efficiently so.
+	 * but less efficiently so.
 	 */
 	struct FSparseEntityIterator : FEntityIterator
 	{
