[Mass]
Fixed const subsystem access for QueryExecutor.

#Synced-CL 50558293

diff --git a/Engine/Source/Runtime/MassEntity/Public/MassQueryExecutor.h b/Engine/Source/Runtime/MassEntity/Public/MassQueryExecutor.h
--- a/Engine/Source/Runtime/MassEntity/Public/MassQueryExecutor.h
+++ b/Engine/Source/Runtime/MassEntity/Public/MassQueryExecutor.h
@@ -760,7 +760,7 @@ struct FConstSubsystemAccess
 	const TSubsystem& Get() const
 	{
 		check(Subsystem);
-		return Subsystem;
+		return *Subsystem;
 	}
 
 	inline const TSubsystem& operator*() const
