[Mass] removed TStructTypeBitSet::GetAtIndex as it duplicates IsBitSet

#rnx

#Synced-CL 50439895

diff --git a/Engine/Source/Runtime/CoreUObject/Public/StructUtils/StructTypeBitSet.h b/Engine/Source/Runtime/CoreUObject/Public/StructUtils/StructTypeBitSet.h
--- a/Engine/Source/Runtime/CoreUObject/Public/StructUtils/StructTypeBitSet.h
+++ b/Engine/Source/Runtime/CoreUObject/Public/StructUtils/StructTypeBitSet.h
@@ -1197,15 +1197,6 @@ public:
 		return StructTypesBitArray.Contains(StructTypeIndex);
 	}
 
-	/**
-	 * Checks whether the bitset has a bit set at Index.
-	 * Indices past the end of the bitset are handled seamlessly (returning false).
-	 */
-	bool GetAtIndex(const int32 Index) const
-	{
-		return (Index >=0) && StructTypesBitArray.Contains(Index);
-	}
-
 	/**
 	 * Performs a union operation with another bitset.
 	 * @param Other - The other bitset to combine with.
diff --git a/Engine/Source/Runtime/MassEntity/Private/MassEntityBuilder.cpp b/Engine/Source/Runtime/MassEntity/Private/MassEntityBuilder.cpp
--- a/Engine/Source/Runtime/MassEntity/Private/MassEntityBuilder.cpp
+++ b/Engine/Source/Runtime/MassEntity/Private/MassEntityBuilder.cpp
@@ -430,7 +430,7 @@ bool FEntityBuilder::AppendDataFromEntity(const FMassEntityHandle SourceEntityHa
 			CA_ASSUME(SparseElementType);
 			const FConstStructView FragmentInstance = EntityManager->GetSparseElementDataForEntity(SparseElementType, SourceEntityHandle);
 
-			if (Composition.GetElementsBitSet().GetAtIndex(*It))
+			if (Composition.GetElementsBitSet().IsBitSet(*It))
 			{
 				// it's possible there's an instance in Fragments already
 				FInstancedStruct* ExistingStruct = Fragments.FindByPredicate([&SparseElementType](const FInstancedStruct& Element)
diff --git a/Engine/Source/Runtime/MassEntity/Public/MassArchetypeTypes.h b/Engine/Source/Runtime/MassEntity/Public/MassArchetypeTypes.h
--- a/Engine/Source/Runtime/MassEntity/Public/MassArchetypeTypes.h
+++ b/Engine/Source/Runtime/MassEntity/Public/MassArchetypeTypes.h
@@ -295,7 +295,7 @@ struct FMassRawEntityInChunkData
 	bool operator==(const FMassRawEntityInChunkData & Other) const;
 
 	uint8* const ChunkRawMemory = nullptr;
-	const int32 IndexWithinChunk = INDEX_NONE;;
+	const int32 IndexWithinChunk = INDEX_NONE;
 };
 
 //-----------------------------------------------------------------------------
diff --git a/Engine/Source/Runtime/MassEntity/Public/MassRequirements.h b/Engine/Source/Runtime/MassEntity/Public/MassRequirements.h
--- a/Engine/Source/Runtime/MassEntity/Public/MassRequirements.h
+++ b/Engine/Source/Runtime/MassEntity/Public/MassRequirements.h
@@ -726,7 +726,7 @@ struct FMassFragmentRequirements
 	bool HasLinkedEntityRequirements() const;
 
 	/**
-	 * Note that DoesArchetypeMatchRequirements functions to not check sparse elements. Sparse elements are not
+	 * Note that DoesArchetypeMatchRequirements functions do not check sparse elements. Sparse elements are not
 	 * a part of immutable composition of a given archetype and can change at runtime, and as such cannot be used
 	 * for filtering archetypes when the results of that filtering get cached (like in FMassEntityQuery)
 	 */
