[Mass] Fixing Archetyp/chunk level cached sparse element bitsets not getting updated when sparse elements get removed.

Introduced a dedicated type-instance counter to FChunkSparseElements.

#Synced-CL 50439895

diff --git a/Engine/Source/Developer/MassEntityTestSuite/Private/MassSparseElementsTest.cpp b/Engine/Source/Developer/MassEntityTestSuite/Private/MassSparseElementsTest.cpp
--- a/Engine/Source/Developer/MassEntityTestSuite/Private/MassSparseElementsTest.cpp
+++ b/Engine/Source/Developer/MassEntityTestSuite/Private/MassSparseElementsTest.cpp
@@ -15,9 +15,6 @@
 
 UE_DISABLE_OPTIMIZATION_SHIP
 
-// to be enabled once this functionality gets implemented
-#define WITH_ARCHETYPE_UPDATING_BITSETS_ON_REMOVE 0
-
 namespace UE::Mass::Test::SparseElements
 {
 #if WITH_MASSENTITY_DEBUG
@@ -47,12 +44,12 @@ namespace UE::Mass::Test::SparseElements
 
 	struct FTagTestOperations : FAutoCleanupOperations
 	{
-		static bool HasElement(const FMassArchetypeHandle& ArchetypeHandle, TNotNull<UScriptStruct*> SparseElementType, bool bIncludeTransient)
+		static bool HasSparseElement(const FMassArchetypeHandle& ArchetypeHandle, TNotNull<UScriptStruct*> SparseElementType)
 		{
-			return FMassArchetypeHelper::ArchetypeDataFromHandleChecked(ArchetypeHandle).HasElement(SparseElementType, bIncludeTransient ? EIncludeSparseElements::Yes : EIncludeSparseElements::No);
+			return FMassArchetypeHelper::ArchetypeDataFromHandleChecked(ArchetypeHandle).DoesContainEntitiesWithSparseElement(SparseElementType);
 		}
 
-		static bool HasElement(FMassEntityView EntityView, TNotNull<UScriptStruct*> SparseElementType)
+		static bool HasSparseElement(FMassEntityView EntityView, TNotNull<UScriptStruct*> SparseElementType)
 		{
 			return EntityView.HasElement(SparseElementType);
 		}
@@ -72,12 +69,12 @@ namespace UE::Mass::Test::SparseElements
 
 	struct FFragmentTestOperations : FAutoCleanupOperations
 	{
-		static bool HasElement(const FMassArchetypeHandle& ArchetypeHandle, TNotNull<UScriptStruct*> SparseElementType, bool bIncludeTransient)
+		static bool HasSparseElement(const FMassArchetypeHandle& ArchetypeHandle, TNotNull<UScriptStruct*> SparseElementType)
 		{
-			return FMassArchetypeHelper::ArchetypeDataFromHandleChecked(ArchetypeHandle).HasElement(SparseElementType, bIncludeTransient ? EIncludeSparseElements::Yes : EIncludeSparseElements::No);
+			return FMassArchetypeHelper::ArchetypeDataFromHandleChecked(ArchetypeHandle).DoesContainEntitiesWithSparseElement(SparseElementType);
 		}
 
-		static bool HasElement(FMassEntityView EntityView, TNotNull<UScriptStruct*> SparseElementType)
+		static bool HasSparseElement(FMassEntityView EntityView, TNotNull<UScriptStruct*> SparseElementType)
 		{
 			return EntityView.HasElement(SparseElementType);
 		}
@@ -141,13 +138,13 @@ namespace UE::Mass::Test::SparseElements
 			const FMassArchetypeHandle NewArchetypeHandle = EntityManager->GetArchetypeForEntity(ModifiedEntityHandle);
 			AITEST_TRUE(TEXT("The affected entity is still in the original archetype"), NewArchetypeHandle == OriginalArchetypeHandle);
 			AITEST_TRUE(TEXT("The archetype has the element")
-				, TTestOperations::HasElement(NewArchetypeHandle, SparseElementType, /*bIncludeTransient=*/true));
+				, TTestOperations::HasSparseElement(NewArchetypeHandle, SparseElementType));
 
 			FMassEntityView ModifiedEntityView(*EntityManager, ModifiedEntityHandle);
-			AITEST_TRUE(TEXT("The affected entity has the sparse element"), TTestOperations::HasElement(ModifiedEntityView, SparseElementType));
+			AITEST_TRUE(TEXT("The affected entity has the sparse element"), TTestOperations::HasSparseElement(ModifiedEntityView, SparseElementType));
 
 			FMassEntityView OriginalEntityView(*EntityManager, EntitiesCreated[(ModifiedEntityHandleIndex + 1) % NumEntities]);
-			AITEST_FALSE(TEXT("(NOT) The original entity has the sparse element"), TTestOperations::HasElement(OriginalEntityView, SparseElementType));
+			AITEST_FALSE(TEXT("(NOT) The original entity has the sparse element"), TTestOperations::HasSparseElement(OriginalEntityView, SparseElementType));
 
 			return true;
 		}
@@ -173,27 +170,24 @@ namespace UE::Mass::Test::SparseElements
 				FStructView _ = EntityManager->AddSparseElementToEntity(EntityHandle, SparseElementType);;
 			}
 
-			AITEST_TRUE(TEXT("First removal is expected to succeed"), EntityManager->RemoveSparseElementFromEntity(EntitiesCreated[0], SparseElementType));
+			EntityManager->RemoveSparseElementFromEntity(EntitiesCreated[0], SparseElementType);
 
 			const FMassArchetypeHandle NewArchetypeHandle = EntityManager->GetArchetypeForEntity(EntitiesCreated[0]);
 			AITEST_TRUE(TEXT("The affected entity is still in the original archetype"), NewArchetypeHandle == OriginalArchetypeHandle);
 			AITEST_TRUE(TEXT("The archetype knows some of its entities have the sparse element"), FMassArchetypeHelper::ArchetypeDataFromHandleChecked(NewArchetypeHandle).ContainsAnySparseData());
 
 			FMassEntityView ModifiedEntityView(*EntityManager, EntitiesCreated[0]);
-			AITEST_FALSE(TEXT("(NOT) The affected entity has the sparse element"), TTestOperations::HasElement(ModifiedEntityView, SparseElementType));
+			AITEST_FALSE(TEXT("(NOT) The affected entity has the sparse element"), TTestOperations::HasSparseElement(ModifiedEntityView, SparseElementType));
 
 			FMassEntityView OtherEntityView(*EntityManager, EntitiesCreated[1]);
-			AITEST_TRUE(TEXT("The other entity has the sparse element"), TTestOperations::HasElement(OtherEntityView, SparseElementType));
-
-			AITEST_TRUE(TEXT("Second removal is expected to succeed"),EntityManager->RemoveSparseElementFromEntity(EntitiesCreated[1], SparseElementType));
-			AITEST_FALSE(TEXT("(NOT) The other entity has the sparse element"), TTestOperations::HasElement(OtherEntityView, SparseElementType));
+			AITEST_TRUE(TEXT("The other entity has the sparse element"), TTestOperations::HasSparseElement(OtherEntityView, SparseElementType));
 
-			AITEST_TRUE(TEXT("Removing missing elements fails"),EntityManager->RemoveSparseElementFromEntity(EntitiesCreated[1], SparseElementType) == false);
+			EntityManager->RemoveSparseElementFromEntity(EntitiesCreated[1], SparseElementType);
+			AITEST_FALSE(TEXT("(NOT) The other entity has the sparse element"), TTestOperations::HasSparseElement(OtherEntityView, SparseElementType));
 
-#if WITH_ARCHETYPE_UPDATING_BITSETS_ON_REMOVE
 			AITEST_TRUE(TEXT("The original archetype no longer has the element")
-				, TTestOperations::HasElement(OriginalArchetypeHandle, SparseElementType, /*bIncludeTransient=*/true) == false);
-#endif
+				, TTestOperations::HasSparseElement(OriginalArchetypeHandle, SparseElementType) == false);
+
 			return true;
 		}
 	};
@@ -214,15 +208,15 @@ namespace UE::Mass::Test::SparseElements
 			FStructView _ = EntityManager->AddSparseElementToEntity(ModifiedEntityHandle, SparseElementType);;
 
 			EntityManager->AddFragmentToEntity(ModifiedEntityHandle, FTestFragment_Float::StaticStruct());
-#if WITH_ARCHETYPE_UPDATING_BITSETS_ON_REMOVE
+
 			AITEST_TRUE(TEXT("The original archetype no longer has the element")
-				, TTestOperations::HasElement(OriginalArchetypeHandle, SparseElementType, /*bIncludeTransient=*/true) == false);
-#endif
+				, TTestOperations::HasSparseElement(OriginalArchetypeHandle, SparseElementType) == false);
+
 			const FMassArchetypeHandle NewArchetypeHandle = EntityManager->GetArchetypeForEntity(ModifiedEntityHandle);
 			AITEST_TRUE(TEXT("The entity is expected to end up in the expected archetype"), NewArchetypeHandle == FEntityTestBase::FloatsIntsArchetype);
 
 			FMassEntityView ModifiedEntityView(*EntityManager, ModifiedEntityHandle);
-			AITEST_TRUE(TEXT("The affected entity has the sparse element"), TTestOperations::HasElement(ModifiedEntityView, SparseElementType));
+			AITEST_TRUE(TEXT("The affected entity has the sparse element"), TTestOperations::HasSparseElement(ModifiedEntityView, SparseElementType));
 			
 			return true;
 		}
@@ -252,16 +246,15 @@ namespace UE::Mass::Test::SparseElements
 			InitialCollection = FMassArchetypeEntityCollection(FEntityTestBase::IntsArchetype, EntitiesCreated, FMassArchetypeEntityCollection::NoDuplicates);
 			EntityManager->BatchChangeFragmentCompositionForEntities({InitialCollection}, FMassFragmentBitSet(*FTestFragment_Float::StaticStruct()), {});
 
-#if WITH_ARCHETYPE_UPDATING_BITSETS_ON_REMOVE
 			AITEST_TRUE(TEXT("The original archetype no longer has the element")
-				, TTestOperations::HasElement(OriginalArchetypeHandle, SparseElementType, /*bIncludeTransient=*/true) == false);
-#endif
+				, TTestOperations::HasSparseElement(OriginalArchetypeHandle, SparseElementType) == false);
+
 			FMassEntityHandle ModifiedEntityHandle = EntitiesCreated[0];
 			const FMassArchetypeHandle NewArchetypeHandle = EntityManager->GetArchetypeForEntity(ModifiedEntityHandle);
 			AITEST_TRUE(TEXT("The entity is expected to end up in the expected archetype"), NewArchetypeHandle == FEntityTestBase::FloatsIntsArchetype);
 
 			FMassEntityView ModifiedEntityView(*EntityManager, ModifiedEntityHandle);
-			AITEST_TRUE(TEXT("The affected entity has the sparse element"), TTestOperations::HasElement(ModifiedEntityView, SparseElementType));
+			AITEST_TRUE(TEXT("The affected entity has the sparse element"), TTestOperations::HasSparseElement(ModifiedEntityView, SparseElementType));
 			
 			return true;
 		}
@@ -285,10 +278,8 @@ namespace UE::Mass::Test::SparseElements
 			InitialCollection = FMassArchetypeEntityCollection(FEntityTestBase::IntsArchetype, EntitiesCreated, FMassArchetypeEntityCollection::NoDuplicates);
 			EntityManager->BatchAddSparseElementToEntities({InitialCollection}, SparseElementType);
 
-#if WITH_ARCHETYPE_UPDATING_BITSETS_ON_REMOVE
-			AITEST_TRUE(TEXT("The original archetype no longer has the element")
-				, TTestOperations::HasElement(OriginalArchetypeHandle, SparseElementType, /*bIncludeTransient=*/true) == false);
-#endif
+			AITEST_TRUE(TEXT("The original archetype has the sparse element")
+				, TTestOperations::HasSparseElement(OriginalArchetypeHandle, SparseElementType));
 
 			for (FMassEntityHandle ModifiedEntityHandle : EntitiesCreated)
 			{
@@ -296,7 +287,7 @@ namespace UE::Mass::Test::SparseElements
 				AITEST_TRUE(TEXT("The entity is expected to end up in the expected archetype"), NewArchetypeHandle == OriginalArchetypeHandle);
 
 				FMassEntityView ModifiedEntityView(*EntityManager, ModifiedEntityHandle);
-				AITEST_TRUE(TEXT("The affected entities have the sparse element"), TTestOperations::HasElement(ModifiedEntityView, SparseElementType));
+				AITEST_TRUE(TEXT("The affected entities have the sparse element"), TTestOperations::HasSparseElement(ModifiedEntityView, SparseElementType));
 			}
 			
 			return true;
@@ -321,10 +312,8 @@ namespace UE::Mass::Test::SparseElements
 
 			EntityManager->BatchRemoveSparseElementFromEntities({InitialCollection}, SparseElementType);
 
-#if WITH_ARCHETYPE_UPDATING_BITSETS_ON_REMOVE
 			AITEST_TRUE(TEXT("The original archetype no longer has the element")
-				, TTestOperations::HasElement(OriginalArchetypeHandle, SparseElementType, /*bIncludeTransient=*/true) == false);
-#endif
+				, TTestOperations::HasSparseElement(OriginalArchetypeHandle, SparseElementType) == false);
 
 			for (FMassEntityHandle ModifiedEntityHandle : EntitiesCreated)
 			{
@@ -332,7 +321,7 @@ namespace UE::Mass::Test::SparseElements
 				AITEST_TRUE(TEXT("The entity is expected to end up in the expected archetype"), NewArchetypeHandle == OriginalArchetypeHandle);
 
 				FMassEntityView ModifiedEntityView(*EntityManager, ModifiedEntityHandle);
-				AITEST_FALSE(TEXT("(NOT) The affected entities have the sparse element"), TTestOperations::HasElement(ModifiedEntityView, SparseElementType));
+				AITEST_FALSE(TEXT("(NOT) The affected entities have the sparse element"), TTestOperations::HasSparseElement(ModifiedEntityView, SparseElementType));
 			}
 			
 			return true;
@@ -649,8 +638,6 @@ namespace UE::Mass::Test::SparseElements
 #endif // WITH_MASSENTITY_DEBUG
 } // namespace
 
-#undef WITH_ARCHETYPE_UPDATING_BITSETS_ON_REMOVE
-
 UE_ENABLE_OPTIMIZATION_SHIP
 
 #undef LOCTEXT_NAMESPACE
diff --git a/Engine/Source/Runtime/MassEntity/Internal/MassArchetypeData.h b/Engine/Source/Runtime/MassEntity/Internal/MassArchetypeData.h
--- a/Engine/Source/Runtime/MassEntity/Internal/MassArchetypeData.h
+++ b/Engine/Source/Runtime/MassEntity/Internal/MassArchetypeData.h
@@ -39,11 +39,17 @@ namespace UE::Mass
 			PerEntityElements.AddDefaulted(NumEntitiesPerChunk);
 		}
 
-		void Add(const FMassRawEntityInChunkData& EntityInChunkHandle, TNotNull<const UScriptStruct*> ElementType);
+		/** @return whether it is the first instance of ElementType added to this chunk, i.e. SparseElementsPresent has been modified */
+		bool Add(const FMassRawEntityInChunkData& EntityInChunkHandle, TNotNull<const UScriptStruct*> ElementType);
+		/** @return whether it is the last instance of ElementType removes from this chunk, i.e. SparseElementsPresent has been modified */
 		bool Remove(const FMassRawEntityInChunkData& EntityInChunkHandle, TNotNull<const UScriptStruct*> ElementType);
+		/** @return whether SparseElementsPresent has been modified */
+		bool UpdateElementsPresenceOnEntityRemoval(const int32 IndexWithinChunk);
 
-		void BatchAdd(TNotNull<const UScriptStruct*> ElementType, int32 SubchunkStart, int32 RangeLength);
-		void BatchRemove(TNotNull<const UScriptStruct*> ElementType, int32 SubchunkStart, int32 RangeLength);
+		/** @return whether these are the first instances of ElementType added to this chunk, i.e. SparseElementsPresent has been modified */
+		bool BatchAdd(TNotNull<const UScriptStruct*> ElementType, int32 SubchunkStart, int32 RangeLength);
+		/** @return whether these are the last instances of ElementType removed this chunk, i.e. SparseElementsPresent has been modified */
+		bool BatchRemove(TNotNull<const UScriptStruct*> ElementType, int32 SubchunkStart, int32 RangeLength);
 
 		bool IsEmpty() const
 		{
@@ -70,7 +76,12 @@ namespace UE::Mass
 			return false;
 		}
 
-		void MoveElementsFrom(FChunkSparseElements& OriginalSparseElements, int32 OriginalIndexWithinChunk, int32 NewIndexWithinChunk, int32 Num);
+		/**
+		 * Moves Num sparse elements bitsets from OriginalSparseElements, starting at index OriginalIndexWithinChunk
+		 * to `this` FChunkSparseElements instances, starting at index NewIndexWithinChunk
+		 * @return true if any bits have been added to or remove from SparseElementsPresent, both from `this` and OriginalSparseElements (if it's not the same).
+		 */
+		bool MoveElementsFrom(FChunkSparseElements& OriginalSparseElements, int32 OriginalIndexWithinChunk, int32 NewIndexWithinChunk, int32 Num);
 
 		const FMassElementBitSet& GetSparseElementsPresent() const
 		{
@@ -99,6 +110,42 @@ namespace UE::Mass
 	private:
 		FMassElementBitSet SparseElementsPresent;
 		TArray<FMassElementBitSet> PerEntityElements;
+
+		struct FElementCounter
+		{
+			uint16& Get(const int32 TypeIndex)
+			{
+				if (LIKELY(DenseIndexToElementIndex.IsValidIndex(LastAccessedDenseIndex)
+					&& DenseIndexToElementIndex[LastAccessedDenseIndex] == TypeIndex))
+				{
+					return DenseCounters[LastAccessedDenseIndex];
+				}
+				if (UNLIKELY(DenseIndexToElementIndex.IsEmpty()))
+				{
+					// add a new one at the end
+					LastAccessedDenseIndex = DenseIndexToElementIndex.Add(static_cast<uint16>(TypeIndex));
+					return DenseCounters.Add_GetRef(0);
+				}
+
+				const int32 DenseIndex = Algo::LowerBound(DenseIndexToElementIndex, TypeIndex);
+				LastAccessedDenseIndex = DenseIndex;
+
+				if (DenseIndex == DenseIndexToElementIndex.Num()
+					|| DenseIndexToElementIndex[DenseIndex] != TypeIndex)
+				{
+					// in this case DenseIndex indicates the new element insertion index
+					DenseIndexToElementIndex.Insert(static_cast<uint16>(TypeIndex), DenseIndex);
+					return DenseCounters.Insert_GetRef(0, DenseIndex);
+				}
+
+				return DenseCounters[DenseIndex];
+			}
+		private:
+			int32 LastAccessedDenseIndex = INDEX_NONE;
+			TArray<uint16, TInlineAllocator<16>> DenseCounters;
+			TArray<uint16, TInlineAllocator<16>> DenseIndexToElementIndex;
+		};
+		FElementCounter SparseElementTypesCount;
 	};
 }
 
@@ -112,6 +159,7 @@ private:
 	int32 SerialModificationNumber = 0;
 	TArray<FInstancedStruct> ChunkFragmentData;
 	FMassArchetypeSharedFragmentValues SharedFragmentValues;
+	TArray<uint16, TInlineAllocator<256>> ElementTypeCounts;
 	UE::Mass::FChunkSparseElements SparseElements;
 	
 public:
@@ -176,15 +224,19 @@ public:
 		return SparseElements;
 	}
 
+	/**
+	 * Note that the returned FChunkSparseElements instance might not be initialized.
+	 * It's caller's responsibility to call FChunkSparseElements.IsEmpty to verify the instance's state.
+	 * Call GetSparseElements(uint32) if you need a guaranteed initialized instance.
+	 */
 	UE::Mass::FChunkSparseElements& GetSparseElementsUnsafe()
 	{
-		check(SparseElements.IsEmpty() == false);
 		return SparseElements;
 	}
 
+	/** The caller is responsible for checking whether SparseElements.IsEmpty */
 	const UE::Mass::FChunkSparseElements& GetSparseElementsUnsafe() const
 	{
-		check(SparseElements.IsEmpty() == false);
 		return SparseElements;
 	}
 
@@ -290,8 +342,8 @@ struct FMassArchetypeData
 private:
 	// One-stop-shop variable describing composition of entities hosted by the archetype
 	FMassElementBitSet CompositionBitSet;
-	FMassElementBitSet SparseElementsBitSet;
-
+	mutable FMassElementBitSet CachedSparseElementsBitSet;
+	
 	// Pre-created default chunk fragment templates
 	TArray<FInstancedStruct> ChunkFragmentsTemplate;
 
@@ -333,6 +385,10 @@ private:
 	 */
 	UE::Mass::FSparseElementsStorage* SparseElementsStorage = nullptr;
 
+	UE_MT_DECLARE_RW_ACCESS_DETECTOR(ChunkIterationDetector);
+
+	mutable bool bSparseElementsBitSetDirty = false;
+
 #if WITH_MASSENTITY_DEBUG
 	/** Arrays of names the archetype is referred as. */
 	TArray<FName> DebugNames;
@@ -398,11 +454,14 @@ public:
 	 * @param IncludeSparseElements whether the function should check if the archetype hosts any entities
 	 *	owning such a fragment, provided FragmentType represents a sparse element type
 	 */
-	bool HasElement(const UScriptStruct* ElementType, const UE::Mass::EIncludeSparseElements IncludeSparseElements/* = UE::Mass::EIncludeSparseElements::No*/) const
+	bool HasElement(TNotNull<const UScriptStruct*> ElementType) const
+	{
+		return CompositionBitSet.Contains(*ElementType);
+	}
+
+	bool DoesContainEntitiesWithSparseElement(const UScriptStruct* ElementType) const
 	{
-		return (CompositionBitSet.Contains(*ElementType)
-				|| (IncludeSparseElements == UE::Mass::EIncludeSparseElements::Yes
-					&& UE::Mass::IsSparse(ElementType) && SparseElementsBitSet.Contains(*ElementType)));
+		return UE::Mass::IsSparse(ElementType) && GetSparseElementsBitSet().Contains(*ElementType);
 	}
 
 	bool IsEquivalent(const FMassElementBitSet& OtherCompositionBitSet, const UE::Mass::FArchetypeGroups& OtherGroups) const;
@@ -427,7 +486,7 @@ public:
 	 */
 	FStructView AddSparseElementToEntity(FMassEntityHandle EntityHandle, TNotNull<const UScriptStruct*> ElementType);
 	/** Works with all sparse element types (fragments and tags). */
-	bool RemoveSparseElementFromEntity(FMassEntityHandle EntityHandle, TNotNull<const UScriptStruct*> ElementType);
+	void RemoveSparseElementFromEntity(FMassEntityHandle EntityHandle, TNotNull<const UScriptStruct*> ElementType);
 	/** Expects non-tag sparse element type. Will assert if the assumption is broken. */
 	FStructView GetMutableSparseElementDataForEntity(TNotNull<const UScriptStruct*> ElementType, const FMassEntityHandle EntityHandle) const;
 	/** Expects non-tag sparse element type. Will assert if the assumption is broken. */
@@ -446,6 +505,7 @@ public:
 		, const FMassElementBitSet& RequiredNoneSparseElements) const;
 	bool ContainsAnySparseData() const;
 	bool ContainsSparseElement(TNotNull<const UScriptStruct*> ElementType) const;
+	const FMassElementBitSet& GetSparseElementsBitSet() const;
 
 	FORCEINLINE const int32* GetInternalIndexForEntity(const int32 EntityIndex) const { return EntityMap.Find(EntityIndex); }
 	FORCEINLINE int32 GetInternalIndexForEntityChecked(const int32 EntityIndex) const { return EntityMap.FindChecked(EntityIndex); }
@@ -567,6 +627,8 @@ public:
 	 * @param InPrefix Optional prefix to remove from fragment names
 	 */
 	void DebugPrintEntity(FMassEntityHandle Entity, FOutputDevice& Ar, const TCHAR* InPrefix = TEXT("")) const;
+
+	MASSENTITY_API const FMassArchetypeChunk& DebugGetChunk(const int32 Index) const;
 #endif // WITH_MASSENTITY_DEBUG
 
 	void SetDebugColor(const FColor InDebugColor);
@@ -777,7 +839,7 @@ void FMassArchetypeData::SetFragmentsData(const FMassEntityHandle Entity, TConst
 		}
 		else
 		{
-			SparseElementsBitSet.Add(*FragmentType);
+			bSparseElementsBitSetDirty = true;
 			UE::Mass::FChunkSparseElements& SparseElements = Chunks[EntityInChunkHandle.ChunkIndex].GetSparseElements(NumEntitiesPerChunk);
 			SparseElements.Add(EntityInChunkHandle, FragmentType);
 			FStructView _ = SparseElementsStorage->AddElementInstanceToEntity(Entity, Instance);
@@ -821,7 +883,7 @@ void FMassArchetypeData::SetFragmentData(FMassArchetypeEntityCollection::FConstE
 	}
 	else
 	{
-		SparseElementsBitSet.Add(*FragmentType);
+		bSparseElementsBitSetDirty = true;
 		for (FMassArchetypeChunkIterator ChunkIterator(EntityRangeContainer); ChunkIterator; ++ChunkIterator)
 		{
 			FMassArchetypeChunk& Chunk = Chunks[ChunkIterator->ChunkIndex];
@@ -836,12 +898,32 @@ void FMassArchetypeData::SetFragmentData(FMassArchetypeEntityCollection::FConstE
 		}
 	}
 }
+
 inline bool FMassArchetypeData::ContainsAnySparseData() const
 {
-	return !SparseElementsBitSet.IsEmpty();
+	return !GetSparseElementsBitSet().IsEmpty();
 }
 
 inline bool FMassArchetypeData::ContainsSparseElement(TNotNull<const UScriptStruct*> ElementType) const
 {
-	return SparseElementsBitSet.Contains(*ElementType);
+	return GetSparseElementsBitSet().Contains(*ElementType);
+}
+
+inline const FMassElementBitSet& FMassArchetypeData::GetSparseElementsBitSet() const
+{
+	UE_MT_SCOPED_READ_ACCESS(ChunkIterationDetector);
+
+	if (bSparseElementsBitSetDirty)
+	{
+		CachedSparseElementsBitSet.Reset();
+		for (const FMassArchetypeChunk& Chunk : Chunks)
+		{
+			if (Chunk.HasSparseElements())
+			{
+				CachedSparseElementsBitSet += Chunk.GetSparseElementsUnsafe().GetSparseElementsPresent();
+			}
+		}
+		bSparseElementsBitSetDirty = false;
+	}
+	return CachedSparseElementsBitSet;
 }
diff --git a/Engine/Source/Runtime/MassEntity/Private/MassArchetypeData.cpp b/Engine/Source/Runtime/MassEntity/Private/MassArchetypeData.cpp
--- a/Engine/Source/Runtime/MassEntity/Private/MassArchetypeData.cpp
+++ b/Engine/Source/Runtime/MassEntity/Private/MassArchetypeData.cpp
@@ -32,15 +32,23 @@ namespace UE::Mass
 	//-----------------------------------------------------------------------------
 	// FChunkSparseElements
 	//-----------------------------------------------------------------------------
-	void FChunkSparseElements::Add(const FMassRawEntityInChunkData& EntityInChunkHandle, TNotNull<const UScriptStruct*> ElementType)
+	bool FChunkSparseElements::Add(const FMassRawEntityInChunkData& EntityInChunkHandle, TNotNull<const UScriptStruct*> ElementType)
 	{
 		check(IsEmpty() == false);
 
 		const int32 TypeIndex = FMassElementBitSet::GetTypeIndex(*ElementType);
-		SparseElementsPresent.AddAtIndex(TypeIndex);
 
 		FMassElementBitSet& EntitySparseElements = PerEntityElements[EntityInChunkHandle.IndexWithinChunk];
-		EntitySparseElements.AddAtIndex(TypeIndex);
+		if (EntitySparseElements.IsBitSet(TypeIndex) == false)
+		{
+			EntitySparseElements.AddAtIndex(TypeIndex);
+			if (SparseElementTypesCount.Get(TypeIndex)++ == 0)
+			{
+				SparseElementsPresent.AddAtIndex(TypeIndex);
+				return true;
+			}
+		}
+		return false;
 	}
 
 	bool FChunkSparseElements::Remove(const FMassRawEntityInChunkData& EntityInChunkHandle, TNotNull<const UScriptStruct*> ElementType)
@@ -50,34 +58,87 @@ namespace UE::Mass
 		FMassElementBitSet& EntitySparseElements = PerEntityElements[EntityInChunkHandle.IndexWithinChunk];
 
 		const int32 TypeIndex = FMassElementBitSet::GetTypeIndex(*ElementType);
-		const bool bHasElement = EntitySparseElements.IsBitSet(TypeIndex);
-		EntitySparseElements.RemoveAtIndex(TypeIndex);
-		return bHasElement;
+		if (EntitySparseElements.IsBitSet(TypeIndex))
+		{
+			EntitySparseElements.RemoveAtIndex(TypeIndex);
+			if (--SparseElementTypesCount.Get(TypeIndex) == 0)
+			{
+				SparseElementsPresent.RemoveAtIndex(TypeIndex);
+				return true;
+			}
+		}
+
+		return false;
+	}
+
+	bool FChunkSparseElements::UpdateElementsPresenceOnEntityRemoval(const int32 IndexWithinChunk)
+	{
+		bool bSparseElementsPresentDirty = false;
+		for (FMassElementBitSet::FIndexIterator It = PerEntityElements[IndexWithinChunk].GetIndexIterator(); It; ++It)
+		{
+			if (--SparseElementTypesCount.Get(*It) == 0)
+			{
+				SparseElementsPresent.RemoveAtIndex(*It);
+				bSparseElementsPresentDirty = true;
+			}
+		}
+		PerEntityElements[IndexWithinChunk].Reset();
+		return bSparseElementsPresentDirty;
 	}
 
-	void FChunkSparseElements::BatchAdd(TNotNull<const UScriptStruct*> ElementType, const int32 SubchunkStart, const int32 RangeLength)
+	bool FChunkSparseElements::BatchAdd(TNotNull<const UScriptStruct*> ElementType, const int32 SubchunkStart, const int32 RangeLength)
 	{
 		check(IsEmpty() == false);
 		check(SubchunkStart >= 0 && SubchunkStart + RangeLength <= PerEntityElements.Num());
 
+		uint16 NumAdded = 0;
 		const int32 TypeIndex = FMassElementBitSet::GetTypeIndex(*ElementType);
 		for (int32 EntityIndex = SubchunkStart; EntityIndex < SubchunkStart + RangeLength; ++EntityIndex)
 		{
+			NumAdded += static_cast<int>(PerEntityElements[EntityIndex].IsBitSet(TypeIndex) == false);
 			PerEntityElements[EntityIndex].AddAtIndex(TypeIndex);
 		}
-		SparseElementsPresent.AddAtIndex(TypeIndex);
+
+		bool bSparseElementsPresentDirty = false;
+		if (NumAdded)
+		{
+			uint16& TotalElementsCount = SparseElementTypesCount.Get(TypeIndex);
+			if (TotalElementsCount == 0)
+			{
+				SparseElementsPresent.AddAtIndex(TypeIndex);
+				bSparseElementsPresentDirty = true;
+			}
+			TotalElementsCount += NumAdded;
+		}
+		return bSparseElementsPresentDirty;
 	}
 
-	void FChunkSparseElements::BatchRemove(TNotNull<const UScriptStruct*> ElementType, const int32 SubchunkStart, const int32 RangeLength)
+	bool FChunkSparseElements::BatchRemove(TNotNull<const UScriptStruct*> ElementType, const int32 SubchunkStart, const int32 RangeLength)
 	{
 		check(IsEmpty() == false);
 		check(SubchunkStart >= 0 && SubchunkStart + RangeLength <= PerEntityElements.Num());
 
+		uint16 NumRemoved = 0;
 		const int32 TypeIndex = FMassElementBitSet::GetTypeIndex(*ElementType);
 		for (int32 EntityIndex = SubchunkStart; EntityIndex < SubchunkStart + RangeLength; ++EntityIndex)
 		{
+			NumRemoved += static_cast<int>(PerEntityElements[EntityIndex].IsBitSet(TypeIndex));
 			PerEntityElements[EntityIndex].RemoveAtIndex(TypeIndex);
 		}
+
+		bool bSparseElementsPresentDirty = false;
+		if (NumRemoved)
+		{
+			uint16& TotalElementsCount = SparseElementTypesCount.Get(TypeIndex);
+			check(TotalElementsCount >= NumRemoved);
+			TotalElementsCount -= NumRemoved;
+			if (TotalElementsCount == 0)
+			{
+				SparseElementsPresent.RemoveAtIndex(TypeIndex);
+				bSparseElementsPresentDirty = true;
+			}
+		}
+		return bSparseElementsPresentDirty;
 	}
 
 	bool FChunkSparseElements::Contains(const FMassRawEntityInChunkData& EntityInChunkHandle, TNotNull<const UScriptStruct*> ElementType) const
@@ -86,22 +147,52 @@ namespace UE::Mass
 			&& PerEntityElements[EntityInChunkHandle.IndexWithinChunk].Contains(*ElementType);
 	}
 
-	void FChunkSparseElements::MoveElementsFrom(FChunkSparseElements& OriginalSparseElements, const int32 OriginalIndexWithinChunk, const int32 NewIndexWithinChunk, const int32 Num)
+	bool FChunkSparseElements::MoveElementsFrom(FChunkSparseElements& OriginalSparseElements, const int32 OriginalIndexWithinChunk, const int32 NewIndexWithinChunk, const int32 Num)
 	{
 		check(OriginalSparseElements.PerEntityElements.IsValidIndex(OriginalIndexWithinChunk)
 			&& OriginalIndexWithinChunk + Num <= OriginalSparseElements.PerEntityElements.Num()
 			&& PerEntityElements.IsValidIndex(NewIndexWithinChunk)
 			&& NewIndexWithinChunk + Num <= PerEntityElements.Num());
 
+		bool bSparseElementsPresentDirty = false;
+		if (&OriginalSparseElements == this)
+		{
+			// we're moving inside a single chunk, so we need to check the bitsets we're overriding, since they belong to entities
+			// that are being moved out of the chunk
+			// using NewIndexWithinChunk as the index reference because that's the data that will get overridden 
+			for (int32 Index = NewIndexWithinChunk; Index < NewIndexWithinChunk + Num; ++Index)
+			{
+				bSparseElementsPresentDirty |= UpdateElementsPresenceOnEntityRemoval(Index);
+			}
+		}
+
 		int32 OriginalIndex = OriginalIndexWithinChunk;
 		int32 NewIndex = NewIndexWithinChunk;
-		for (int32 Iterations = 0; Iterations < Num; ++Iterations, ++OriginalIndex, ++NewIndex)
+		for (int32 ItemCount = 0; ItemCount < Num; ++ItemCount, ++OriginalIndex, ++NewIndex)
 		{
-			FMassElementBitSet&& Elements = MoveTemp(OriginalSparseElements.PerEntityElements[OriginalIndex]);
+			PerEntityElements[NewIndex] = MoveTemp(OriginalSparseElements.PerEntityElements[OriginalIndex]);
+		}
 
-			SparseElementsPresent += Elements;
-			PerEntityElements[NewIndex] = MoveTemp(Elements);
+		if (&OriginalSparseElements != this)
+		{
+			for (int32 Index = NewIndexWithinChunk; Index < NewIndexWithinChunk + Num; ++Index)
+			{
+				for (FMassElementBitSet::FIndexIterator It = PerEntityElements[Index].GetIndexIterator(); It; ++It)
+				{
+					if (SparseElementTypesCount.Get(*It)++ == 0)
+					{
+						SparseElementsPresent.AddAtIndex(*It);
+						bSparseElementsPresentDirty = true;
+					}
+					if (--OriginalSparseElements.SparseElementTypesCount.Get(*It) == 0)
+					{
+						OriginalSparseElements.SparseElementsPresent.RemoveAtIndex(*It);
+						bSparseElementsPresentDirty = true;
+					}
+				}
+			}
 		}
+		return bSparseElementsPresentDirty;
 	}
 } // namespace UE::Mass
 
@@ -293,8 +384,7 @@ FStructView FMassArchetypeData::AddSparseElementToEntity(FMassEntityHandle Entit
 	FMassEntityInChunkDataHandle EntityInChunkHandle = MakeEntityHandle(EntityHandle.Index);
 	UE::Mass::FChunkSparseElements& SparseElements = Chunks[EntityInChunkHandle.ChunkIndex].GetSparseElements(NumEntitiesPerChunk);
 
-	SparseElementsBitSet.Add(*ElementType);
-	SparseElements.Add(EntityInChunkHandle, ElementType);
+	bSparseElementsBitSetDirty |= SparseElements.Add(EntityInChunkHandle, ElementType);
 
 	// Sparse elements storage stores only fragments
 	if (UE::Mass::IsA<FMassTag>(ElementType))
@@ -306,17 +396,20 @@ FStructView FMassArchetypeData::AddSparseElementToEntity(FMassEntityHandle Entit
 	return SparseElementsStorage->AddElementToEntity(EntityHandle, ElementType);
 }
 
-bool FMassArchetypeData::RemoveSparseElementFromEntity(FMassEntityHandle EntityHandle, TNotNull<const UScriptStruct*> ElementType)
+void FMassArchetypeData::RemoveSparseElementFromEntity(FMassEntityHandle EntityHandle, TNotNull<const UScriptStruct*> ElementType)
 {
 	FMassEntityInChunkDataHandle EntityInChunkHandle = MakeEntityHandle(EntityHandle.Index);
-	UE::Mass::FChunkSparseElements& SparseElements = Chunks[EntityInChunkHandle.ChunkIndex].GetSparseElements(NumEntitiesPerChunk);
+	UE::Mass::FChunkSparseElements& SparseElements = Chunks[EntityInChunkHandle.ChunkIndex].GetSparseElementsUnsafe();
 
-	if (UE::Mass::IsA<FMassTag>(ElementType) == false)
+	if (SparseElements.IsEmpty() == false)
 	{
-		CA_ASSUME(SparseElementsStorage);
-		SparseElementsStorage->RemoveElementFromEntity(EntityHandle, ElementType);
+		if (UE::Mass::IsA<FMassTag>(ElementType) == false)
+		{
+			CA_ASSUME(SparseElementsStorage);
+			SparseElementsStorage->RemoveElementFromEntity(EntityHandle, ElementType);
+		}
+		bSparseElementsBitSetDirty |= SparseElements.Remove(EntityInChunkHandle, ElementType);
 	}
-	return SparseElements.Remove(EntityInChunkHandle, ElementType);
 }
 
 FStructView FMassArchetypeData::GetMutableSparseElementDataForEntity(TNotNull<const UScriptStruct*> ElementType, const FMassEntityHandle EntityHandle) const
@@ -344,6 +437,8 @@ FMassElementBitSet FMassArchetypeData::GetSparseElementsBitSetForEntity(const FM
 
 FMassArchetypeChunk& FMassArchetypeData::GetOrAddChunk(const FMassArchetypeSharedFragmentValues& SharedFragmentValues, int32& OutAbsoluteIndex, int32& OutIndexWithinChunk)
 {
+	UE_MT_SCOPED_WRITE_ACCESS(ChunkIterationDetector);
+
 	OutAbsoluteIndex = 0;
 	OutIndexWithinChunk = 0;
 
@@ -414,8 +509,10 @@ void FMassArchetypeData::RemoveEntity(FMassEntityHandle Entity)
 
 	if (Chunk.HasSparseElements())
 	{
-		FMassElementBitSet& ElementsBitSet = Chunk.GetSparseElementsUnsafe().GetElementsForEntityUnsafe(IndexWithinChunk);
+		UE::Mass::FChunkSparseElements& SparseElements = Chunk.GetSparseElementsUnsafe();
+		const FMassElementBitSet& ElementsBitSet = SparseElements.GetElementsForEntityUnsafe(IndexWithinChunk);
 		SparseElementsStorage->RemoveEntity(Entity, ElementsBitSet);
+		bSparseElementsBitSetDirty |= SparseElements.UpdateElementsPresenceOnEntityRemoval(IndexWithinChunk);
 	}
 
 	RemoveEntityInternal(AbsoluteIndex);
@@ -491,8 +588,10 @@ void FMassArchetypeData::BatchDestroyEntityChunks(FMassArchetypeEntityCollection
 		{
 			for (int32 RelativeIndex = 0; RelativeIndex < RangeLength; ++RelativeIndex)
 			{
-				FMassElementBitSet& ElementsBitSet = Chunk.GetSparseElementsUnsafe().GetElementsForEntityUnsafe(EntityRange.SubchunkStart + RelativeIndex);
+				UE::Mass::FChunkSparseElements& SparseElements = Chunk.GetSparseElementsUnsafe();
+				FMassElementBitSet& ElementsBitSet = SparseElements.GetElementsForEntityUnsafe(EntityRange.SubchunkStart + RelativeIndex);
 				SparseElementsStorage->RemoveEntity(DyingEntityPtr[RelativeIndex], ElementsBitSet);
+				bSparseElementsBitSetDirty |= SparseElements.UpdateElementsPresenceOnEntityRemoval(EntityRange.SubchunkStart + RelativeIndex);
 				ElementsBitSet.Reset();
 			}
 		}
@@ -531,8 +630,7 @@ bool FMassArchetypeData::HasSparseElementForEntity(TNotNull<const UScriptStruct*
 	checkf(Chunks.IsValidIndex(EntityInChunkHandle.ChunkIndex), TEXT("%hs: Invalid chunk index (%d out of %d)")
 		, __FUNCTION__, EntityInChunkHandle.ChunkIndex, Chunks.Num());
 
-	return SparseElementsBitSet.Contains(*ElementType)
-		&& Chunks[EntityInChunkHandle.ChunkIndex].HasSparseElements()
+	return Chunks[EntityInChunkHandle.ChunkIndex].HasSparseElements()
 		&& Chunks[EntityInChunkHandle.ChunkIndex].GetSparseElementsUnsafe().Contains(EntityInChunkHandle, ElementType);
 }
 
@@ -610,7 +708,8 @@ void FMassArchetypeData::SetSharedFragmentsData(const FMassEntityHandle Entity,
 		{
 			UE::Mass::FChunkSparseElements& OldSparseElements = OldChunk.GetSparseElementsUnsafe();
 			UE::Mass::FChunkSparseElements& NewSparseElements = NewChunk.GetSparseElements(NumEntitiesPerChunk);
-			NewSparseElements.MoveElementsFrom(OldSparseElements, OldIndexWithinChunk, NewIndexWithinChunk, 1);
+			const bool bElementsPresentDirty = NewSparseElements.MoveElementsFrom(OldSparseElements, OldIndexWithinChunk, NewIndexWithinChunk, 1);
+			bSparseElementsBitSetDirty |= bElementsPresentDirty;
 		}
 
 		// Clean up the old chunk
@@ -836,7 +935,8 @@ int32 FMassArchetypeData::CompactEntities(const double TimeAllowed)
 			{
 				UE::Mass::FChunkSparseElements& OldSparseElements = ChunkToEmpty->GetSparseElementsUnsafe();
 				UE::Mass::FChunkSparseElements& NewSparseElements = ChunkToFill->GetSparseElements(NumEntitiesPerChunk);
-				NewSparseElements.MoveElementsFrom(OldSparseElements, FromIndex, ToIndex, NumberOfEntitiesToMove);
+				const bool bElementsPresentDirty = NewSparseElements.MoveElementsFrom(OldSparseElements, FromIndex, ToIndex, NumberOfEntitiesToMove);
+				bSparseElementsBitSetDirty |= bElementsPresentDirty;
 			}
 
 			FMassEntityHandle* FromEntity = &ChunkToEmpty->GetEntityArrayElementRef(EntityListOffsetWithinChunk, FromIndex);
@@ -1259,6 +1359,11 @@ void FMassArchetypeData::DebugPrintEntity(FMassEntityHandle Entity, FOutputDevic
 	}
 }
 
+const FMassArchetypeChunk& FMassArchetypeData::DebugGetChunk(const int32 Index) const
+{
+	static FMassArchetypeChunk DummyChunk(0, {}, {});
+	return Chunks.IsValidIndex(Index) ? Chunks[Index] : DummyChunk;
+}
 #endif // WITH_MASSENTITY_DEBUG
 
 void FMassArchetypeData::SetDebugColor(const FColor InDebugColor)
@@ -1529,7 +1634,8 @@ void FMassArchetypeData::BatchRemoveEntitiesInternal(const int32 ChunkIndex, con
 		if (Chunk.HasSparseElements())
 		{
 			UE::Mass::FChunkSparseElements& SparseElements = Chunk.GetSparseElementsUnsafe();
-			SparseElements.MoveElementsFrom(SparseElements, SwapStartIndex, StartIndexWithinChunk, NumberToMove);
+			const bool bElementsPresentDirty = SparseElements.MoveElementsFrom(SparseElements, SwapStartIndex, StartIndexWithinChunk, NumberToMove);
+			bSparseElementsBitSetDirty |= bElementsPresentDirty;
 		}
 
 		// Update the entity table and map
@@ -1599,9 +1705,9 @@ void FMassArchetypeData::MoveFragmentsToAnotherArchetypeInternal(FMassArchetypeD
 		UE::Mass::FChunkSparseElements& SourceSparseElements = SourceChunk.GetSparseElements(NumEntitiesPerChunk);
 		// calling GetSparseElements to make sure FChunkSparseElements gets created with the right size (FChunkSparseElements)
 		UE::Mass::FChunkSparseElements& TargetElements = TargetChunk.GetSparseElements(TargetArchetype.NumEntitiesPerChunk);
-		TargetElements.MoveElementsFrom(SourceSparseElements, OriginalIndexWithinChunk, NewIndexWithinChunk, ElementsNum);
-
-		TargetArchetype.SparseElementsBitSet += TargetElements.GetSparseElementsPresent();
+		const bool bElementsPresentDirty = TargetElements.MoveElementsFrom(SourceSparseElements, OriginalIndexWithinChunk, NewIndexWithinChunk, ElementsNum);
+		bSparseElementsBitSetDirty |= bElementsPresentDirty;
+		TargetArchetype.bSparseElementsBitSetDirty |= bElementsPresentDirty;
 	}
 }
 
@@ -1644,9 +1750,8 @@ void FMassArchetypeData::BatchSetFragmentValues(TConstArrayView<FMassArchetypeEn
 			}
 			else
 			{
-				SparseElementsBitSet.Add(*FragmentType);
 				UE::Mass::FChunkSparseElements& SparseElements = Chunk.GetSparseElements(NumEntitiesPerChunk);
-				SparseElements.BatchAdd(FragmentType, EntityRange.SubchunkStart, RangeLength);
+				bSparseElementsBitSetDirty |= SparseElements.BatchAdd(FragmentType, EntityRange.SubchunkStart, RangeLength);
 
 				const FMassEntityHandle* EntitiesArray = Chunk.GetEntityArray(EntityListOffsetWithinChunk);
 				TArrayView<const FMassEntityHandle> EntityHandlesView(&EntitiesArray[EntityRange.SubchunkStart], RangeLength);
@@ -1664,7 +1769,6 @@ void FMassArchetypeData::BatchAddSparseElementToEntities(const FMassArchetypeEnt
 	checkf(UE::Mass::IsSparse(ElementType), TEXT("%s is not a sparse element type"), *ElementType->GetName());
 
 	const bool bIsTag = UE::Mass::IsA<FMassTag>(ElementType);
-	SparseElementsBitSet.Add(*ElementType);
 
 	TArray<FMassArchetypeEntityCollection::FArchetypeEntityRange> Subchunks(EntityCollection.GetRanges());
 	for (const FMassArchetypeEntityCollection::FArchetypeEntityRange EntityRange : Subchunks)
@@ -1673,7 +1777,7 @@ void FMassArchetypeData::BatchAddSparseElementToEntities(const FMassArchetypeEnt
 		const int32 RangeLength = CalculateRangeLength(EntityRange, Chunk);
 
 		UE::Mass::FChunkSparseElements& SparseElements = Chunk.GetSparseElements(NumEntitiesPerChunk);
-		SparseElements.BatchAdd(ElementType, EntityRange.SubchunkStart, RangeLength);
+		bSparseElementsBitSetDirty |= SparseElements.BatchAdd(ElementType, EntityRange.SubchunkStart, RangeLength);
 
 		if (!bIsTag)
 		{
@@ -1688,11 +1792,6 @@ void FMassArchetypeData::BatchRemoveSparseElementFromEntities(const FMassArchety
 {
 	checkf(UE::Mass::IsSparse(ElementType), TEXT("%s is not a sparse element type"), *ElementType->GetName());
 
-	if (SparseElementsBitSet.Contains(*ElementType) == false)
-	{
-		return;
-	}
-
 	const bool bIsTag = UE::Mass::IsA<FMassTag>(ElementType);
 
 	TArray<FMassArchetypeEntityCollection::FArchetypeEntityRange> Subchunks(EntityCollection.GetRanges());
@@ -1701,14 +1800,18 @@ void FMassArchetypeData::BatchRemoveSparseElementFromEntities(const FMassArchety
 		FMassArchetypeChunk& Chunk = Chunks[EntityRange.ChunkIndex];
 		const int32 RangeLength = CalculateRangeLength(EntityRange, Chunk);
 
-		UE::Mass::FChunkSparseElements& SparseElements = Chunk.GetSparseElements(NumEntitiesPerChunk);
-		SparseElements.BatchRemove(ElementType, EntityRange.SubchunkStart, RangeLength);
-
-		if (!bIsTag)
+		UE::Mass::FChunkSparseElements& SparseElements = Chunk.GetSparseElementsUnsafe();
+		if (SparseElements.IsEmpty() == false)
 		{
-			const FMassEntityHandle* EntitiesArray = Chunk.GetEntityArray(EntityListOffsetWithinChunk);
-			TArrayView<const FMassEntityHandle> EntityHandlesView(&EntitiesArray[EntityRange.SubchunkStart], RangeLength);
-			SparseElementsStorage->BatchRemoveElementFromEntities(EntityHandlesView, ElementType);
+			const bool bElementsPresentDirty = SparseElements.BatchRemove(ElementType, EntityRange.SubchunkStart, RangeLength);
+			bSparseElementsBitSetDirty |= bElementsPresentDirty;
+
+			if (!bIsTag)
+			{
+				const FMassEntityHandle* EntitiesArray = Chunk.GetEntityArray(EntityListOffsetWithinChunk);
+				TArrayView<const FMassEntityHandle> EntityHandlesView(&EntitiesArray[EntityRange.SubchunkStart], RangeLength);
+				SparseElementsStorage->BatchRemoveElementFromEntities(EntityHandlesView, ElementType);
+			}
 		}
 	}
 }
diff --git a/Engine/Source/Runtime/MassEntity/Private/MassEntityManager.cpp b/Engine/Source/Runtime/MassEntity/Private/MassEntityManager.cpp
--- a/Engine/Source/Runtime/MassEntity/Private/MassEntityManager.cpp
+++ b/Engine/Source/Runtime/MassEntity/Private/MassEntityManager.cpp
@@ -1438,17 +1438,17 @@ void FMassEntityManager::AddSparseElementsToEntity(FMassEntityHandle EntityHandl
 	}
 }
 
-bool FMassEntityManager::RemoveSparseElementFromEntity(FMassEntityHandle EntityHandle, TNotNull<const UScriptStruct*> ElementType)
+void FMassEntityManager::RemoveSparseElementFromEntity(FMassEntityHandle EntityHandle, TNotNull<const UScriptStruct*> ElementType)
 {
 	CHECK_SPARSE_ELEMENT(ElementType);
 
-	CHECK_SYNC_API_RETURN(return false);
+	CHECK_SYNC_API();
 	CheckIfEntityIsActive(EntityHandle);
 
 	TSharedPtr<FMassArchetypeData>& CurrentArchetype = GetEntityStorageInterface().GetArchetypeAsShared(EntityHandle.Index);
 	check(CurrentArchetype);
 
-	return CurrentArchetype->RemoveSparseElementFromEntity(EntityHandle, ElementType);
+	CurrentArchetype->RemoveSparseElementFromEntity(EntityHandle, ElementType);
 }
 
 bool FMassEntityManager::DoesEntityHaveElement(FMassEntityHandle Entity, TNotNull<const UScriptStruct*> ElementType) const
diff --git a/Engine/Source/Runtime/MassEntity/Private/MassEntityView.cpp b/Engine/Source/Runtime/MassEntity/Private/MassEntityView.cpp
--- a/Engine/Source/Runtime/MassEntity/Private/MassEntityView.cpp
+++ b/Engine/Source/Runtime/MassEntity/Private/MassEntityView.cpp
@@ -127,6 +127,6 @@ bool FMassEntityView::HasElement(TNotNull<const UScriptStruct*> ElementType, con
 {
 	checkSlow(EntityDataHandle.IsValid(Archetype));
 	return Archetype
-		&& (Archetype->HasElement(ElementType, UE::Mass::EIncludeSparseElements::No) 
+		&& (Archetype->HasElement(ElementType) 
 			|| (IncludeSparseElements == UE::Mass::EIncludeSparseElements::Yes && Archetype->HasSparseElementForEntity(ElementType, EntityDataHandle)));
 }
diff --git a/Engine/Source/Runtime/MassEntity/Public/MassEntityManager.h b/Engine/Source/Runtime/MassEntity/Public/MassEntityManager.h
--- a/Engine/Source/Runtime/MassEntity/Public/MassEntityManager.h
+++ b/Engine/Source/Runtime/MassEntity/Public/MassEntityManager.h
@@ -430,11 +430,11 @@ public:
 		FStructView _ = AddSparseElementToEntity(EntityHandle, T::StaticStruct());
 	}
 
-	UE_API bool RemoveSparseElementFromEntity(FMassEntityHandle EntityHandle, TNotNull<const UScriptStruct*> ElementType);
+	UE_API void RemoveSparseElementFromEntity(FMassEntityHandle EntityHandle, TNotNull<const UScriptStruct*> ElementType);
 	template<UE::Mass::CSparse T>
-	bool RemoveSparseElementFromEntity(FMassEntityHandle EntityHandle)
+	void RemoveSparseElementFromEntity(FMassEntityHandle EntityHandle)
 	{
-		return RemoveSparseElementFromEntity(EntityHandle, T::StaticStruct());
+		RemoveSparseElementFromEntity(EntityHandle, T::StaticStruct());
 	}
 
 	/**
