[Mass]
Added FProcessingQueue as a new execution path for Mass phases, enabling processor and chunk-level scheduling with explicit contention, dependency tracking, and worker pools.
Added query pre/post callbacks and chunk queueing into FMassExecutionContext and FMassEntityQuery.
Added cvar mass.UseProcessingQueue to MassProcessingPhaseManager to toggle between new FProcessingQueue and current execution path.

#Synced-CL 50677047

diff --git a/Engine/Source/Developer/MassEntityTestSuite/Private/MassComplexTest.cpp b/Engine/Source/Developer/MassEntityTestSuite/Private/MassComplexTest.cpp
--- a/Engine/Source/Developer/MassEntityTestSuite/Private/MassComplexTest.cpp
+++ b/Engine/Source/Developer/MassEntityTestSuite/Private/MassComplexTest.cpp
@@ -765,7 +765,7 @@ struct FComplex_DestroyEntitiesDuringQuery : FProcessingPhasesTestBase
 
 							// Direct access to EntityManager (Main Buffer).
 							// This is safe because mass.FullyParallel=false ensures this runs on the GameThread.
-							EntityManager->Defer().DestroyEntity(Entity);
+							Context.Defer().DestroyEntity(Entity);
 						}
 					});
 
diff --git a/Engine/Source/Runtime/MassEntity/Private/MassEntityQuery.cpp b/Engine/Source/Runtime/MassEntity/Private/MassEntityQuery.cpp
--- a/Engine/Source/Runtime/MassEntity/Private/MassEntityQuery.cpp
+++ b/Engine/Source/Runtime/MassEntity/Private/MassEntityQuery.cpp
@@ -42,6 +42,9 @@ struct FScopedEntityQueryContext
 		: Query(InQuery), CachedExecutionContext(ExecutionContext)
 		, ScopedAccessDetector(InQuery)
 	{
+		// callback might pause execution until resources are available, must be made first
+		ExecutionContext.PreQueryCallback(&InQuery);
+		
 		CachedExecutionContext.PushQuery(InQuery);
 
 		// Setup output device to collect errors that can be reported by CacheSubsystemRequirements
@@ -64,6 +67,9 @@ struct FScopedEntityQueryContext
 
 	~FScopedEntityQueryContext()
 	{
+		// The query has discovered all the chunks for processing.
+		CachedExecutionContext.PostQueryCallback(&Query);
+
 		if (IsSuccessfullySetUp())
 		{
 			CachedExecutionContext.ClearExecutionData();
@@ -270,10 +276,28 @@ void FMassEntityQuery::ForEachEntityChunk(FMassExecutionContext& ExecutionContex
 
 			const FMassArchetypeHandle& ArchetypeHandle = ExecutionContext.GetEntityCollection().GetArchetype();
 			FMassArchetypeData& ArchetypeData = FMassArchetypeHelper::ArchetypeDataFromHandleChecked(ArchetypeHandle);
-			ArchetypeData.ExecuteFunction(ExecutionContext, ExecuteFunction
-				, EmptyMapping
-				, ExecutionContext.GetEntityCollection().GetRanges()
-				, ChunkCondition);
+			
+			if (ExecutionContext.HasChunkQueueFunc())
+			{
+				for (const FMassArchetypeEntityCollection::FArchetypeEntityRange& Range : ExecutionContext.GetEntityCollection().GetRanges())
+				{
+					ExecutionContext.QueueChunkChecked({
+						.SourceQuery = this,
+						.Archetype = &ArchetypeData,
+						.EntityRange = Range,
+						.ExecuteFunction = &ExecuteFunction,
+						.ChunkCondition = &ChunkCondition,
+						.RequirementMapping = nullptr,
+						.bParallelExecution = false});
+				}
+			}
+			else
+			{
+				ArchetypeData.ExecuteFunction(ExecutionContext, ExecuteFunction
+					, EmptyMapping
+					, ExecutionContext.GetEntityCollection().GetRanges()
+					, ChunkCondition);
+			}
 		}
 	}
 	else
@@ -304,10 +328,28 @@ void FMassEntityQuery::ForEachEntityChunk(FMassExecutionContext& ExecutionContex
 			
 			UE_TRACE_SCOPED_MASS_QUERY_FOR_EACH_REPORT_ARCHETYPE(ArchetypeData)
 			
-			ArchetypeData.ExecuteFunction(ExecutionContext, ExecuteFunction
-				, GetRequirementsMappingForArchetype(ArchetypeHandle)
-				, ExecutionContext.GetEntityCollection().GetRanges()
-				, ChunkCondition);
+			if (ExecutionContext.HasChunkQueueFunc())
+			{
+				for (const FMassArchetypeEntityCollection::FArchetypeEntityRange& Range : ExecutionContext.GetEntityCollection().GetRanges())
+				{
+					ExecutionContext.QueueChunkChecked({
+						.SourceQuery = this,
+						.Archetype = &ArchetypeData,
+						.EntityRange = Range,
+						.ExecuteFunction = &ExecuteFunction,
+						.ChunkCondition = &ChunkCondition,
+						.RequirementMapping = &GetRequirementsMappingForArchetype(ArchetypeHandle),
+						.bParallelExecution = false });
+				}
+			}
+			else
+			{
+				ArchetypeData.ExecuteFunction(ExecutionContext, ExecuteFunction
+					, GetRequirementsMappingForArchetype(ArchetypeHandle)
+					, ExecutionContext.GetEntityCollection().GetRanges()
+					, ChunkCondition);
+			}
+			
 		}
 		else
 		{
@@ -321,10 +363,26 @@ void FMassEntityQuery::ForEachEntityChunk(FMassExecutionContext& ExecutionContex
 				const FMassArchetypeHandle& ArchetypeHandle = ValidArchetypes[ArchetypeIndex];
 				FMassArchetypeData& ArchetypeData = FMassArchetypeHelper::ArchetypeDataFromHandleChecked(ArchetypeHandle);
 				
-				UE_TRACE_SCOPED_MASS_QUERY_FOR_EACH_REPORT_ARCHETYPE(ArchetypeData)
+				if (ExecutionContext.HasChunkQueueFunc())
+				{
+					for (int32 ChunkIndex = 0; ChunkIndex < ArchetypeData.GetChunkCount(); ChunkIndex++)
+					{
+						ExecutionContext.QueueChunkChecked({
+							.SourceQuery = this,
+							.Archetype = &ArchetypeData,
+							.EntityRange = FMassArchetypeEntityCollection::FArchetypeEntityRange(ChunkIndex),
+							.ExecuteFunction = &ExecuteFunction,
+							.ChunkCondition = &ChunkCondition,
+							.RequirementMapping = &ArchetypeFragmentMapping[ArchetypeIndex] });
+					}
+				}
+				else
+				{
+					UE_TRACE_SCOPED_MASS_QUERY_FOR_EACH_REPORT_ARCHETYPE(ArchetypeData)
 
-				ArchetypeData.ExecuteFunction(ExecutionContext, ExecuteFunction, ArchetypeFragmentMapping[ArchetypeIndex], ChunkCondition);
-				ExecutionContext.ClearFragmentViews(*this);
+					ArchetypeData.ExecuteFunction(ExecutionContext, ExecuteFunction, ArchetypeFragmentMapping[ArchetypeIndex], ChunkCondition);
+					ExecutionContext.ClearFragmentViews(*this);
+				}
 			}
 		}
 	}
@@ -397,10 +455,30 @@ void FMassEntityQuery::ForEachEntityChunk(FMassExecutionContext& ExecutionContex
 			bDone = true;
 		}
 
-		UE_TRACE_SCOPED_MASS_QUERY_FOR_EACH_REPORT_ARCHETYPE(ArchetypeData)
+		if (ExecutionContext.HasChunkQueueFunc())
+		{
+			const int32 MaxChunkIndex = FMath::Min(ArchetypeData.GetChunkCount(), Limiter.MaxChunkIndex);
+			Limiter.ChunkIndex = Limiter.ChunkIndex >= MaxChunkIndex ? 0 : FMath::Max(0, Limiter.ChunkIndex);
+			for (; Limiter.ChunkIndex < MaxChunkIndex && Limiter.EntityCountRemaining > 0; Limiter.ChunkIndex++)
+			{
+				const FMassArchetypeEntityCollection::FArchetypeEntityRange Range = FMassArchetypeEntityCollection::FArchetypeEntityRange(Limiter.ChunkIndex);
+				Limiter.EntityCountRemaining -= ArchetypeData.CalculateRangeLength(Range);
+				ExecutionContext.QueueChunkChecked({
+					.SourceQuery = this,
+					.Archetype = &ArchetypeData,
+					.EntityRange = Range,
+					.ExecuteFunction = &ExecuteFunction,
+					.ChunkCondition = &ChunkCondition,
+					.RequirementMapping = &ArchetypeFragmentMapping[ArchetypeIndex] });
+			}
+		}
+		else
+		{
+			UE_TRACE_SCOPED_MASS_QUERY_FOR_EACH_REPORT_ARCHETYPE(ArchetypeData)
 
-		ArchetypeData.ExecuteFunction(ExecutionContext, ExecuteFunction, ArchetypeFragmentMapping[ArchetypeIndex], ChunkCondition, &Limiter);
-		ExecutionContext.ClearFragmentViews(*this);
+			ArchetypeData.ExecuteFunction(ExecutionContext, ExecuteFunction, ArchetypeFragmentMapping[ArchetypeIndex], ChunkCondition, &Limiter);
+			ExecutionContext.ClearFragmentViews(*this);
+		}
 		
 		if (Limiter.ChunkIndex >= ArchetypeData.GetChunkCount())
 		{
@@ -544,7 +622,23 @@ void FMassEntityQuery::ParallelForEachEntityChunk(FMassExecutionContext& Executi
 	{
 		const EParallelForFlags ParallelForFlags = 
 			EnumHasAnyFlags(Flags, EParallelExecutionFlags::AutoBalance) ? EParallelForFlags::Unbalanced : EParallelForFlags::None;
-		if (bAllowParallelCommands)
+		
+		if (ExecutionContext.HasChunkQueueFunc())
+		{
+			for (FChunkJob& Job : Jobs)
+			{
+				const FMassQueryRequirementIndicesMapping* RequirementMapping = Job.ArchetypeIndex != INDEX_NONE ? &ArchetypeFragmentMapping[Job.ArchetypeIndex] : nullptr;
+				ExecutionContext.QueueChunkChecked({
+					.SourceQuery = this,
+					.Archetype = &Job.Archetype,
+					.EntityRange = Job.EntityRange,
+					.ExecuteFunction = &ExecuteFunction,
+					.ChunkCondition = &ChunkCondition,
+					.RequirementMapping = RequirementMapping,
+					.bParallelExecution = true});
+			}
+		}
+		else if (bAllowParallelCommands)
 		{
 			struct FTaskContext
 			{
diff --git a/Engine/Source/Runtime/MassEntity/Private/MassProcessingPhaseManager.cpp b/Engine/Source/Runtime/MassEntity/Private/MassProcessingPhaseManager.cpp
--- a/Engine/Source/Runtime/MassEntity/Private/MassProcessingPhaseManager.cpp
+++ b/Engine/Source/Runtime/MassEntity/Private/MassProcessingPhaseManager.cpp
@@ -13,6 +13,7 @@
 #include "MassEntityTrace.h"
 #include "MassProcessingContext.h"
 #include "Misc/StringOutputDevice.h"
+#include "MassProcessingQueue.h"
 
 #include UE_INLINE_GENERATED_CPP_BY_NAME(MassProcessingPhaseManager)
 
@@ -25,10 +26,12 @@ namespace UE::Mass::Tweakables
 {
 	bool bFullyParallel = MASS_DO_PARALLEL;
 	bool bMakePrePhysicsTickFunctionHighPriority = true;
+	bool bUseProcessingQueue = false;
 
 	FAutoConsoleVariableRef CVars[] = {
 		{TEXT("mass.FullyParallel"), bFullyParallel, TEXT("Enables mass processing distribution to all available thread (via the task graph)")},
 		{TEXT("mass.MakePrePhysicsTickFunctionHighPriority"), bMakePrePhysicsTickFunctionHighPriority, TEXT("Whether to make the PrePhysics tick function high priority - can minimise GameThread waits by starting parallel work as soon as possible")},
+		{TEXT("mass.UseProcessingQueue"), bUseProcessingQueue, TEXT("Enable MassProcessingQueue for executing processing phases.")},
 	};
 }
 
@@ -81,7 +84,23 @@ void FMassProcessingPhase::ExecuteTick(float DeltaTime, ELevelTick TickType, ENa
 
 	bIsDuringMassProcessing = true;	
 
-	if (bRunInParallelMode && PhaseManager->IsPaused() == false)
+	if (UE::Mass::Tweakables::bUseProcessingQueue&& PhaseManager->IsPaused() == false)
+	{
+		ObserversLockPtr = EntityManager.GetOrMakeObserversLock().ToSharedPtr();
+		FMassProcessingContext Context(EntityManager, DeltaTime);
+		auto OnProcessingDone = [this, DeltaTime]()
+			{
+				OnParallelExecutionDone(DeltaTime);
+			};
+		auto OnPreCommandFlush = [this]()
+			{
+				bIsDuringMassProcessing = false;
+				ObserversLockPtr = nullptr;
+			};
+		ProcessingQueue.Execute(MoveTemp(Context), OnProcessingDone, CurrentThread, OnPreCommandFlush);
+		ObserversLockPtr = nullptr;
+	}
+	else if (bRunInParallelMode && PhaseManager->IsPaused() == false)
 	{
 		bool bWorkRequested = false;
 		if (PhaseProcessor->IsEmpty() == false)
@@ -101,7 +120,6 @@ void FMassProcessingPhase::ExecuteTick(float DeltaTime, ELevelTick TickType, ENa
 
 			const FGraphEventRef PipelineCompletionEvent = UE::Mass::Executor::TriggerParallelTasks(*PhaseProcessor, MoveTemp(Context), 
 				OnProcessingDone, CurrentThread, OnPreCommandFlush);
-
 			if (PipelineCompletionEvent.IsValid())
 			{
 				MyCompletionGraphEvent->DontCompleteUntil(PipelineCompletionEvent);
@@ -127,6 +145,7 @@ void FMassProcessingPhase::ExecuteTick(float DeltaTime, ELevelTick TickType, ENa
 			
 			// release the observer lock before the FMassProcessingContext is destroyed and commands are flushed
 			ObserversLockPtr = nullptr;
+
 		}
 
 		bIsDuringMassProcessing = false;
@@ -572,6 +591,12 @@ void FMassProcessingPhaseManager::OnPhaseStart(FMassProcessingPhase& Phase)
 
 			GraphBuildState.bInitialized = true;
 
+			if (UE::Mass::Tweakables::bUseProcessingQueue)
+			{
+				UE::Mass::FProcessingQueue& PhaseQueue = ProcessingPhases[PhaseAsInt].ProcessingQueue;
+				PhaseQueue.SetProcessors(PhaseProcessor->GetChildProcessorsView());
+			}
+
 #if WITH_MASSENTITY_DEBUG
 			UObject* OwnerPtr = Owner.Get();
 			// print it all out to vislog
diff --git a/Engine/Source/Runtime/MassEntity/Private/MassProcessingQueue.cpp b/Engine/Source/Runtime/MassEntity/Private/MassProcessingQueue.cpp
new file mode 100644
--- /dev/null
+++ b/Engine/Source/Runtime/MassEntity/Private/MassProcessingQueue.cpp
@@ -0,0 +1,689 @@
+// Copyright Epic Games, Inc. All Rights Reserved.
+
+#include "MassProcessingQueue.h"
+#include "Containers/BitArray.h"
+#include "Containers/Queue.h"
+#include "MassProcessingTypes.h"
+#include "MassProcessor.h"
+#include "MassTypeManager.h"
+#include "MassCommandBuffer.h"
+#include "MassExecutionContext.h"
+#include "MassProcessingContext.h"
+#include "Logging/MessageLog.h"
+#include "MassArchetypeData.h"
+#include "MassEntityQuery.h"
+
+#define LOCTEXT_NAMESPACE "MassProcessingQueue"
+
+#define MASS_PROCESSING_QUEUE_LOGGING 0
+#define REUSE_CHUNK_QUEUE_ENTRIES 1
+
+namespace UE::Mass::Tweakables
+{
+	int32 MaxProcessorWorkerCount = 8;
+	int32 MaxChunkWorkerCount = 16;
+
+	FAutoConsoleVariableRef ProcessingQueueCVars[] = {
+		{TEXT("mass.ProcessingQueue.MaxProcessorWorkers"), MaxProcessorWorkerCount, TEXT("Max number of concurrent Processor Workers")},
+		{TEXT("mass.ProcessingQueue.MaxChunkWorkers"), MaxChunkWorkerCount, TEXT("Max number of concurrent Chunk Workers")},
+	};
+}
+
+namespace UE::Mass
+{
+	//----------------------------------------------------------------------//
+	//  FProcessingQueue
+	//----------------------------------------------------------------------//
+	void FProcessingQueue::AddProcessor(TNotNull<UMassProcessor*> InProcessor)
+	{
+		Processors.AppendProcessor(*InProcessor);
+		bProcessorListDirty = true;
+	}
+
+	void FProcessingQueue::AddProcessors(FMassRuntimePipeline& InPipeline)
+	{
+		Processors.AppendProcessors(InPipeline.GetMutableProcessors());
+		bProcessorListDirty = true;
+	}
+
+	void FProcessingQueue::AddProcessors(TConstArrayView<UMassProcessor*> InProcessors)
+	{
+		for (UMassProcessor* Proc : InProcessors)
+		{
+			Processors.AppendProcessor(*Proc);
+		}
+		bProcessorListDirty = true;
+	}
+
+	void FProcessingQueue::SetProcessors(TConstArrayView<UMassProcessor*> InProcessors)
+	{
+		Processors.Reset();
+		AddProcessors(InProcessors);
+		InitQueue();
+	}
+
+	void FProcessingQueue::Execute(FProcessingContext&& ProcessingContext, TFunction<void()> OnDoneNotification, ENamedThreads::Type CurrentThread, TFunction<void()> OnPreCommandFlushNotification)
+	{
+		TRACE_CPUPROFILER_EVENT_SCOPE_STR("Mass FProcessingQueue::Execute");
+
+		if (bProcessorListDirty)
+		{
+			InitQueue();
+		}
+
+		FMassExecutionContext ExecutionContext = MoveTemp(ProcessingContext).GetExecutionContext();
+		FMassEntityManager& EntityManager = ExecutionContext.GetEntityManagerChecked();
+
+		MainThreadProcessorCount = 0;
+		RemainingProcessorIndices.Reset();
+		RemainingMainThreadProcessorIndices.Reset();
+
+		int32 QueueIndex = 0;
+
+		for (FQueuedProcessor& QueuedProcessor : ProcessorQueue)
+		{
+			QueuedProcessor.PrepareForExecute();
+			if (QueuedProcessor.bRequiresGameThread || UE::Mass::Tweakables::MaxProcessorWorkerCount <= 0)
+			{
+				MainThreadProcessorCount++;
+				RemainingMainThreadProcessorIndices.Add(QueueIndex++);
+			}
+			else
+			{
+				RemainingProcessorIndices.Add(QueueIndex++);
+			}
+		}
+
+		PendingMainThreadProcessorCount.store(MainThreadProcessorCount);
+		PendingProcessorCount.store(ProcessorQueue.Num());
+		CompletedProcessorCount.store(0);
+
+		{
+			FMassEntityManager::FScopedProcessing ProcessingScope = EntityManager.NewProcessingScope();
+
+			UE::Tasks::FTask QueueRunner = (UE::Tasks::Launch(TEXT("MassProcessingQueue Runner Task"),
+				[this, &ExecutionContext]()
+				{
+					TRACE_CPUPROFILER_EVENT_SCOPE_TEXT("MassProcessingQueue Async Runner Task");
+					Run(ExecutionContext);
+				},
+				LowLevelTasks::ETaskPriority::BackgroundHigh));
+
+			if (MainThreadProcessorCount > 0)
+			{
+				UE::Tasks::FTask MainThreadRunner = UE::Tasks::Launch(TEXT("Mass Main Thread Processing Task"),
+					[this, &ExecutionContext]()
+					{
+						TRACE_CPUPROFILER_EVENT_SCOPE_TEXT("MassProcessingQueue Main-Thread Runner Task");
+						RunMainThread(ExecutionContext);
+					},
+					UE::Tasks::ETaskPriority::High,
+					UE::Tasks::EExtendedTaskPriority::Inline);
+
+				MainThreadRunner.Wait();
+			}
+
+			if (bBlockMainThread)
+			{
+				QueueRunner.Wait();
+
+				check(CompletedProcessorCount.load() == ProcessorQueue.Num());
+			}
+			else
+			{
+				checkf(false, TEXT("BlockMainThread = false is not yet supported"));
+			}
+		}
+
+		// There is a very small chance that the entity manager bookeeping for the ProcessingScope count d
+		while (EntityManager.IsProcessing())
+		{
+			FPlatformProcess::YieldCycles(1);
+		}
+
+		// this needs to happen after the ProcessingScope is destroyed:
+		if (OnPreCommandFlushNotification)
+		{
+			// this will unlock observers so they execute when commands are flushed:
+			OnPreCommandFlushNotification();
+		}
+
+		//@todo: Command buffers need to be merged to sort them, but it would probably be faster to
+			//	 sort them in place and process the commands directly. This would need a multi-container 
+			//	 sort method.
+		if (&ExecutionContext.Defer() != &EntityManager.Defer())
+		{
+			ExecutionContext.Defer().MoveAppend(EntityManager.Defer());
+		}
+
+		FWriteScopeLock ScopeLock(CommandBuffersLock);
+		while (CommandBufferPool.Num() > 0)
+		{
+			ExecutionContext.Defer().MoveAppend(*CommandBufferPool.Pop());
+		}
+
+		ExecutionContext.SetFlushDeferredCommands(true);
+		ExecutionContext.FlushDeferred();
+
+		CommandBufferPool.Reset();
+		OnDoneNotification();
+	}
+
+	void FProcessingQueue::Reserve(int32 RequestorIndex)
+	{
+		FChangeListener Listener(QueueChangeTracker);
+
+		for (;;)
+		{
+			if (TryReserve(RequestorIndex))
+			{
+				return;
+			}
+			Listener.WaitForChange();
+		}
+	}
+
+	bool FProcessingQueue::TryReserve(int32 RequestorIndex)
+	{
+		return ActiveContentionTracker.TryActivate(RequestorIndex);
+	}
+
+	bool FProcessingQueue::ReleaseAndTryReserve(int32 IndexToRealease, int32 IndexToReserve)
+	{
+		// unlock:
+		ActiveContentionTracker.Deactivate(IndexToRealease);
+		return TryReserve(IndexToReserve);
+	}
+
+	void FProcessingQueue::ReleaseAndReserve(int32 IndexToRealease, int32 IndexToReserve)
+	{
+		if (ReleaseAndTryReserve(IndexToRealease, IndexToReserve))
+		{
+			return;
+		}
+		Reserve(IndexToReserve);
+	}
+
+	void FProcessingQueue::Release(int32 RequestorIndex)
+	{
+		ActiveContentionTracker.Deactivate(RequestorIndex);
+	}
+
+	FQueuedProcessor* FProcessingQueue::GetNextReadyProcessor(TArray<int32>& InOutRemainingIndices)
+	{
+		for (int32 QueueIndex : InOutRemainingIndices)
+		{
+			FQueuedProcessor& QueuedProcessor = ProcessorQueue[QueueIndex];
+			if (QueuedProcessor.State == FQueuedProcessor::EState::Ready)
+			{
+				if (!ActiveContentionTracker.TryActivate(QueuedProcessor.ContentionIndex))
+				{
+					continue;
+				}
+
+				FChangeScope QueueChange(QueueChangeTracker);
+
+				QueuedProcessor.State = FQueuedProcessor::EState::Running;
+				
+
+				--PendingProcessorCount;
+				if (QueuedProcessor.bRequiresGameThread)
+				{
+					--PendingMainThreadProcessorCount;
+				}
+				// Ordering doesn't matter as this won't affect explicit ordering dependencies
+				InOutRemainingIndices.RemoveSingleSwap(QueueIndex, EAllowShrinking::No);
+				return &QueuedProcessor;
+			}
+		}
+		return nullptr;
+	}
+
+	bool FProcessingQueue::TryReserveChunkWorkerSlot()
+	{
+		int32 CurrentWorkerCount = ChunkWorkerCount.fetch_add(1, std::memory_order_release) + 1;
+		if (CurrentWorkerCount > UE::Mass::Tweakables::MaxChunkWorkerCount)
+		{
+			// we exceeded the chunk worker count so we need to undo the increment and fail
+			--ChunkWorkerCount;
+			return false;
+		}
+		return true;
+	}
+
+	void FProcessingQueue::ReserveChunkWorkerSlot()
+	{
+		// @todo: allow processors to request N number of workers and 
+		int32 CurrentWorkerCount = ChunkWorkerCount.fetch_add(1, std::memory_order_release) + 1;
+
+		while (CurrentWorkerCount > UE::Mass::Tweakables::MaxChunkWorkerCount)
+		{
+			// we exceeded the chunk worker count so we need to undo the increment and wait for it to go down
+			--ChunkWorkerCount;
+
+			FChangeListener ChunkWorkerCountListener(ChunkWorkerCount);
+			ChunkWorkerCountListener.WaitForChange();
+
+			// try again
+			CurrentWorkerCount = ChunkWorkerCount.fetch_add(1, std::memory_order_release) + 1;
+		}
+	}
+
+	void FProcessingQueue::ReleaseChunkWorkerSlot()
+	{
+		--ChunkWorkerCount;
+	}
+	
+	bool FProcessingQueue::TryActivateChunk(FQueuedChunk& InChunk, int32 DirectChunkContentionIndex)
+	{
+		// Calculated from "ProcessChunksBefore" or "ProcessChunksAfter"
+		if (InChunk.RemainingDependencyCount > 0)
+		{
+			return false; 
+		}
+
+		const bool bActive = ActiveChunkContentionTracker.TryAcquireAccess(InChunk, DirectChunkContentionIndex);
+		return bActive;
+	}
+
+	void FProcessingQueue::DeactivateChunk(FQueuedChunk& InChunk)
+	{
+		ActiveChunkContentionTracker.ReleaseAccess(InChunk, InChunk.QueryInfo->ChunkContentionIndex);
+	}
+
+	void FProcessingQueue::ProcessorFinished(FQueuedProcessor& CompletedProcessor)
+	{
+		FChangeScope QueueChange(QueueChangeTracker);
+
+		++CompletedProcessorCount;
+		CompletedProcessor.State = FQueuedProcessor::EState::Completed;
+
+		if (CompletedProcessor.DependentProcessorIndices.Num() > 0)
+		{
+			for (int32 DependentProcessorIndex : CompletedProcessor.DependentProcessorIndices)
+			{
+				if (ProcessorQueue[DependentProcessorIndex].RemainingDependencyCount.fetch_sub(1) == 1)
+				{
+					check(ProcessorQueue[DependentProcessorIndex].State == FQueuedProcessor::EState::WaitingForDependencies);
+					ProcessorQueue[DependentProcessorIndex].State = FQueuedProcessor::EState::Ready;
+				}
+			}
+		}
+		ActiveContentionTracker.Deactivate(CompletedProcessor.ContentionIndex);
+	}
+
+	void FProcessingQueue::MakeProcessorWorker(FMassExecutionContext& ExecutionContext, FQueuedProcessor* AssignedProcessor, bool bInLine)
+	{
+		if (bInLine)
+		{
+			if (AssignedProcessor)
+			{
+				FChangeScope ProcessorWorkerInlineScope(QueueChangeTracker);
+				FProcessorWorker ProcessorWorker(*this, ExecutionContext, true);
+				ProcessorWorker.AssignProcessor(AssignedProcessor);
+				ProcessorWorker.ExecuteProcessor();
+			}
+		}
+		else
+		{
+			UE::Tasks::Launch(TEXT("Mass Processor Worker Task"),
+				[this, &ExecutionContext, AssignedProcessor]()
+				{
+					FChangeScope ProcessorWorkerTaskScope(QueueChangeTracker);
+					TRACE_CPUPROFILER_EVENT_SCOPE_TEXT("Mass Processor Worker Task");
+					FProcessorWorker ProcessorWorker(*this, ExecutionContext);
+					if (AssignedProcessor)
+					{
+						FChangeScope ProcessorWorkerExecuteScope(QueueChangeTracker);
+						ProcessorWorker.AssignProcessor(AssignedProcessor);
+						ProcessorWorker.ExecuteProcessor();
+					}
+				},
+				LowLevelTasks::ETaskPriority::Normal);
+		}
+	}
+
+	void FProcessingQueue::MakeChunkWorker(FQueuedChunk* AssignedChunk, bool bInline)
+	{
+		if (bInline)
+		{
+			FChangeScope ChunkWorkerInlineScope(QueueChangeTracker);
+			FChunkWorker ChunkWorker(*this, AssignedChunk, bInline);
+			if (ChunkWorker.Chunk)
+			{
+				ChunkWorker.ExecuteChunk(true);
+			}
+		}
+		else
+		{
+			UE::Tasks::FTask WorkerTask = UE::Tasks::Launch(TEXT("Mass Chunk Worker Task"),
+				[this, AssignedChunk]()
+				{
+					FChangeScope ChunkWorkerTaskScope(QueueChangeTracker);
+					FChunkWorker ChunkWorker(*this, AssignedChunk);
+					if (ChunkWorker.Chunk)
+					{
+						ChunkWorker.ExecuteChunk(false);
+					}
+				}
+				, UE::Tasks::ETaskPriority::Normal);
+		}
+		
+	}
+
+	TSharedPtr<FMassCommandBuffer> FProcessingQueue::GetOrCreateCommandBuffer()
+	{
+		if (CommandBuffersLock.TryWriteLock())
+		{
+			if (CommandBufferPool.IsEmpty())
+			{
+				CommandBuffersLock.WriteUnlock();
+				return MakeShared<FMassCommandBuffer>();
+			}
+			TSharedPtr<FMassCommandBuffer> CommandBuffer = CommandBufferPool.Pop();
+			CommandBuffersLock.WriteUnlock();
+			CommandBuffer->ForceUpdateCurrentThreadID();
+			return CommandBuffer;
+		}
+		// if the lock fails, just make a new one instead of blocking
+		// @todo: consider creating a pool of chunk workers each with an owned command buffer to avoid creating them excessively
+		return MakeShared<FMassCommandBuffer>();
+	}
+
+	void FProcessingQueue::ReturnCommandBufferToPool(TSharedPtr<FMassCommandBuffer>& InCommandBuffer)
+	{
+		FWriteScopeLock ScopeLock(CommandBuffersLock);
+		check(InCommandBuffer.IsValid());
+		checkSlow(CommandBufferPool.Contains(InCommandBuffer) == false);
+		CommandBufferPool.Add(InCommandBuffer);
+	}
+
+	void FProcessingQueue::RunMainThread(FMassExecutionContext& ExecutionContext)
+	{
+		int32 CountRemaining;
+		do
+		{
+			FQueuedProcessor* QueuedProcessor = GetNextReadyProcessor(RemainingMainThreadProcessorIndices);
+			if (QueuedProcessor)
+			{
+				MakeProcessorWorker(ExecutionContext, QueuedProcessor, true);
+			}
+
+			// count all processors if there are not off-main-thread workers
+			CountRemaining = UE::Mass::Tweakables::MaxProcessorWorkerCount > 0 ? PendingMainThreadProcessorCount.load() : PendingProcessorCount.load();
+		} while (CountRemaining > 0);
+	}
+
+	void FProcessingQueue::Run(FMassExecutionContext& ExecutionContext)
+	{
+		FChangeListener ChaneListener(QueueChangeTracker);
+		for (;;)
+		{
+			int32 ProcessorsRemaining = PendingProcessorCount.load() - PendingMainThreadProcessorCount.load();
+
+			if (ProcessorsRemaining > 0)
+			{
+				int32 ProcessorWorkersToMake = UE::Mass::Tweakables::MaxProcessorWorkerCount - ProcessorWorkerCount.load();
+				while (ProcessorWorkersToMake-- > 0)
+				{
+					FQueuedProcessor* QueuedProcessor = GetNextReadyProcessor(RemainingProcessorIndices);
+					if (QueuedProcessor)
+					{
+						MakeProcessorWorker(ExecutionContext, QueuedProcessor);
+					}
+					else
+					{
+						break;
+					}
+				}
+			}
+
+			if (IsComplete())
+			{
+				break;
+			}
+			ChaneListener.WaitForChange();
+		}
+	}
+
+	bool FProcessingQueue::IsComplete() const
+	{
+		return CompletedProcessorCount.load() == ProcessorQueue.Num();
+	}
+
+	int32 FProcessingQueue::ComputeGraphDepthRecursive(int32 ProcessorIndex, TArray<int32>& BestDepth)
+	{
+		if (BestDepth[ProcessorIndex] != -1)
+		{
+			return BestDepth[ProcessorIndex];
+		}
+
+		int32 MaxChildDepth = 0;
+
+		// Find everyone who depends on THIS processor (ExecuteAfter this)
+		for (int32 DependentIndex : ProcessorQueue[ProcessorIndex].DependentProcessorIndices)
+		{
+			int32 ChildDepth = ComputeGraphDepthRecursive(DependentIndex, BestDepth);
+			if (ChildDepth > MaxChildDepth)
+			{
+				MaxChildDepth = ChildDepth;
+			}
+		}
+
+		int32 MyDepth = MaxChildDepth + 1; // I am 1 deeper than my deepest dependent
+		BestDepth[ProcessorIndex] = MyDepth;
+		return MyDepth;
+	}
+
+	void FProcessingQueue::InitializeExplicitDependencies()
+	{
+		// Reset and build base maps
+		TMap<FName, TArray<int32>> ClassToIndices; // class name -> processor indices
+		TMap<FName, TArray<int32>> GroupToIndices; // group name -> processor indices
+
+		ClassToIndices.Reserve(ProcessorQueue.Num());
+		GroupToIndices.Reserve(ProcessorQueue.Num());
+
+		for (int32 ProcessorIndex = 0; ProcessorIndex < ProcessorQueue.Num(); ++ProcessorIndex)
+		{
+			FQueuedProcessor& QueuedProcessor = ProcessorQueue[ProcessorIndex];
+			QueuedProcessor.DependentProcessorIndices.Reset();
+			QueuedProcessor.DependencyCount = 0;
+			const FName ClassName = QueuedProcessor.Processor->GetClass()->GetFName();
+			ClassToIndices.FindOrAdd(ClassName).Add(ProcessorIndex);
+
+			const FName GroupName = QueuedProcessor.Processor->GetExecutionOrder().ExecuteInGroup;
+			if (!GroupName.IsNone())
+			{
+				GroupToIndices.FindOrAdd(GroupName).Add(ProcessorIndex);
+			}
+		}
+
+		// Helper: add a dependency on From to To meaning "To waits for From"
+		auto AddEdge = [this](int32 FromIdx, int32 ToIdx)
+			{
+				if (FromIdx == INDEX_NONE || ToIdx == INDEX_NONE || FromIdx == ToIdx)
+				{
+					return;
+				}
+				FQueuedProcessor& From = ProcessorQueue[FromIdx];
+				FQueuedProcessor& To = ProcessorQueue[ToIdx];
+
+				if(!From.DependentProcessorIndices.Contains(ToIdx))
+				{
+					From.DependentProcessorIndices.Add(ToIdx);
+					++To.DependencyCount;
+				}
+			};
+
+		// Helper: resolve a name to a set of processor indices
+		auto GetIndicesForName = [&GroupToIndices, &ClassToIndices](const FName& Name) -> TSet<int32>
+			{
+				TSet<int32> Indices;
+
+				if (const TArray<int32>* GroupMatches = GroupToIndices.Find(Name))
+				{
+					for (int32 Idx : *GroupMatches)
+					{
+						Indices.Add(Idx);
+					}
+				}
+				if (const TArray<int32>* ClassMatches = ClassToIndices.Find(Name))
+				{
+					for (int32 Idx : *ClassMatches)
+					{
+						Indices.Add(Idx);
+					}
+				}
+				return Indices;
+			};
+
+		for (int32 ProcessorIndex = 0; ProcessorIndex < ProcessorQueue.Num(); ++ProcessorIndex)
+		{
+			UMassProcessor* const Processor = ProcessorQueue[ProcessorIndex].Processor;
+			
+			const FMassProcessorExecutionOrder& Order = Processor->GetExecutionOrder();
+			for (const FName& BeforeName : Order.ExecuteBefore)
+			{
+				const TSet<int32> BeforeIndices = GetIndicesForName(BeforeName);
+				for (int32 BeforeIndex : BeforeIndices)
+				{
+					AddEdge(ProcessorIndex, BeforeIndex);
+				}
+			}
+
+			for (const FName& AfterName : Order.ExecuteAfter)
+			{
+				const TSet<int32> AfterIndices = GetIndicesForName(AfterName);
+				for (int32 AfterIndex : AfterIndices)
+				{
+					AddEdge(AfterIndex, ProcessorIndex);
+				}
+			}
+		}
+		if (CheckForValidDependencies())
+		{
+			// @todo: distribute workers favoring higher graph depth
+			// Calculate Graph Depth for each processor
+			TArray<int32> BestDepth;
+			BestDepth.Init(-1, ProcessorQueue.Num());
+
+			for (int32 i = 0; i < ProcessorQueue.Num(); ++i)
+			{
+				ProcessorQueue[i].GraphDepth = ComputeGraphDepthRecursive(i, BestDepth);
+			}
+		}
+	}
+
+	bool FProcessingQueue::CheckForValidDependencies()
+	{
+		const int32 NumProcessors = ProcessorQueue.Num();
+		if (NumProcessors == 0)
+		{
+			return true;
+		}
+
+		TQueue<int32> Queue;
+		for (int32 ProcessorIndex = 0; ProcessorIndex < NumProcessors; ++ProcessorIndex)
+		{
+			ProcessorQueue[ProcessorIndex].PrepareForExecute();
+			if (ProcessorQueue[ProcessorIndex].RemainingDependencyCount.load() == 0)
+			{
+				Queue.Enqueue(ProcessorIndex);
+			}
+		}
+
+		int32 Processed = 0;
+		int32 Index;
+		while (Queue.Dequeue(Index))
+		{
+			++Processed;
+
+			FQueuedProcessor& QueuedProcessor = ProcessorQueue[Index];
+			for (int32 DepIndex : QueuedProcessor.DependentProcessorIndices)
+			{
+				FQueuedProcessor& Dep = ProcessorQueue[DepIndex];
+				int32 NewCount = --Dep.RemainingDependencyCount;
+				if (NewCount == 0)
+				{
+					Queue.Enqueue(DepIndex);
+				}
+			}
+		}
+
+		checkf(Processed == NumProcessors, TEXT("Explicit dependency cycle detected."));
+		return Processed == NumProcessors;
+	}
+
+	void FProcessingQueue::InitQueue()
+	{
+		bProcessorListDirty = false;
+		const int32 ProcessorCount = Processors.Num();
+		// todo: sort processors
+		int32 GlobalContendersCount = 0;
+		int32 DirectContendersCount = 0;
+		ProcessorQueue.Reset(ProcessorCount);
+		RemainingProcessorIndices.Reserve(ProcessorCount);
+
+		for (UMassProcessor* Processor : Processors.GetProcessorsView())
+		{
+			FQueuedProcessor& QueuedProcessor = ProcessorQueue.Emplace_GetRef(Processor);
+			QueuedProcessor.ContentionIndex = GlobalContendersCount++;
+			for (int32 QueryIndex = 0; QueryIndex < QueuedProcessor.QueryCount; ++QueryIndex)
+			{
+				QueuedProcessor.QueryInfos[QueryIndex].QueryContentionIndex = GlobalContendersCount++;
+				QueuedProcessor.QueryInfos[QueryIndex].ChunkGlobalContentionIndex = GlobalContendersCount++;
+				QueuedProcessor.QueryInfos[QueryIndex].ChunkContentionIndex = DirectContendersCount++;
+			}
+			QueuedProcessor.InitThreadRequirement();
+		}
+
+		InitializeExplicitDependencies();
+
+
+		ContentionMatrix.Init(GlobalContendersCount);
+		ActiveContentionTracker.Init();
+
+		ChunkContentionMatrix.Init(DirectContendersCount);
+		ActiveChunkContentionTracker.Init(ChunkContentionMatrix);
+
+		for (int32 ProcessorIndexA = 0; ProcessorIndexA < ProcessorCount - 1; ++ProcessorIndexA)
+		{
+			FQueuedProcessor& QueuedProcessorA = ProcessorQueue[ProcessorIndexA];
+
+			for (int32 ProcessorIndexB = ProcessorIndexA + 1; ProcessorIndexB < ProcessorCount; ++ProcessorIndexB)
+			{
+				FQueuedProcessor& QueuedProcessorB = ProcessorQueue[ProcessorIndexB];
+
+				// init mutual contention for ProcessorA vs ProcessorB
+				QueuedProcessorA.InitWriteContention(ContentionMatrix, QueuedProcessorB);
+
+				for (int32 QueryIndexB = 0; QueryIndexB < QueuedProcessorB.QueryCount; ++QueryIndexB)
+				{
+					// init mutual contention for ProcessorA vs each query owned by ProcessorB
+					QueuedProcessorB.InitWriteContention(ContentionMatrix, QueryIndexB, QueuedProcessorA);
+				}
+
+				for (int32 QueryIndexA = 0; QueryIndexA < QueuedProcessorA.QueryCount; ++QueryIndexA)
+				{
+					// init mutual contention for ProcessorB vs each query owned by ProcessorA
+					QueuedProcessorA.InitWriteContention(ContentionMatrix, QueryIndexA, QueuedProcessorB);
+					for (int32 QueryIndexB = 0; QueryIndexB < QueuedProcessorB.QueryCount; ++QueryIndexB)
+					{
+						// init mutual cross-chunk contention for each query owned by ProcessorB vs each query owned by ProcessorA
+						// "cross-chunk contention" means there's contention even if the chunks are different (indirect access, subsytems, etc)
+						QueuedProcessorA.InitCrossChunkWriteContention(ContentionMatrix, QueryIndexA, QueuedProcessorB, QueryIndexB);
+
+						// init mutual chunk contention for each query owned by ProcessorB vs each query owned by ProcessorA
+						// "chunk contention" means there's contention only within the access for a given chunk (direct access, composition changes, etc)
+						QueuedProcessorA.InitChunkWriteContention(ChunkContentionMatrix, QueryIndexA, QueuedProcessorB, QueryIndexB);
+					}
+				}
+			}
+		}
+	}
+}
+
+#undef MASS_PROCESSING_QUEUE_LOGGING
+#undef REUSE_CHUNK_QUEUE_ENTRIES
+#undef LOCTEXT_NAMESPACE
\ No newline at end of file
diff --git a/Engine/Source/Runtime/MassEntity/Private/MassProcessingQueueTypes.cpp b/Engine/Source/Runtime/MassEntity/Private/MassProcessingQueueTypes.cpp
new file mode 100644
--- /dev/null
+++ b/Engine/Source/Runtime/MassEntity/Private/MassProcessingQueueTypes.cpp
@@ -0,0 +1,806 @@
+// Copyright Epic Games, Inc. All Rights Reserved.
+
+#include "MassProcessingQueueTypes.h"
+#include "MassProcessingTypes.h"
+#include "MassProcessingQueue.h"
+#include "MassProcessor.h"
+#include "MassTypeManager.h"
+#include "MassCommandBuffer.h"
+#include "MassExecutionContext.h"
+#include "MassProcessingContext.h"
+#include "Logging/MessageLog.h"
+#include "MassArchetypeData.h"
+#include "Containers/BitArray.h"
+#include "MassEntityQuery.h"
+
+#define LOCTEXT_NAMESPACE "MassProcessingQueueTypes"
+
+#define MASS_PROCESSING_QUEUE_LOGGING 0
+
+namespace UE::Mass::Tweakables
+{
+	bool bFullLogging = false;
+	bool bLockLogging = false;
+
+	FAutoConsoleVariableRef ProcessingQueueTypesCVars[] = {
+		{TEXT("mass.ProcessingQueue.FullLogging"), bFullLogging, TEXT("Log Everything...")},
+		{TEXT("mass.ProcessingQueue.LockLogging"), bLockLogging, TEXT("Log All lock operations...")},
+	};
+}
+
+static inline const TCHAR* MassThreadTag(bool bMain)
+{
+	return bMain ? TEXT("Main") : TEXT("Async");
+}
+
+static inline const TCHAR* MassInlineTag(bool bInline)
+{
+	return bInline ? TEXT("Inline") : TEXT("Async");
+}
+
+#if MASS_PROCESSING_QUEUE_LOGGING
+#define LOG_PROC_WORKER(Fmt, ...) \
+	UE_CLOG(UE::Mass::Tweakables::bFullLogging, LogMass, VeryVerbose, TEXT("[Mass][FProcessorWorker][%s][TID=%u][%p][L%u] " Fmt), \
+		MassThreadTag(bInlineExecution), FPlatformTLS::GetCurrentThreadId(), this, __LINE__, ##__VA_ARGS__) \
+
+#define LOG_PROC_WORKER_CHUNK(ChunkPtr, Fmt, ...) \
+	UE_CLOG(UE::Mass::Tweakables::bFullLogging, LogMass, VeryVerbose, TEXT("[Mass][FProcessorWorker][%s][TID=%u][%p][Chunk=%p][SourceQuery=%p][L%u] " Fmt), \
+		MassThreadTag(bInlineExecution), FPlatformTLS::GetCurrentThreadId(), this, ChunkPtr, ChunkPtr->Params.SourceQuery, __LINE__, ##__VA_ARGS__)
+
+#define LOG_CHUNK_WORKER(Fmt, ...) \
+	UE_CLOG(UE::Mass::Tweakables::bFullLogging, LogMass, VeryVerbose, TEXT("[Mass][ChunkWorker][TID=%u][%p][L%u] " Fmt), \
+		FPlatformTLS::GetCurrentThreadId(), this, __LINE__, ##__VA_ARGS__)
+
+#define LOG_CHUNK_WORKER_CHUNK(ChunkPtr, Fmt, ...) \
+	UE_CLOG(UE::Mass::Tweakables::bFullLogging, LogMass, VeryVerbose, TEXT("[Mass][ChunkWorker][TID=%u][%p][Chunk=%p][SourceQuery=%p][L%u] " Fmt), \
+		FPlatformTLS::GetCurrentThreadId(), this, ChunkPtr, ChunkPtr->Params.SourceQuery, __LINE__, ##__VA_ARGS__)
+
+#else //MASS_PROCESSING_QUEUE_LOGGING
+#define LOG_PROC_WORKER(...) \
+	do {} while (0)
+#define LOG_PROC_WORKER_CHUNK(...) \
+	do {} while (0)
+#define LOG_CHUNK_WORKER(...) \
+	do {} while (0)
+#define LOG_CHUNK_WORKER_CHUNK(...) \
+	do {} while (0)
+
+#endif //MASS_PROCESSING_QUEUE_LOGGING
+
+#define CHECK_THREAD_CONSISTENCY(bRequireGameThread) \
+	checkfSlow( \
+		!(bRequireGameThread && !IsInGameThread()), \
+		TEXT("Thread consistency error: GameThread is required, but current thread is not the game thread!") \
+	)
+
+namespace UE::Mass
+{
+	// Enum -> text helpers
+	static inline const TCHAR* LexToString(UE::Mass::FQueuedProcessor::EState S)
+	{
+		using E = UE::Mass::FQueuedProcessor::EState;
+		switch (S)
+		{
+		case E::WaitingForDependencies:
+			return TEXT("WaitingForDependencies");
+		case E::Ready:
+			return TEXT("Ready");
+		case E::Running:
+			return TEXT("Active");
+		case E::Completed:
+			return TEXT("Complete");
+		case E::MAX:
+			return TEXT("MAX");
+		default:
+			return TEXT("<Unknown>");
+		}
+	}
+
+	static inline const TCHAR* LexToString(UE::Mass::FQueuedChunk::EState S)
+	{
+		using E = UE::Mass::FQueuedChunk::EState;
+		switch (S)
+		{
+		case E::Pending:
+			return TEXT("Pending");
+		case E::Running:
+			return TEXT("Running");
+		case E::Complete:
+			return TEXT("Complete");
+		case E::MAX:
+			return TEXT("MAX");
+		default:
+			return TEXT("<Unknown>");
+		}
+	}
+
+	//----------------------------------------------------------------------//
+	//  FQueuedProcessor
+	//----------------------------------------------------------------------//
+	FQueuedProcessor::FQueuedProcessor(TNotNull<UMassProcessor*> InProcessor)
+		: Processor(InProcessor)
+	{
+		bAggregateRequirements = Processor->bAggregateRequirements || (Processor->OwnedQueries.Num() > 1);
+		QueryCount = Processor->OwnedQueries.Num();
+
+		for (const FElementBitSet& Composition : Processor->GetProcessorEntityCreationRequirements().GetCompositions())
+		{
+			DataAccess.IndirectWritesElements += Composition;
+		}
+		
+		DataAccess.WritesSubsystems = Processor->GetProcessorRequirements().GetRequiredMutableSubsystems();
+		DataAccess.ReadsSubsystems = Processor->GetProcessorRequirements().GetRequiredConstSubsystems() + DataAccess.WritesSubsystems;
+
+		QueryInfos.SetNum(QueryCount);
+
+		for (int32 QueryIndex = 0; QueryIndex < QueryCount; ++QueryIndex)
+		{
+			QueryInfos[QueryIndex].SourceQuery = Processor->OwnedQueries[QueryIndex];
+			QueryInfos[QueryIndex].DataAccess.SetFragmentWriteRequirements(*Processor->OwnedQueries[QueryIndex]);
+			QueryInfos[QueryIndex].DataAccess.IndirectWritesElements = Processor->OwnedQueries[QueryIndex]->GetIndirectReadWriteFragments();
+			QueryInfos[QueryIndex].DataAccess.IndirectReadsElements = Processor->OwnedQueries[QueryIndex]->GetIndirectReadOnlyFragments() + QueryInfos[QueryIndex].DataAccess.IndirectWritesElements;
+			QueryInfos[QueryIndex].DataAccess.WritesSubsystems = Processor->OwnedQueries[QueryIndex]->GetRequiredMutableSubsystems();
+			QueryInfos[QueryIndex].DataAccess.ReadsSubsystems = Processor->OwnedQueries[QueryIndex]->GetRequiredConstSubsystems() + QueryInfos[QueryIndex].DataAccess.WritesSubsystems;
+			
+			if (bAggregateRequirements)
+			{
+				DataAccess.Add(QueryInfos[QueryIndex].DataAccess);
+			}
+		}
+
+		if (bAggregateRequirements)
+		{
+			for (FQueryInfo& Info : QueryInfos)
+			{
+				Info.DataAccess.Add(DataAccess);
+			}
+		}
+	}
+
+	void FQueuedProcessor::PrepareForExecute()
+	{
+		RemainingDependencyCount.store(DependencyCount);
+		State = DependencyCount == 0 ? EState::Ready : EState::WaitingForDependencies;
+		CurrentQueryIndex = QueryCount > 0 ? 0 : INDEX_NONE;
+	}
+
+	const FQueryInfo& FQueuedProcessor::GetCurrentQueryInfo() const
+	{
+		return QueryInfos[CurrentQueryIndex];
+	}
+
+	void FQueuedProcessor::InitThreadRequirement()
+	{
+		bRequiresGameThread = Processor->DoesRequireGameThreadExecution() || Processor->ProcessorRequirements.DoesRequireGameThreadExecution();
+
+		for (int QueryIndex = 0; QueryIndex < Processor->OwnedQueries.Num(); ++QueryIndex)
+		{
+			FQueryInfo& Info = QueryInfos[QueryIndex];
+			FMassEntityQuery& Query = *(Processor->OwnedQueries[QueryIndex]);
+
+			Info.bChunksRequireGameThread = Query.DoesRequireGameThreadExecution();
+			if (bAggregateThreadRequirements)
+			{
+				bRequiresGameThread |= Info.bChunksRequireGameThread;
+			}
+		}
+		if (bAggregateThreadRequirements)
+		{
+			for (FQueryInfo& Info : QueryInfos)
+			{
+				Info.bChunksRequireGameThread = bRequiresGameThread;
+			}
+		}
+
+	}
+
+	void FQueuedProcessor::InitWriteContention(FContentionMatrix& ContentionMatrix, FQueuedProcessor& Other)
+	{
+		if (Processor == Other.Processor)
+		{
+			return;
+		}
+
+		if (DataAccess.SubsystemsCollide(Other.DataAccess))
+		{
+			ContentionMatrix.SetMutual(ContentionIndex, Other.ContentionIndex);
+			return;
+		}
+
+		// find created entity collisions:
+		if (Processor->GetProcessorEntityCreationRequirements().HasAny(Other.Processor->GetProcessorEntityCreationRequirements()))
+		{
+			ContentionMatrix.SetMutual(ContentionIndex, Other.ContentionIndex);
+			return;
+		}
+	}
+
+	void FQueuedProcessor::InitWriteContention(FContentionMatrix& ContentionMatrix, int32 MyQueryIndex, FQueuedProcessor& Other)
+	{
+		FMassEntityQuery* Query = Processor->OwnedQueries[MyQueryIndex];
+
+		// Do the processor and query have a write collision on any subsystems?
+		if (QueryInfos[MyQueryIndex].DataAccess.SubsystemsCollide(Other.DataAccess))
+		{
+			// Subsystem access collision applies to processing all chunks, but not running the query.
+			ContentionMatrix.SetMutual(QueryInfos[MyQueryIndex].ChunkGlobalContentionIndex, Other.ContentionIndex);
+			return;
+		}
+
+		// Can the processor create entities the query would find?
+		for (const FMassElementBitSet& Composition : Other.Processor->GetProcessorEntityCreationRequirements().GetCompositions())
+		{
+			if (Query->DoesArchetypeMatchRequirements(Composition))
+			{
+				// This means the query can not run, and its chunks can not process while that processor is active.
+				ContentionMatrix.SetMutual(QueryInfos[MyQueryIndex].ChunkGlobalContentionIndex, Other.ContentionIndex);
+				return;
+			}
+		}
+
+		// Do the query and processor both specify an identical creation requirement?:
+		if (Other.Processor->GetProcessorEntityCreationRequirements().HasAny(Query->GetEntityCreationRequirements()))
+		{
+			// If they both create entities of the exact same archetype then they contend on that FMassArchetypeData.
+			ContentionMatrix.SetMutual(QueryInfos[MyQueryIndex].ChunkGlobalContentionIndex, Other.ContentionIndex);
+			return;
+		}
+	}
+
+	void FQueuedProcessor::InitCrossChunkWriteContention(FContentionMatrix& ContentionMatrix, int32 MyQueryIndex, FQueuedProcessor& Other, int32 OtherQueryIndex)
+	{
+		FMassEntityQuery* QueryA = Processor->OwnedQueries[MyQueryIndex];
+		FMassEntityQuery* QueryB = Other.Processor->OwnedQueries[OtherQueryIndex];
+		const FDataAccess& DataA = QueryInfos[MyQueryIndex].DataAccess;
+		const FDataAccess& DataB = Other.QueryInfos[OtherQueryIndex].DataAccess;
+
+		// Do the processor and query have a write collision on any subsystems?
+		if (DataA.SubsystemsCollide(DataB))
+		{
+			// Subsystem write collisions are global (cross-chunk), but do not block running the query.
+			ContentionMatrix.SetMutual(QueryInfos[MyQueryIndex].ChunkGlobalContentionIndex, Other.QueryInfos[OtherQueryIndex].ChunkGlobalContentionIndex);
+			return;
+		}
+
+		// Can QueryA create entities QueryB would find?
+		for (const FMassElementBitSet& Composition : QueryA->GetEntityCreationRequirements().GetCompositions())
+		{
+			if (QueryB->DoesArchetypeMatchRequirements(Composition))
+			{
+				// Entity creation and composition changes are global (cross-chunk) contention.
+				ContentionMatrix.SetMutual(QueryInfos[MyQueryIndex].ChunkGlobalContentionIndex, Other.QueryInfos[OtherQueryIndex].ChunkGlobalContentionIndex);
+				return;
+			}
+		}
+
+		// Can QueryB create entities QueryA would find?
+		for (const FMassElementBitSet& Composition : QueryB->GetEntityCreationRequirements().GetCompositions())
+		{
+			if (QueryA->DoesArchetypeMatchRequirements(Composition))
+			{
+				// Entity creation and composition changes are global (cross-chunk) contention.
+				ContentionMatrix.SetMutual(QueryInfos[MyQueryIndex].ChunkGlobalContentionIndex, Other.QueryInfos[OtherQueryIndex].ChunkGlobalContentionIndex);
+				return;
+			}
+		}
+
+		// Do QueryA and QueryB specify an identical entity creation requirement?:
+		if (QueryA->GetEntityCreationRequirements().HasAny(QueryB->GetEntityCreationRequirements()))
+		{
+			// They both create entities of an identical archetype which is cross-chunk contention.
+			ContentionMatrix.SetMutual(QueryInfos[MyQueryIndex].ChunkGlobalContentionIndex, Other.QueryInfos[OtherQueryIndex].ChunkGlobalContentionIndex);
+			return;
+		}
+
+		// indirect vs indirect fragment access collision:
+		if (DataA.IndirectWritesElements.HasAny(DataB.IndirectWritesElements)
+			|| DataA.IndirectWritesElements.HasAny(DataB.IndirectReadsElements)
+			|| DataB.IndirectWritesElements.HasAny(DataA.IndirectReadsElements))
+		{
+			// Indirect contention is cross-chunk.
+			ContentionMatrix.SetMutual(QueryInfos[MyQueryIndex].ChunkGlobalContentionIndex, Other.QueryInfos[OtherQueryIndex].ChunkGlobalContentionIndex);
+			return;
+		}
+
+		// indirect writes vs direct reads or writes collision:
+		if (DataA.IndirectWritesElements.HasAny(DataB.DirectReadsElements)
+			|| DataB.IndirectWritesElements.HasAny(DataA.DirectReadsElements)
+			|| DataA.IndirectWritesElements.HasAny(DataB.DirectWritesElements)
+			|| DataB.IndirectWritesElements.HasAny(DataA.DirectWritesElements))
+		{
+			// Any indirect contention is cross-chunk.
+			ContentionMatrix.SetMutual(QueryInfos[MyQueryIndex].ChunkGlobalContentionIndex, Other.QueryInfos[OtherQueryIndex].ChunkGlobalContentionIndex);
+			return;
+		}
+
+		// direct writes vs indirect reads collision:
+		if (DataA.DirectWritesElements.HasAny(DataB.IndirectReadsElements)
+			|| DataB.DirectWritesElements.HasAny(DataA.IndirectReadsElements))
+		{
+			// Any indirect contention is cross-chunk.
+			ContentionMatrix.SetMutual(QueryInfos[MyQueryIndex].ChunkGlobalContentionIndex, Other.QueryInfos[OtherQueryIndex].ChunkGlobalContentionIndex);
+			return;
+		}
+
+		// note: We don't test direct vs direct fragment access because those can only collide within a single chunk.
+		//       See InitChunkWriteContention.
+	}
+
+	void FQueuedProcessor::InitChunkWriteContention(FContentionMatrix& ContentionMatrix, int32 MyQueryIndex, FQueuedProcessor& Other, int32 OtherQueryIndex)
+	{
+		// note: we don't care about archetype collisions here because we're only testing possible collisions within the same chunk so we know the archetypes match
+
+		const FDataAccess& DataA = QueryInfos[MyQueryIndex].DataAccess;
+		const FDataAccess& DataB = Other.QueryInfos[OtherQueryIndex].DataAccess;
+
+		if (DataA.DirectWritesElements.HasAny(DataB.DirectReadsElements)
+			|| DataB.DirectWritesElements.HasAny(DataA.DirectReadsElements))
+		{
+			ContentionMatrix.SetMutual(QueryInfos[MyQueryIndex].ChunkContentionIndex, Other.QueryInfos[OtherQueryIndex].ChunkContentionIndex);
+			return;
+		}
+
+		// all other contention should be caught by the cross-chunk contention check
+	}
+
+	//----------------------------------------------------------------------//
+	//   FDataAccess
+	//----------------------------------------------------------------------//
+	void FDataAccess::Add(const FDataAccess& Other)
+	{
+		DirectReadsElements += Other.DirectReadsElements;
+		DirectWritesElements += Other.DirectWritesElements;
+		IndirectReadsElements += Other.IndirectReadsElements;
+		IndirectWritesElements += Other.IndirectWritesElements;
+		ReadsSubsystems += Other.ReadsSubsystems;
+		WritesSubsystems += Other.WritesSubsystems;
+	}
+
+	bool FDataAccess::SubsystemsCollide(const FDataAccess& Other) const
+	{
+		// find mutable subsystem collisions:
+		return WritesSubsystems.HasAny(Other.WritesSubsystems)
+			|| WritesSubsystems.HasAny(Other.ReadsSubsystems)
+			|| Other.WritesSubsystems.HasAny(ReadsSubsystems);
+	}
+
+	bool FDataAccess::CollidesWithIndirect(const FDataAccess& Other) const
+	{
+		// indirect vs indirect fragment access collision:
+		const bool bIvI = IndirectWritesElements.HasAny(Other.IndirectWritesElements)
+			|| IndirectWritesElements.HasAny(Other.IndirectReadsElements)
+			|| Other.IndirectWritesElements.HasAny(IndirectReadsElements);
+
+		// indirect writes vs direct reads or writes collision:
+		const bool bIvD = IndirectWritesElements.HasAny(Other.DirectReadsElements)
+			|| Other.IndirectWritesElements.HasAny(DirectReadsElements)
+			|| IndirectWritesElements.HasAny(Other.DirectWritesElements)
+			|| Other.IndirectWritesElements.HasAny(DirectWritesElements);
+
+		// direct writes vs indirect reads collision:
+		const bool bDvI = DirectWritesElements.HasAny(Other.IndirectReadsElements)
+			|| Other.DirectWritesElements.HasAny(IndirectReadsElements);
+
+		return bIvI || bIvD || bDvI;
+	}
+
+	bool FDataAccess::CollidesWithDirect(const FDataAccess& Other) const
+	{
+		const bool bDvD = DirectWritesElements.HasAny(Other.DirectReadsElements)
+			|| DirectWritesElements.HasAny(Other.DirectWritesElements)
+			|| Other.DirectWritesElements.HasAny(DirectReadsElements);
+		return bDvD;
+	}
+
+	void FDataAccess::SetFragmentWriteRequirements(const FMassFragmentRequirements& FragmentRequirements)
+	{
+		auto ProcessFragmentRequirements = [this](TConstArrayView<FMassFragmentRequirementDescription> Descriptions)
+			{
+				for (const FMassFragmentRequirementDescription& Description : Descriptions)
+				{
+					if (Description.StructType && Description.Presence != EMassFragmentPresence::None)
+					{
+						switch (Description.AccessMode)
+						{
+						case EMassFragmentAccess::ReadWrite:
+							DirectWritesElements.Add(*Description.StructType);
+							DirectReadsElements.Add(*Description.StructType);
+							break;
+						case EMassFragmentAccess::ReadOnly:
+							DirectReadsElements.Add(*Description.StructType);
+							break;
+						default:
+							break;
+						}
+					}
+				}
+			};
+
+		ProcessFragmentRequirements(FragmentRequirements.GetFragmentRequirements());
+		ProcessFragmentRequirements(FragmentRequirements.GetChunkFragmentRequirements());
+		ProcessFragmentRequirements(FragmentRequirements.GetConstSharedFragmentRequirements());
+		ProcessFragmentRequirements(FragmentRequirements.GetSharedFragmentRequirements());
+	}
+
+	FProcessorWorker::FProcessorWorker(FProcessingQueue& InQueue, const FMassExecutionContext& InExecutionContext, bool bInMainThread)
+		: OwningQueue(InQueue)
+		, GlobalExecutionContext(InExecutionContext)
+		, bInlineExecution(bInMainThread)
+		, CommandBuffer(OwningQueue.GetOrCreateCommandBuffer())
+	{
+		CHECK_THREAD_CONSISTENCY(bInlineExecution);
+		LOG_PROC_WORKER(TEXT("Creating FProcessorWorker (%p): bInlineExecution: %s"), this, bInlineExecution ? TEXT("true") : TEXT("false"));
+		
+		const int32 CurProcessorWorkerCount = ++OwningQueue.ProcessorWorkerCount;
+#if WITH_MASSENTITY_DEBUG
+		if (CurProcessorWorkerCount > OwningQueue.MaxUsedProcessorWorkers)
+		{
+			OwningQueue.MaxUsedProcessorWorkers = CurProcessorWorkerCount;
+		}
+#endif
+	}
+
+	FProcessorWorker::~FProcessorWorker()
+	{
+		LOG_PROC_WORKER(TEXT("Destroying FProcessorWorker"));
+		Cleanup();
+	}
+
+	void FProcessorWorker::Cleanup()
+	{
+		OwningQueue.ReturnCommandBufferToPool(CommandBuffer);
+		--OwningQueue.ProcessorWorkerCount;
+	}
+
+	FChunkWorker::FChunkWorker(FProcessingQueue& InQueue, FQueuedChunk* InChunk, bool bInUseContextCommandBuffer)
+		: OwningQueue(InQueue)
+		, Chunk(InChunk)
+		, bUseContextCommandBuffer(bInUseContextCommandBuffer)
+	{
+		LOG_CHUNK_WORKER(TEXT("Creating FChunkWorker"));
+	}
+
+	FChunkWorker::~FChunkWorker()
+	{
+		Cleanup();
+	}
+
+	void FChunkWorker::Cleanup()
+	{
+		if (CommandBuffer.IsValid())
+		{
+			OwningQueue.ReturnCommandBufferToPool(CommandBuffer);
+		}
+		OwningQueue.ReleaseChunkWorkerSlot();
+	}
+
+	TSharedPtr<FMassCommandBuffer>& FChunkWorker::GetOrCreateCommandBuffer()
+	{
+		if (!CommandBuffer.IsValid())
+		{
+			CommandBuffer = OwningQueue.GetOrCreateCommandBuffer();
+		}
+		CommandBuffer->ForceUpdateCurrentThreadID();
+		return CommandBuffer;
+	}
+
+	void FChunkWorker::ExecuteChunk(bool bInlineExecution)
+	{
+		LOG_CHUNK_WORKER_CHUNK(Chunk, TEXT("ExecuteChunk"));
+		check(Chunk);
+		CHECK_THREAD_CONSISTENCY(Chunk->QueryInfo->bChunksRequireGameThread);
+
+		static const FMassQueryRequirementIndicesMapping EmptyMapping;
+
+		FMassExecutionContext& SourceContext = *Chunk->Params.SourceContext;
+
+		if (bUseContextCommandBuffer && !Chunk->Params.bParallelExecution)
+		{
+			if (!bInlineExecution)
+			{
+				SourceContext.Defer().ForceUpdateCurrentThreadID();
+			}
+			LOG_CHUNK_WORKER_CHUNK(Chunk, TEXT("Executing chunk serially"));
+			Chunk->Params.Archetype->ExecutionFunctionForChunk(
+				SourceContext,
+				*Chunk->Params.ExecuteFunction,
+				Chunk->Params.RequirementMapping ? *Chunk->Params.RequirementMapping : EmptyMapping,
+				Chunk->Params.EntityRange,
+				*Chunk->Params.ChunkCondition);
+		}
+		else
+		{
+			LOG_CHUNK_WORKER_CHUNK(Chunk, TEXT("Executing chunk in parallel"));
+
+			// need to create a local execution context for parallel execution
+			FMassExecutionContext LocalExecutionContext(SourceContext);
+			LocalExecutionContext.PushQuery(*Chunk->Params.SourceQuery);
+			LocalExecutionContext.SetDeferredCommandBuffer(GetOrCreateCommandBuffer());
+			Chunk->Params.Archetype->ExecutionFunctionForChunk(
+				LocalExecutionContext,
+				*Chunk->Params.ExecuteFunction,
+				Chunk->Params.RequirementMapping ? *Chunk->Params.RequirementMapping : EmptyMapping,
+				Chunk->Params.EntityRange,
+				*Chunk->Params.ChunkCondition);
+
+			// @todo: this is pointless and just prevents an ensure in the context destructor...
+			// there's nothing inherently unsafe about a context being destroyed while it had a query on the stack
+			// all the state that cares about that is local to the context
+			LocalExecutionContext.PopQuery(*Chunk->Params.SourceQuery);
+		}
+
+		LOG_CHUNK_WORKER_CHUNK(Chunk, TEXT("ExecuteChunk Complete, Reporting to OwningQueue"));
+		Chunk->OwningProcessorWorker->ChunkDone(Chunk);
+	}
+
+	void FQueuedProcessor::SetCurrentQuery(const FMassEntityQuery* SourceQuery)
+	{
+		// processors could execute queries in any order
+		for (int Index = 0; Index < QueryInfos.Num(); ++Index)
+		{
+			if (QueryInfos[Index].SourceQuery == SourceQuery)
+			{
+				CurrentQueryIndex = Index;
+				return;
+			}
+		}
+		checkf(false, TEXT("Query not found"));
+	}
+
+	void FProcessorWorker::ExecuteProcessor()
+	{
+		CHECK_THREAD_CONSISTENCY(bInlineExecution);
+		check(RunningProcessor);
+
+		TRACE_CPUPROFILER_EVENT_SCOPE_TEXT(*FString::Printf(TEXT("UE::Mass::FProcessorWorker::ExecuteProcessor: %s"), *RunningProcessor->Processor->GetProcessorName()));
+
+		LOG_PROC_WORKER(TEXT("ExecuteProcessor: Processor=%s"), *RunningProcessor->Processor->GetProcessorName());
+		LOG_PROC_WORKER(TEXT("ExecuteProcessor: QueuedProcessor=%p"), RunningProcessor);
+
+		if (!RunningProcessor->Processor->IsActive())
+		{
+			FQueuedProcessor* CompletedProcessor = RunningProcessor;
+			RunningProcessor = nullptr;
+			OwningQueue.ProcessorFinished(*CompletedProcessor);
+			return;
+		}
+
+		FMassExecutionContext ProcessorExecutionContext(GlobalExecutionContext);
+#if WITH_MASSENTITY_DEBUG
+		DebugCurrentExecutionContext = &ProcessorExecutionContext;
+#endif
+		ProcessorExecutionContext.SetDeferredCommandBuffer(CommandBuffer);
+
+		ProcessorExecutionContext.SetPreQueryCallback([this](const FMassEntityQuery* SourceQuery)
+			{
+				RunningProcessor->SetCurrentQuery(SourceQuery);
+				WaitForQueryResources();
+			});
+
+		ProcessorExecutionContext.SetPostQueryCallback([this](const FMassEntityQuery* SourceQuery)
+			{
+				RunningProcessor->SetCurrentQuery(SourceQuery);
+				ProcessChunks();
+			});
+
+		ProcessorExecutionContext.SetChunkQueueFunc([this](FMassChunkProcessingQueueParams&& Params)
+			{
+				QueueChunk(MoveTemp(Params));
+			});
+
+		FMassEntityManager& EntityManager = ProcessorExecutionContext.GetEntityManagerChecked();
+		FMassEntityManager::FScopedProcessing ProcessingScope = EntityManager.NewProcessingScope();
+
+		RunningProcessor->Processor->CallExecute(EntityManager, ProcessorExecutionContext);
+
+		LOG_PROC_WORKER(TEXT("ExecuteProcessor: Processor CallExecute complete. Proc=%p"), RunningProcessor);
+		check(ChunksRemaining.load() == 0);
+		
+		FQueuedProcessor* CompletedProcessor = RunningProcessor;
+#if WITH_MASSENTITY_DEBUG
+		DebugCurrentExecutionContext = nullptr;
+#endif
+		RunningProcessor = nullptr;
+		OwningQueue.ProcessorFinished(*CompletedProcessor);
+	}
+
+	void FProcessorWorker::WaitForQueryResources()
+	{
+		// the processor has made a "ForEachEntityChunk" call so hold here until the query resources are free
+
+		LOG_PROC_WORKER(TEXT("WaitForQueryResources: Proc=%p"), RunningProcessor);
+
+		const FQueryInfo& QueryInfo = RunningProcessor->GetCurrentQueryInfo();
+		// release processor contention index and lock the query index:
+		OwningQueue.ReleaseAndReserve(RunningProcessor->ContentionIndex, QueryInfo.QueryContentionIndex);
+		LOG_PROC_WORKER(TEXT("WaitForQueryResources: Resources locked! Proc=%p"), RunningProcessor);
+		// the processor will continue to the query now which will queue chunks
+
+		if (bInlineExecution && OwningQueue.TryReserve(QueryInfo.ChunkContentionIndex))
+		{
+			// if executing in-line, try here to reserve global chunk resources so chunks can process immediately on queing
+			GlobalChunkContentionIndex = QueryInfo.ChunkContentionIndex;
+		}
+	}
+
+	void FProcessorWorker::ProcessChunks()
+	{
+		LOG_PROC_WORKER(TEXT("ProcessChunks: Proc=%p"), RunningProcessor);
+		UE_MT_SCOPED_WRITE_ACCESS(RunningProcessor->ChunkQueueAccessDetector);
+
+		TArray<int32>& Pending = RunningProcessor->PendingChunkIndices;
+		if (!Pending.IsEmpty())
+		{
+			// All queued chunks should have the same query and thus the same contention indices.
+			// So we can get them and reserve them outside the processing loop.
+			const int32 QueryContentionIndex = RunningProcessor->ChunkQueue[0].QueryInfo->QueryContentionIndex;
+			const int32 DirectContentionIndex = RunningProcessor->ChunkQueue[0].QueryInfo->ChunkContentionIndex;
+
+
+			// the query has run so we can release that contention index and acquire the global chunk index if needed
+			if (GlobalChunkContentionIndex != INDEX_NONE)
+			{
+				OwningQueue.Release(QueryContentionIndex);
+			}
+			else
+			{
+				GlobalChunkContentionIndex = RunningProcessor->ChunkQueue[0].QueryInfo->ChunkGlobalContentionIndex;
+				OwningQueue.ReleaseAndReserve(QueryContentionIndex, GlobalChunkContentionIndex);
+			}
+
+			TArray<FQueuedChunk>& Queue = RunningProcessor->ChunkQueue;
+			const bool bRunInline = bInlineExecution || (!RunningProcessor->ChunkQueue[0].Params.bParallelExecution);
+			
+			while (!Pending.IsEmpty())
+			{
+				for (int32 PendingIndex = Pending.Num() - 1; PendingIndex >= 0; --PendingIndex)
+				{
+					int32 ChunkIndex = Pending[PendingIndex];
+					FQueuedChunk& QueuedChunk = Queue[ChunkIndex];
+
+					OwningQueue.ReserveChunkWorkerSlot();
+
+					if (OwningQueue.TryActivateChunk(QueuedChunk, DirectContentionIndex))
+					{
+						Pending.RemoveAtSwap(PendingIndex, EAllowShrinking::No);
+						QueuedChunk.State = FQueuedChunk::EState::Running;
+						// the chunk worker will release the chunk worker slot on completion
+						OwningQueue.MakeChunkWorker(&QueuedChunk, bRunInline);
+					}
+					else
+					{
+						// chunk has active contention, release our worker slot and move on
+						OwningQueue.ReleaseChunkWorkerSlot();
+					}
+				}
+			}
+
+			// wait for all chunks to finish processing
+			FChangeListener RemainingChunks(ChunksRemaining);
+			RemainingChunks.WaitForValue(0);
+			RunningProcessor->ChunkQueue.Reset();
+		}
+
+		LOG_PROC_WORKER(TEXT("ProcessChunks: Chunks complete, waiting for main processor resource lock. Proc=%p"), RunningProcessor);
+		
+		// Chunks are all processed, so re-acquire lock for processor requirements
+		if (GlobalChunkContentionIndex == INDEX_NONE)
+		{
+			OwningQueue.Reserve(RunningProcessor->ContentionIndex);
+		}
+		else
+		{
+			OwningQueue.ReleaseAndReserve(GlobalChunkContentionIndex, RunningProcessor->ContentionIndex);
+			GlobalChunkContentionIndex = INDEX_NONE;
+		}
+		
+		if (!bInlineExecution)
+		{
+			// Chunk workers might have run on a different thread using this command buffer.
+			// Reset the thread ID to handle deferred commands from the processor after the queue.
+			CommandBuffer->ForceUpdateCurrentThreadID();
+		}
+	}
+
+	void FProcessorWorker::QueueChunk(FMassChunkProcessingQueueParams&& Params)
+	{
+		TArray<FQueuedChunk>& Queue = RunningProcessor->ChunkQueue;
+		UE_MT_SCOPED_WRITE_ACCESS(RunningProcessor->ChunkQueueAccessDetector);
+		FQueuedChunk& QueuedChunk = Queue.Emplace_GetRef(FQueuedChunk(this, MoveTemp(Params)));
+
+		QueuedChunk.QueryInfo = &RunningProcessor->GetCurrentQueryInfo();
+		QueuedChunk.OwningProcessorWorker = this;
+		if (GlobalChunkContentionIndex == INDEX_NONE && QueuedChunk.Params.bParallelExecution == false)
+		{
+			GlobalChunkContentionIndex = QueuedChunk.QueryInfo->ChunkContentionIndex;
+			if (!OwningQueue.TryReserve(GlobalChunkContentionIndex))
+			{
+				GlobalChunkContentionIndex = INDEX_NONE;
+			}
+		}
+
+		bool ChunkQueued = true;
+		// chunk queueing is serial, but chunk completion can happen concurrently so we track completion atomically
+		++ChunksRemaining;
+		if (GlobalChunkContentionIndex != INDEX_NONE && OwningQueue.TryReserveChunkWorkerSlot())
+		{
+			// bHasGlobalChunkLock implies inline execution
+			// @todo: start processing chunks here immediately for all execution modes
+
+			if (OwningQueue.TryActivateChunk(QueuedChunk, QueuedChunk.QueryInfo->ChunkContentionIndex))
+			{
+				QueuedChunk.State = FQueuedChunk::EState::Running;
+
+				// the chunk worker will release the chunk worker slot on completion
+				// MakeChunkWorker with bInline = true will not return until the chunk is processed
+				OwningQueue.MakeChunkWorker(&QueuedChunk, true /*bInline*/);
+
+				// the chunk will be processed inline above so it's safe to remove it from the queue immediately
+				Queue.RemoveAt(Queue.Num() - 1, 1, EAllowShrinking::No);
+				ChunkQueued = false;
+			}
+			else
+			{
+				// chunk has active contention, release our worker slot and move on
+				OwningQueue.ReleaseChunkWorkerSlot();
+			}
+		}
+
+		if (ChunkQueued)
+		{
+			RunningProcessor->PendingChunkIndices.Add(Queue.Num() - 1);
+		}
+
+		LOG_PROC_WORKER_CHUNK(&QueuedChunk, TEXT("QueueChunk: Chunk queued. ChunksRemaining=%d Proc=%p"), ChunksRemaining.load(), RunningProcessor);
+	}
+
+	void FProcessorWorker::ChunkDone(FQueuedChunk* InChunk)
+	{
+		const int32 ChunksRemainingValue = --ChunksRemaining;
+		InChunk->State = FQueuedChunk::EState::Complete;
+		OwningQueue.DeactivateChunk(*InChunk);
+
+		LOG_PROC_WORKER_CHUNK(InChunk,
+			TEXT("ChunkDone: ChunksRemaining=%d%s"),
+			ChunksRemainingValue,
+			ChunksRemainingValue <= 0 ? TEXT(" (Triggering ChunksDone)") : TEXT(""));
+	}
+
+	FQueuedChunk::FQueuedChunk(FProcessorWorker* InOwningProcessorWorker, FMassChunkProcessingQueueParams&& InParams)
+		: OwningProcessorWorker(InOwningProcessorWorker)
+		, Params(MoveTemp(InParams))
+	{
+	}
+
+	FQueuedChunk::FQueuedChunk(FQueuedChunk&& Other) noexcept
+		: State(Other.State)
+		, OwningProcessorWorker(Other.OwningProcessorWorker)
+		, QueryInfo(Other.QueryInfo)
+		, Params(MoveTemp(Other.Params))
+	{
+		// reset the state of Other
+		Other.State = EState::MAX;
+		Other.QueryInfo = nullptr;
+		Other.OwningProcessorWorker = nullptr;
+	}
+	
+	void FProcessorWorker::AssignProcessor(FQueuedProcessor* QueuedProcessor)
+	{
+		LOG_PROC_WORKER(TEXT("Processor Assigned: Processor=%s"), *QueuedProcessor->Processor->GetProcessorName());
+		RunningProcessor = QueuedProcessor;
+
+		bInlineExecution |= RunningProcessor->bRequiresGameThread;
+
+#if WITH_MASSENTITY_DEBUG
+		QueuedProcessor->DebugCurrentProcessorWorker = this;
+#endif
+	}
+}
+
+#undef CHECK_THREAD_CONSISTENCY
+#undef LOG_PROC_WORKER
+#undef LOG_PROC_WORKER_CHUNK
+#undef LOG_CHUNK_WORKER
+#undef LOG_CHUNK_WORKER_CHUNK
+#undef MASS_PROCESSING_QUEUE_LOGGING
+#undef LOCTEXT_NAMESPACE
\ No newline at end of file
diff --git a/Engine/Source/Runtime/MassEntity/Public/MassArchetypeTypes.h b/Engine/Source/Runtime/MassEntity/Public/MassArchetypeTypes.h
--- a/Engine/Source/Runtime/MassEntity/Public/MassArchetypeTypes.h
+++ b/Engine/Source/Runtime/MassEntity/Public/MassArchetypeTypes.h
@@ -26,10 +26,12 @@ struct FMassEntityView;
 struct FMassDebugger;
 struct FMassArchetypeHelper;
 struct FMassArchetypeVersionedHandle;
+struct FMassChunkProcessingQueueParams;
 
 using FMassEntityExecuteFunction = TFunction< void(FMassExecutionContext& /*ExecutionContext*/, int32 /*EntityIndex*/) >;
 using FMassExecuteFunction = TFunction< void(FMassExecutionContext& /*ExecutionContext*/) >;
 using FMassChunkConditionFunction = TFunction< bool(const FMassExecutionContext& /*ExecutionContext*/) >;
+using FMassQueueChunkFunction = TFunction< void(FMassChunkProcessingQueueParams&&) >;
 
 //-----------------------------------------------------------------------------
 // FMassArchetypeHandle
@@ -342,6 +344,17 @@ struct FMassQueryRequirementIndicesMapping
 	}
 };
 
+struct FMassChunkProcessingQueueParams
+{
+	FMassExecutionContext* SourceContext = nullptr;
+	FMassEntityQuery* SourceQuery = nullptr;
+	FMassArchetypeData* Archetype = nullptr;
+	const FMassArchetypeEntityCollection::FArchetypeEntityRange EntityRange;
+	const FMassExecuteFunction* ExecuteFunction = nullptr;
+	const FMassChunkConditionFunction* ChunkCondition = nullptr;
+	const FMassQueryRequirementIndicesMapping* RequirementMapping = nullptr;
+	bool bParallelExecution = false;
+};
 
 //-----------------------------------------------------------------------------
 // INLINES
diff --git a/Engine/Source/Runtime/MassEntity/Public/MassExecutionContext.h b/Engine/Source/Runtime/MassEntity/Public/MassExecutionContext.h
--- a/Engine/Source/Runtime/MassEntity/Public/MassExecutionContext.h
+++ b/Engine/Source/Runtime/MassEntity/Public/MassExecutionContext.h
@@ -867,6 +867,56 @@ protected:
 
 	MASSENTITY_API bool ExecuteSparseElementsFilterImpl(const FSparseEntityIterator& EntityIt) const;
 	bool ExecuteSparseChunkFilterImpl(const FMassArchetypeChunk& Chunk) const;
+
+protected:
+	TFunction<void(const FMassEntityQuery*)> PreQueryCallbackFunc;
+	TFunction<void(const FMassEntityQuery*)> PostQueryCallbackFunc;
+
+	FMassQueueChunkFunction ChunkQueueFunc;
+
+public:
+	void SetPreQueryCallback(TFunction<void(const FMassEntityQuery*)> Function)
+	{
+		PreQueryCallbackFunc = Function;
+	}
+
+	void PreQueryCallback(const FMassEntityQuery* SourceQuery)
+	{
+		if (PreQueryCallbackFunc)
+		{
+			PreQueryCallbackFunc(SourceQuery);
+		}
+	}
+
+	void SetPostQueryCallback(TFunction<void(const FMassEntityQuery*)> Function)
+	{
+		PostQueryCallbackFunc = Function;
+	}
+
+	bool HasChunkQueueFunc() const
+	{
+		return ChunkQueueFunc.IsSet();
+	}
+
+	void QueueChunkChecked(FMassChunkProcessingQueueParams&& Params)
+	{
+		check(ChunkQueueFunc.IsSet());
+		Params.SourceContext = this;
+		ChunkQueueFunc(MoveTemp(Params));
+	}
+
+	void SetChunkQueueFunc(FMassQueueChunkFunction&& InChunkQueueFunc)
+	{
+		ChunkQueueFunc = MoveTemp(InChunkQueueFunc);
+	}
+
+	void PostQueryCallback(const FMassEntityQuery* SourceQuery)
+	{
+		if (PostQueryCallbackFunc)
+		{
+			PostQueryCallbackFunc(SourceQuery);
+		}
+	}
 };
 
 //-----------------------------------------------------------------------------
diff --git a/Engine/Source/Runtime/MassEntity/Public/MassProcessingPhaseManager.h b/Engine/Source/Runtime/MassEntity/Public/MassProcessingPhaseManager.h
--- a/Engine/Source/Runtime/MassEntity/Public/MassProcessingPhaseManager.h
+++ b/Engine/Source/Runtime/MassEntity/Public/MassProcessingPhaseManager.h
@@ -13,6 +13,7 @@
 #include "MassProcessingTypes.h"
 #include "MassProcessor.h"
 #include "MassProcessorDependencySolver.h"
+#include "MassProcessingQueue.h"
 #include "MassProcessingPhaseManager.generated.h"
 
 
@@ -88,6 +89,8 @@ protected:
 	// composite processor representing work to be performed. GC-referenced via AddReferencedObjects
 	TObjectPtr<UMassCompositeProcessor> PhaseProcessor = nullptr;
 
+	UE::Mass::FProcessingQueue ProcessingQueue;
+
 	EMassProcessingPhase Phase = EMassProcessingPhase::MAX;
 	FOnPhaseEvent OnPhaseStart;
 	FOnPhaseEvent OnPhaseEnd;
diff --git a/Engine/Source/Runtime/MassEntity/Public/MassProcessingQueue.h b/Engine/Source/Runtime/MassEntity/Public/MassProcessingQueue.h
new file mode 100644
--- /dev/null
+++ b/Engine/Source/Runtime/MassEntity/Public/MassProcessingQueue.h
@@ -0,0 +1,138 @@
+// Copyright Epic Games, Inc. All Rights Reserved.
+
+#pragma once
+
+#include "Containers/Array.h"
+#include "MassProcessingTypes.h"
+#include "MassProcessingQueueTypes.h"
+
+class UMassProcessor;
+struct FMassExecutionContext;
+
+#define UE_API MASSENTITY_API
+
+namespace UE::Mass
+{
+	struct FProcessingQueue
+	{
+	public:
+		/** Add a processor to the queue and resolve dependencies */
+		UE_API void AddProcessor(TNotNull<UMassProcessor*> InProcessor);
+
+		/** Add processors from a pipeline to the queue and resolve dependencies */
+		UE_API void AddProcessors(FMassRuntimePipeline& InPipeline);
+
+		/** Add processors to the queue and resolve dependencies */
+		UE_API void AddProcessors(TConstArrayView<UMassProcessor*> InProcessors);
+
+		/** Clear the processing queue and re-initialize with the provided array */
+		UE_API void SetProcessors(TConstArrayView<UMassProcessor*> InProcessors);
+
+		/** Immediately execute the processors in the queue */
+		UE_API void Execute(FProcessingContext&& ProcessingContext, TFunction<void()> OnDoneNotification, ENamedThreads::Type CurrentThread, TFunction<void()> OnPreCommandFlushNotification);
+
+		/** Immediately execute the processors in the queue */
+		UE_API void InitQueue();
+
+		/** Should Execute(...) await the completion of all processors before returning */
+		bool bBlockMainThread = true;
+
+	private:
+		friend FProcessorWorker;
+		friend FChunkWorker;
+
+		void RunMainThread(FMassExecutionContext& ExecutionContext);
+		void Run(FMassExecutionContext& ExecutionContext);
+
+		enum class EThreadFlag : uint8
+		{
+			MainThread = 1 << 0,
+			OffMainThread = 1 << 1,
+
+			Any = MainThread | OffMainThread
+		};
+
+		/** Gets the next processor that is ready to execute, sets its state to running and updates contention tracking and other bookeeping */
+		FQueuedProcessor* GetNextReadyProcessor(TArray<int32>& InOutRemainingIndices);
+
+		bool IsComplete() const;
+		std::atomic<int32> PendingMainThreadProcessorCount;
+		std::atomic<int32> PendingProcessorCount;
+		std::atomic<int32> CompletedProcessorCount;
+
+		/** Removes RequestorIndex from ActiveItems */
+		void Release(int32 RequestorIndex);
+
+		/** Checks if bits in Contention collide with bits in ActiveItems, if not, adds RequestorIndex to ActiveItems */
+		bool TryReserve(int32 RequestorIndex);
+		/** TryReserve then if that fails, wait until contending items are released and try again */
+		void Reserve(int32 RequestorIndex);
+
+		/** Release a reservation and immediately attempt TryReserve on a different index */
+		bool ReleaseAndTryReserve(int32 IndexToRealease, int32 IndexToReserve);
+		/** ReleaseAndReserve, then if that fails immediately Reserve (which will wait) */
+		void ReleaseAndReserve(int32 IndexToRealease, int32 IndexToReserve);
+
+		bool TryActivateChunk(FQueuedChunk& InChunk, int32 DirectChunkContentionIndex);
+		void DeactivateChunk(FQueuedChunk& InChunk);
+		bool TryReserveChunkWorkerSlot();
+		void ReserveChunkWorkerSlot();
+		void ReleaseChunkWorkerSlot();
+
+		/** Report that a processor is finished, set its state internal state and update queue bookeeping */
+		void ProcessorFinished(FQueuedProcessor& CompletedProcessor);
+
+		/**
+		 * Make a processor worker and immediately execute the assigned processor
+		 * If bInline is true then the processor is executed immediately on this thread and this method won't return until it's complete
+		 * Otherwise, a task is created to run the processor worker and execute the processor asynchronously.
+		 */
+		void MakeProcessorWorker(FMassExecutionContext& ExecutionContext, FQueuedProcessor* AssignedProcessor = nullptr, bool bInLine = false);
+		void MakeChunkWorker(FQueuedChunk* AssignedChunk = nullptr, bool bInline = false);
+		TSharedPtr<FMassCommandBuffer> GetOrCreateCommandBuffer();
+		void ReturnCommandBufferToPool(TSharedPtr<FMassCommandBuffer>& InCommandBuffer);
+		void InitializeExplicitDependencies();
+		bool CheckForValidDependencies();
+
+		/** Recursively computes dependency depth */
+		int32 ComputeGraphDepthRecursive(int32 ProcessorIndex, TArray<int32>& Memo);
+
+		/** Tracks all global contention */
+		FContentionMatrix ContentionMatrix;
+
+		/** Tracks active global contention */
+		FActiveContentionTracker ActiveContentionTracker{ ContentionMatrix };
+		
+		FContentionMatrix ChunkContentionMatrix;
+		/** Only tracks possible contention of direct chunk access (i.e. contention that can only occur accessing the same chunk) */
+		FActiveChunkTracker ActiveChunkContentionTracker;
+
+		TArray<FQueuedProcessor> ProcessorQueue;
+
+		/** array of indices in ProcessorQueue that still need to execute on the main thread */
+		TArray<int32> RemainingMainThreadProcessorIndices;
+
+		/** array of indices in ProcessorQueue that still need to execute */
+		TArray<int32> RemainingProcessorIndices;
+
+		bool bProcessorListDirty = false;
+		FMassRuntimePipeline Processors;
+		int32 MainThreadProcessorCount = 0;
+
+		FRWLock CommandBuffersLock;
+		std::atomic<int32> ChunkWorkerCount{ 0 };
+		std::atomic<int32> ProcessorWorkerCount{ 0 };
+		std::atomic<uint32> QueueChangeTracker{ 0 };
+		//FActiveChunkSet ActiveChunks;
+#if WITH_MASSENTITY_DEBUG
+		int32 MaxUsedProcessorWorkers = 0;
+		int32 MaxUsedChunkWorkers = 0;
+#endif
+		TArray<TSharedPtr<FMassCommandBuffer>> CommandBufferPool;
+	};
+
+} // namespace UE::Mass
+
+
+
+#undef UE_API
diff --git a/Engine/Source/Runtime/MassEntity/Public/MassProcessingQueueTypes.h b/Engine/Source/Runtime/MassEntity/Public/MassProcessingQueueTypes.h
new file mode 100644
--- /dev/null
+++ b/Engine/Source/Runtime/MassEntity/Public/MassProcessingQueueTypes.h
@@ -0,0 +1,1027 @@
+// Copyright Epic Games, Inc. All Rights Reserved.
+
+#pragma once
+
+#include "MassEntityQuery.h"
+#include "MassExecutionContext.h"
+
+#define UE_API MASSENTITY_API
+
+class UMassProcessor;
+
+namespace UE::Mass
+{
+	struct FProcessingQueue;
+	struct FDataAccess;
+	struct FQueuedProcessor;
+	struct FQueuedChunk;
+	struct FQueryInfo;
+	class FChunkWorker;
+	class FProcessorWorker;
+
+	/**
+	 * Dynamically sized bit-matrix to track contention between items.
+	 */
+	struct FContentionMatrix
+	{
+	public:
+		using WordType = uint64;
+		static constexpr int32 BitsPerWord = 64;
+		static constexpr int32 BitsPerWordLog2 = 6;
+
+		FContentionMatrix() = default;
+
+		explicit FContentionMatrix(int32 InNumElements)
+		{
+			Init(InNumElements);
+		}
+
+		/**
+		 * Resizes the matrix to hold InNumElements.
+		 * This clears all existing data.
+		 */
+		void Init(int32 InNumElements)
+		{
+			check(InNumElements >= 0);
+			NumElements = InNumElements;
+
+			// Calculate how many 64-bit words we need per row
+			WordsPerRow = (NumElements + (BitsPerWord - 1)) / BitsPerWord;
+
+			// Total size = Rows * Width
+			const int32 TotalWords = NumElements * WordsPerRow;
+
+			MatrixBits.Reset();
+			MatrixBits.SetNumZeroed(TotalWords);
+		}
+
+		FORCEINLINE int32 GetSize() const
+		{
+			return NumElements;
+		}
+
+		FORCEINLINE int32 GetWordsPerRow() const
+		{
+			return WordsPerRow;
+		}
+
+		/** Zero the entire matrix. */
+		FORCEINLINE void Clear()
+		{
+			FMemory::Memzero(MatrixBits.GetData(), MatrixBits.Num() * sizeof(WordType));
+		}
+
+		/** Set contends(i,j) = true */
+		void Set(int32 RowIndex, int32 ColumnIndex)
+		{
+			check(IsValidIndex(RowIndex) && IsValidIndex(ColumnIndex));
+			check(RowIndex != ColumnIndex);
+			const int32 WordIndex = (ColumnIndex >> BitsPerWordLog2);
+			const int32 BitOffset = (ColumnIndex & (BitsPerWord - 1));
+
+			// Calculate flattened index: Row * Stride + ColWord
+			const int32 FlatIndex = (RowIndex * WordsPerRow) + WordIndex;
+
+			MatrixBits[FlatIndex] |= (WordType(1) << BitOffset);
+		}
+
+		/** Set contends(i,j) = true and contends(j,i) = true*/
+		void SetMutual(int32 RowIndex, int32 ColumnIndex)
+		{
+			Set(RowIndex, ColumnIndex);
+			Set(ColumnIndex, RowIndex);
+		}
+
+		/** Test contention(i,j) */
+		FORCEINLINE bool Test(int32 RowIndex, int32 ColumnIndex) const
+		{
+			check(IsValidIndex(RowIndex) && IsValidIndex(ColumnIndex));
+
+			const int32 WordIndex = (ColumnIndex >> BitsPerWordLog2);
+			const int32 BitOffset = (ColumnIndex & (BitsPerWord-1));
+			const int32 FlatIndex = (RowIndex * WordsPerRow) + WordIndex;
+
+			return ((MatrixBits[FlatIndex] >> BitOffset) & 1u) != 0;
+		}
+
+		/** raw pointer to the start of the row's words. */
+		FORCEINLINE const WordType* GetRowWordPointer(int32 RowIndex) const
+		{
+			check(IsValidIndex(RowIndex));
+			return &MatrixBits[RowIndex * WordsPerRow];
+		}
+
+	private:
+		int32 NumElements = 0;
+		int32 WordsPerRow = 0;
+
+		TArray<WordType> MatrixBits;
+
+		FORCEINLINE bool IsValidIndex(int32 Index) const
+		{
+			return Index >= 0 && Index < NumElements;
+		}
+	};
+
+	/**
+	 * A single chunk processing callback generated by FMassEntityQuery::ForEachEntityChunk.
+	 */
+	struct FQueuedChunk
+	{
+		enum class EState : uint8
+		{
+			Pending,
+			Running,
+			Complete,
+			MAX
+		};
+
+		FQueuedChunk() = default;
+		FQueuedChunk(FProcessorWorker* InOwningProcessorWorker, FMassChunkProcessingQueueParams&& InParams);
+		FQueuedChunk(FQueuedChunk&& Other) noexcept;
+
+		FQueuedChunk(const FQueuedChunk&) = delete;
+		FQueuedChunk& operator=(const FQueuedChunk&) = delete;
+		FQueuedChunk& operator=(FQueuedChunk&& Other) = delete;
+
+		uint64 GetChunkID() const
+		{
+			return (reinterpret_cast<uint64>(Params.Archetype) << 16) ^ static_cast<uint64>(Params.EntityRange.ChunkIndex);
+		}
+
+		bool IsSameArchetypeChunk(const FQueuedChunk& Other) const
+		{
+			return GetChunkID() == Other.GetChunkID();
+		}
+
+		EState State = EState::Pending;
+		int32 ActiveChunkSlotIndex = INDEX_NONE;
+
+		/** Count of how many "ProcessAfter" chunks remain (includes "ProcessBefore" from other chunks/processors) */
+		int32 RemainingDependencyCount = 0;
+		FProcessorWorker* OwningProcessorWorker = nullptr;
+		const FQueryInfo* QueryInfo = nullptr;
+		FMassChunkProcessingQueueParams Params;
+	};
+
+	/** Tracks which FQueuedChunks are active to manage contention */
+	class FActiveChunkTracker
+	{
+	public:
+		using WordType = uint64;
+		static constexpr int32 MaxActiveChunks = 32;
+		static constexpr uint64 InvalidChunkID = 0;
+		FActiveChunkTracker() = default;
+
+		void Init(const FContentionMatrix& InMatrix)
+		{
+			MatrixRef = &InMatrix;
+			WordsPerRow = InMatrix.GetWordsPerRow();
+
+			// Initialize IDs to 0
+			for (int32 i = 0; i < MaxActiveChunks; ++i)
+			{
+				ChunkIDs[i].store(InvalidChunkID, std::memory_order_relaxed);
+				SlotRefCounts[i] = 0;
+				SlotLocks[i].clear();
+			}
+
+			SlotAssignmentLock.clear();
+
+			// Resize bitmask buffer: MaxActiveChunks slots * WordsPerRow
+			BitmaskBuffer.Reset();
+			BitmaskBuffer.SetNumZeroed(MaxActiveChunks * WordsPerRow);
+		}
+
+		/**
+		 * Checks active contention and reserves the chunk's contention index if safe
+		 */
+		bool TryAcquireAccess(FQueuedChunk& Chunk, int32 DirectChunkContentionIndex)
+		{
+			const uint64 ChunkID = Chunk.GetChunkID();
+			checkSlow(ChunkID != InvalidChunkID);
+			checkSlow(DirectChunkContentionIndex >= 0);
+			checkSlow(Chunk.ActiveChunkSlotIndex == -1);
+
+			// Fast scan (relaxed)
+			int32 FoundIndex = -1;
+			for (int32 i = 0; i < MaxActiveChunks; ++i)
+			{
+				if (ChunkIDs[i].load(std::memory_order_relaxed) == ChunkID)
+				{
+					FoundIndex = i;
+					break;
+				}
+			}
+
+			if (FoundIndex != -1)
+			{
+				// Existing slot
+				LockSlot(FoundIndex);
+
+				// Validate slot ownership
+				if (ChunkIDs[FoundIndex].load(std::memory_order_acquire) != ChunkID)
+				{
+					UnlockSlot(FoundIndex);
+					// Retry from scratch (extremely unlikely)
+					return TryAcquireAccess(Chunk, DirectChunkContentionIndex);
+				}
+
+				WordType* SlotBits = GetSlotBits(FoundIndex);
+				const WordType* JobRow = MatrixRef->GetRowWordPointer(DirectChunkContentionIndex);
+
+				// Contention check
+				for (int32 w = 0; w < WordsPerRow; ++w)
+				{
+					if ((SlotBits[w] & JobRow[w]) != 0)
+					{
+						UnlockSlot(FoundIndex);
+						return false;
+					}
+				}
+
+				// Mark active
+				AddJobToBitmask(SlotBits, DirectChunkContentionIndex);
+				++SlotRefCounts[FoundIndex];
+
+				UnlockSlot(FoundIndex);
+
+				Chunk.ActiveChunkSlotIndex = FoundIndex;
+				return true;
+			}
+
+			// assign slot for new ChunkID
+			while (SlotAssignmentLock.test_and_set(std::memory_order_acquire))
+			{
+				FPlatformProcess::YieldCycles(10);
+			}
+
+			// Double-check under AllocationLock (someone may have added it)
+			for (int32 i = 0; i < MaxActiveChunks; ++i)
+			{
+				if (ChunkIDs[i].load(std::memory_order_relaxed) == ChunkID)
+				{
+					SlotAssignmentLock.clear(std::memory_order_release);
+					// Re-run the found path (will lock+validate)
+					return TryAcquireAccess(Chunk, DirectChunkContentionIndex);
+				}
+			}
+
+			int32 AssignedSlotIndex = -1;
+			// find an empty slot
+			for (int32 i = 0; i < MaxActiveChunks; ++i)
+			{
+				if (ChunkIDs[i].load(std::memory_order_acquire) == InvalidChunkID)
+				{
+					AssignedSlotIndex = i;
+					break;
+				}
+			}
+
+			if (AssignedSlotIndex == -1)
+			{
+				SlotAssignmentLock.clear(std::memory_order_release);
+				checkf(false, TEXT("FActiveChunkTracker overflow! Active chunks exceeded MaxActiveChunks (%d)."), MaxActiveChunks);
+				return false;
+			}
+
+			// Initialize under slot lock BEFORE setting ChunkID
+			LockSlot(AssignedSlotIndex);
+
+#if DO_CHECK
+			check(ChunkIDs[AssignedSlotIndex].load(std::memory_order_relaxed) == InvalidChunkID);
+#endif
+
+			WordType* SlotBits = GetSlotBits(AssignedSlotIndex);
+			FMemory::Memzero(SlotBits, WordsPerRow * sizeof(WordType));
+
+#if DO_CHECK
+			const int32 WordIdx = (DirectChunkContentionIndex >> FContentionMatrix::BitsPerWordLog2);
+			check(WordIdx >= 0 && WordIdx < WordsPerRow);
+#endif
+
+			AddJobToBitmask(SlotBits, DirectChunkContentionIndex);
+			SlotRefCounts[AssignedSlotIndex] = 1;
+
+			// Publish LAST
+			ChunkIDs[AssignedSlotIndex].store(ChunkID, std::memory_order_release);
+
+			UnlockSlot(AssignedSlotIndex);
+			SlotAssignmentLock.clear(std::memory_order_release);
+
+			Chunk.ActiveChunkSlotIndex = AssignedSlotIndex;
+			return true;
+		}
+
+		void ReleaseAccess(FQueuedChunk& Chunk, int32 DirectChunkContentionIndex)
+		{
+			const int32 SlotIndex = (int32)Chunk.ActiveChunkSlotIndex;
+			checkSlow(SlotIndex >= 0 && SlotIndex < MaxActiveChunks);
+
+			const uint64 ChunkID = Chunk.GetChunkID();
+			checkSlow(ChunkID != InvalidChunkID);
+			checkSlow(DirectChunkContentionIndex >= 0);
+
+			LockSlot(SlotIndex);
+
+			// Validate slot still belongs to this chunk
+			check(ChunkIDs[SlotIndex].load(std::memory_order_acquire) == ChunkID);
+
+			WordType* SlotBits = GetSlotBits(SlotIndex);
+
+			const int32 WordIdx = (DirectChunkContentionIndex >> FContentionMatrix::BitsPerWordLog2);
+
+#if DO_CHECK
+			check(WordIdx >= 0 && WordIdx < WordsPerRow);
+#endif
+			const int32 BitOffset = (DirectChunkContentionIndex & (FContentionMatrix::BitsPerWord - 1));
+			SlotBits[WordIdx] &= ~(WordType(1) << BitOffset);
+
+			--SlotRefCounts[SlotIndex];
+#if DO_CHECK
+			check(SlotRefCounts[SlotIndex] >= 0);
+#endif
+
+			if (SlotRefCounts[SlotIndex] == 0)
+			{
+				// release so prior writes are visible before anyone observes Invalid and reuses slot
+				ChunkIDs[SlotIndex].store(InvalidChunkID, std::memory_order_release);
+			}
+
+			UnlockSlot(SlotIndex);
+
+			Chunk.ActiveChunkSlotIndex = INDEX_NONE;
+		}
+
+
+	private:
+		const FContentionMatrix* MatrixRef = nullptr;
+		int32 WordsPerRow = 0;
+
+		alignas(64) std::atomic<uint64> ChunkIDs[MaxActiveChunks];
+		std::atomic_flag SlotLocks[MaxActiveChunks];
+		int32 SlotRefCounts[MaxActiveChunks];
+
+		std::atomic_flag SlotAssignmentLock;
+		TArray<WordType> BitmaskBuffer;
+
+		FORCEINLINE WordType* GetSlotBits(int32 SlotIndex)
+		{
+			return &BitmaskBuffer[SlotIndex * WordsPerRow];
+		}
+
+		FORCEINLINE void LockSlot(int32 SlotIndex)
+		{
+			while (SlotLocks[SlotIndex].test_and_set(std::memory_order_acquire))
+			{
+				FPlatformProcess::YieldCycles(10);
+			}
+		}
+
+		FORCEINLINE void UnlockSlot(int32 SlotIndex)
+		{
+			SlotLocks[SlotIndex].clear(std::memory_order_release);
+		}
+
+		FORCEINLINE void AddJobToBitmask(WordType* SlotBits, int32 JobIndex)
+		{
+			const int32 WordIdx = (JobIndex >> FContentionMatrix::BitsPerWordLog2);
+			const int32 BitOffset = (JobIndex & (FContentionMatrix::BitsPerWord - 1));
+			SlotBits[WordIdx] |= (WordType(1) << BitOffset);
+		}
+	};
+
+	/**
+	 * Uses a pre-calculated FContentionMatrix to help dynamically track active contending items
+	 */
+	struct FActiveContentionTracker
+	{
+	public:
+		using WordType = uint64;
+		static constexpr int32 BitsPerWord = 64;
+		static constexpr int32 BitsPerWordLog2 = 6;
+
+		FActiveContentionTracker(const FContentionMatrix& InMatrix)
+			: Matrix(InMatrix)
+		{
+			Init();
+		}
+
+		void Init()
+		{
+			const int32 Size = Matrix.GetSize();
+			const int32 WordsPerRow = Matrix.GetWordsPerRow();
+
+			ActiveItemBits.SetNumZeroed(WordsPerRow);
+			ActiveContenderCount.SetNumZeroed(Size);
+		}
+
+		FORCEINLINE bool IsActive(int32 ItemIndex)
+		{
+			check(IsValidIndex(ItemIndex));
+			const int32 WordIndex = (ItemIndex >> BitsPerWordLog2);
+			const int32 BitOffset = (ItemIndex & (BitsPerWord-1));
+
+			if (!ActiveItemBits.IsValidIndex(WordIndex))
+			{
+				return false;
+			}
+
+#if defined(__cpp_lib_atomic_ref)
+			std::atomic_ref<WordType> AtomicWord(ActiveItemBits[WordIndex]);
+			return ((AtomicWord.load(std::memory_order_acquire) >> BitOffset) & 1u) != 0;
+#else
+			const int64 SignedValue = FPlatformAtomics::AtomicRead(reinterpret_cast<volatile const int64*>(&ActiveItemBits[WordIndex]));
+			const uint64 Word = static_cast<uint64>(SignedValue);
+			return ((Word >> BitOffset) & 1ull) != 0;
+#endif
+		}
+
+		bool TryActivate(int32 ItemIndex)
+		{
+			check(IsValidIndex(ItemIndex));
+			const int32 WordIndex = (ItemIndex >> BitsPerWordLog2);
+
+			if (HasAnyActiveContention(ItemIndex))
+			{
+				return false;
+			}
+
+			const int32 BitOffset = (ItemIndex & (BitsPerWord - 1));
+			const WordType Mask = (WordType(1) << BitOffset);
+
+#if defined(__cpp_lib_atomic_ref)
+			std::atomic_ref<WordType> AtomicWord(ActiveItemBits[WordIndex]);
+			AtomicWord.fetch_or(Mask, std::memory_order_acq_rel);
+#else
+			FPlatformAtomics::InterlockedOr(reinterpret_cast<volatile int64*>(&ActiveItemBits[WordIndex]), static_cast<int64>(Mask));
+#endif
+
+			// assume success and update contention counts first
+			UpdateContentionCounts(ItemIndex, 1);
+
+			// see if a contender activated while this thread was in UpdateContentionCounts
+			if (HasAnyActiveContention(ItemIndex))
+			{
+				FScopeLock Lock(&ActiveContenderCriticalSection);
+				// check one more time inside the critical section to eliminate all possible races
+				if (HasAnyActiveContention(ItemIndex))
+				{
+					// undo the change to the contention counts
+					UpdateContentionCounts(ItemIndex, -1);
+					return false;
+				}
+			}
+
+			return true;
+		}
+
+		FORCEINLINE bool Deactivate(int32 ItemIndex)
+		{
+			check(IsValidIndex(ItemIndex));
+
+			const int32 WordIndex = (ItemIndex >> BitsPerWordLog2);
+			const int32 BitOffset = (ItemIndex & (BitsPerWord - 1));
+			const WordType Mask = (WordType(1) << BitOffset);
+
+#if defined(__cpp_lib_atomic_ref)
+			std::atomic_ref<WordType> AtomicWord(ActiveItemBits[WordIndex]);
+			const WordType PreviousValue = AtomicWord.fetch_and(~Mask, std::memory_order_acq_rel);
+#else
+			FPlatformAtomics::InterlockedAnd(reinterpret_cast<volatile int64*>(&ActiveItemBits[WordIndex]), ~Mask);
+#endif
+
+			UpdateContentionCounts(ItemIndex, -1);
+			return true;
+		}
+
+		FORCEINLINE bool HasAnyActiveContention(int32 ItemIndex)
+		{
+			check(IsValidIndex(ItemIndex));
+
+			// not all platforms support atomic std::atomic_ref, but it's typically faster so do it when we can
+#if defined(__cpp_lib_atomic_ref)
+			std::atomic_ref<int32> AtomicCount(ActiveContenderCount[ItemIndex]);
+			return AtomicCount.load(std::memory_order_acquire) > 0;
+#else
+			return FPlatformAtomics::AtomicRead(reinterpret_cast<volatile int32*>(&ActiveContenderCount[ItemIndex])) > 0;
+#endif
+		}
+
+		FORCEINLINE int32 GetActiveContentionCount(int32 ItemIndex)
+		{
+			check(IsValidIndex(ItemIndex));
+			// not all platforms support atomic std::atomic_ref, but it's typically faster so do it when we can
+#if defined(__cpp_lib_atomic_ref)
+			std::atomic_ref<int32> AtomicCount(ActiveContenderCount[ItemIndex]);
+			return AtomicCount.load(std::memory_order_acquire);
+#else
+			return FPlatformAtomics::AtomicRead(reinterpret_cast<volatile int32*>(&ActiveContenderCount[ItemIndex]));
+#endif
+		}
+
+	private:
+		const FContentionMatrix& Matrix;
+		TArray<WordType> ActiveItemBits;
+		TArray<int32> ActiveContenderCount;
+
+		FCriticalSection ActiveContenderCriticalSection;
+
+		FORCEINLINE bool IsValidIndex(int32 Index) const
+		{
+			return Index >= 0 && Index < Matrix.GetSize();
+		}
+
+		void UpdateContentionCounts(int32 ItemIndex, int32 Delta)
+		{
+			const WordType* RESTRICT RowWords = Matrix.GetRowWordPointer(ItemIndex);
+			const int32 MatrixSize = Matrix.GetSize();
+			const int32 WordsPerRow = Matrix.GetWordsPerRow();
+			const int32 DiagWordIndex = (ItemIndex >> BitsPerWordLog2);
+			const int32 DiagBitOffset = (ItemIndex & (BitsPerWord-1));
+
+			int32* RawCounts = ActiveContenderCount.GetData();
+
+			for (int32 RowWordIndex = 0, BaseBitIndex = 0; RowWordIndex < WordsPerRow; ++RowWordIndex, BaseBitIndex += BitsPerWord)
+			{
+				WordType ContenderWord = (RowWordIndex != DiagWordIndex)
+					? RowWords[RowWordIndex]
+					: (RowWords[RowWordIndex] & ~(WordType(1) << DiagBitOffset));
+
+				while (ContenderWord != 0)
+				{
+					const int32 TrailingZeroCount = (int32)FMath::CountTrailingZeros64(ContenderWord);
+					const int32 NeighborIndex = BaseBitIndex + TrailingZeroCount;
+
+					if (NeighborIndex < MatrixSize)
+					{
+						// not all platforms support atomic std::atomic_ref, but it's typically faster so do it when we can
+#if defined(__cpp_lib_atomic_ref)
+						std::atomic_ref<int32> NeighborCount(RawCounts[NeighborIndex]);
+						NeighborCount.fetch_add(Delta, std::memory_order_release);
+#else
+						FPlatformAtomics::InterlockedAdd(reinterpret_cast<volatile int32*>(&RawCounts[NeighborIndex]), Delta);
+#endif
+					}
+
+					ContenderWord &= (ContenderWord - 1);
+				}
+			}
+		}
+	};
+
+	/** Tracks which chunks are actively being processed with a refcount */
+	struct FActiveChunkSet
+	{
+	private:
+		struct FChunkSlot
+		{
+			uint64 Index;
+
+		};
+	public:
+		static constexpr int32 Capacity = 32;
+		static constexpr uint64 EmptyValue = 0;
+		const int32 MaxSlots;
+
+		FActiveChunkSet(int32 InMaxSlots)
+			: MaxSlots(InMaxSlots)
+		{
+			checkf(MaxSlots <= Capacity, TEXT("FActiveChunkSet::MaxSlots (%d) exceeded Capacity (%d)."), MaxSlots, Capacity);
+
+			for (std::atomic<uint64>& Slot : Slots)
+			{
+				Slot.store(EmptyValue, std::memory_order_relaxed);
+			}
+			for (std::atomic<int32>& RefCount : RefCounts)
+			{
+				RefCount.store(0, std::memory_order_relaxed);
+			}
+		}
+
+		/**
+		 * Adds an ID to the set.
+		 * Returns true if added, false if it was already present.
+		 */
+		bool TryAddExclusive(uint64 ID)
+		{
+			check(ID != EmptyValue);
+
+			int32 CandidateIndex = INDEX_NONE;
+
+			for (;;)
+			{
+				for (int32 Index = 0; Index < MaxSlots; ++Index)
+				{
+					const uint64 Val = Slots[Index].load(std::memory_order_seq_cst);
+
+					if (Val == ID)
+					{
+						return false; // Already exists
+					}
+
+					if (CandidateIndex == INDEX_NONE && Val == EmptyValue)
+					{
+						CandidateIndex = Index;
+						break;
+					}
+				}
+
+				if (CandidateIndex != INDEX_NONE)
+				{
+					uint64 Expected = EmptyValue;
+					if (Slots[CandidateIndex].compare_exchange_strong(Expected, ID, std::memory_order_seq_cst))
+					{
+						// this is an exclusive add so we can just force refcount to 1 for the previously empty slot
+						RefCounts[CandidateIndex].store(1, std::memory_order_seq_cst);
+						return true;
+					}
+
+					// someone else took this slot or added an item, loop and try again
+				}
+				else
+				{
+					checkf(false, TEXT("FActiveChunkSet overflow! Active chunks exceeded Capacity (%d)."), Capacity);
+					return false;
+				}
+			}
+		}
+
+		/**
+		 * Adds the ID to the set if not present
+		 * Increments the refcount for the ID
+		 * Returns the new refcount
+		 */
+		int32 AddGetRefCount(uint64 ID)
+		{
+			check(ID != EmptyValue);
+
+			int32 CandidateIndex = INDEX_NONE;
+
+			for (;;)
+			{
+				for (int32 Index = 0; Index < MaxSlots; ++Index)
+				{
+					const uint64 Val = Slots[Index].load(std::memory_order_seq_cst);
+
+					if (Val == ID)
+					{
+						int32 NewCount = RefCounts[Index].fetch_add(1, std::memory_order_seq_cst) + 1;
+						return NewCount;
+					}
+
+					if (CandidateIndex == INDEX_NONE && Val == EmptyValue)
+					{
+						CandidateIndex = Index;
+						break;
+					}
+				}
+
+				if (CandidateIndex != INDEX_NONE)
+				{
+					uint64 Expected = EmptyValue;
+					if (Slots[CandidateIndex].compare_exchange_strong(Expected, ID, std::memory_order_seq_cst))
+					{
+						// this is an exclusive add so we can just force refcount to 1 for the previously empty slot
+						RefCounts[CandidateIndex].store(1, std::memory_order_seq_cst);
+						return 1;
+					}
+
+					// someone else took this slot or added an item, loop and try again
+				}
+				else
+				{
+					checkf(false, TEXT("FActiveChunkSet overflow! Active chunks exceeded MaxSlots (%d)."), MaxSlots);
+					return 0;
+				}
+			}
+		}
+
+		/**
+		 * Decrement the refcount for ID and remove it if the count is 0
+		 */
+		void Remove(uint64 ID)
+		{
+			checkSlow(ID != EmptyValue);
+
+			for (int32 Index = 0; Index < MaxSlots; ++Index)
+			{
+				std::atomic<uint64>& Slot = Slots[Index];
+
+				if (Slot.load(std::memory_order_seq_cst) == ID)
+				{
+					int32 NewCount = RefCounts[Index].fetch_add(-1, std::memory_order_seq_cst) - 1;
+
+					if (NewCount == 0)
+					{
+						Slot.store(EmptyValue, std::memory_order_seq_cst);
+						// double check that the refcount is still 0:
+						NewCount = RefCounts[Index].load(std::memory_order_seq_cst);
+						if (NewCount > 0)
+						{
+							// something else added a refcount between our operations
+							// restore the ID and move on
+							Slot.store(ID, std::memory_order_seq_cst);
+						}
+						return;
+					}
+				}
+			}
+		}
+
+		/**
+		 * Checks if the ID exists in the set.
+		 */
+		bool Contains(uint64 ID) const
+		{
+			checkSlow(ID != EmptyValue);
+			for (int32 i = 0; i < MaxSlots; ++i)
+			{
+				if (Slots[i].load(std::memory_order_acquire) == ID)
+				{
+					return true;
+				}
+			}
+			return false;
+		}
+
+	private:
+		// alignas(64) to avoid false sharing
+		alignas(64) std::atomic<uint64> Slots[Capacity];
+		alignas(64) std::atomic<int32> RefCounts[Capacity];
+	};
+
+	/**
+	 * Increments a change-tracking atomic variable when leaving a scope.
+	 * Used with FChangeListener for efficient syncronizaion of workers waiting for resources.
+	 */
+	template <typename T>
+	struct FChangeScope
+	{
+		explicit FChangeScope(std::atomic<T>& InTrackingValue)
+			:Tracker(InTrackingValue)
+		{
+		}
+		FChangeScope(const FChangeScope&) = delete;
+		FChangeScope& operator=(const FChangeScope&) = delete;
+		FChangeScope(FChangeScope&& Other) = delete;
+
+		~FChangeScope()
+		{
+			++Tracker;
+		}
+	private:
+		std::atomic<T>& Tracker;
+	};
+
+	/**
+	 * Waits for an atomic variable to change or become a specific value.
+	 * Used with FChangeScope for efficient syncronizaion of workers waiting for resources.
+	 */
+	template <typename T>
+	struct FChangeListener
+	{
+		explicit FChangeListener(std::atomic<T>& InValueToTrack)
+			:ValueToTrack(InValueToTrack)
+			// Set the initial value at construction time to capture changes that happen between declaration and subsequent Wait calls.
+			, InitialValue(ValueToTrack.load(std::memory_order_acquire))
+		{
+		}
+
+		void WaitForChange()
+		{
+			for (;;)
+			{
+				// @todo: optimize spincount and pause for 1 cycle in the loop
+				for (int32 SpinCount = 0; SpinCount < 128; ++SpinCount)
+				{
+					const T CurrentValue = ValueToTrack.load(std::memory_order_acquire);
+					if (CurrentValue != InitialValue)
+					{
+						// caller might wait for changes again
+						InitialValue = CurrentValue;
+						return;
+					}
+					FPlatformProcess::YieldCycles(1);
+				}
+				FPlatformProcess::YieldThread();
+			}
+		}
+
+		void WaitForValue(T TargetValue)
+		{
+			for (;;)
+			{
+				for (int32 SpinCount = 0; SpinCount < 2000; ++SpinCount)
+				{
+					const T CurrentValue = ValueToTrack.load(std::memory_order_acquire);
+					if (CurrentValue == TargetValue)
+					{
+						// caller might wait for changes again
+						InitialValue = CurrentValue;
+						return;
+					}
+				}
+				FPlatformProcess::YieldThread();
+			}
+		}
+	private:
+		std::atomic<T>& ValueToTrack;
+		T InitialValue;
+	};
+
+	/** Single struct to specify data access for a query */
+	struct FDataAccess
+	{
+		FElementBitSet DirectReadsElements;
+		FElementBitSet DirectWritesElements;
+		FElementBitSet IndirectReadsElements;
+		FElementBitSet IndirectWritesElements;
+		FMassExternalSubsystemBitSet ReadsSubsystems;
+		FMassExternalSubsystemBitSet WritesSubsystems;
+
+		void Add(const FDataAccess& Other);
+		bool SubsystemsCollide(const FDataAccess& Other) const;
+		bool CollidesWithIndirect(const FDataAccess& Other) const;
+		bool CollidesWithDirect(const FDataAccess& Other) const;
+		void SetFragmentWriteRequirements(const FMassFragmentRequirements& FragmentRequirements);
+	};
+
+	/** Stores processing queue-specific data about a FMassEntityQuery */
+	struct FQueryInfo
+	{
+		bool bChunksRequireGameThread = false;
+		
+		/** Contention index to run the query (gather chunks but not execute) */
+		int32 QueryContentionIndex = INDEX_NONE;
+		
+		/** Contention index for cross-chunk contention on each chunk processed */
+		int32 ChunkGlobalContentionIndex = INDEX_NONE;
+		
+		/** Contention index for direct chunk contention on each chunk processed */
+		int32 ChunkContentionIndex = INDEX_NONE;
+
+		FDataAccess DataAccess;
+		const FMassEntityQuery* SourceQuery = nullptr;
+	};
+
+	/** A single processor managed by the processing queue. */
+	struct FQueuedProcessor
+	{
+		enum class EState : uint8
+		{
+			/** Initialized, and waiting for explicit ordering dependencies */
+			WaitingForDependencies,
+
+			/** Explicit dependencies met, ready to activate (will activate when at the front of the queue and no contending items are active) */
+			Ready,
+
+			/** Currently running */
+			Running,
+
+			Completed,
+			MAX
+		};
+		/** Current state in the execution pipeline */
+		EState State = EState::MAX;
+
+		TNotNull<UMassProcessor*> Processor;
+
+		/**
+		 * If the processor or any query requires game thread execution,
+		 * then aggregate that requirement to all queries and the processor itself.
+		 * Useful for emulating the legacy processor behavior.
+		 */
+		bool bAggregateThreadRequirements = true;
+
+		/**
+		 * Aggregate all data access requirements of the processor and queries to all queries and the processor itself
+		 * Useful for emulating the legacy processor behavior.
+		 */
+		bool bAggregateRequirements = true;
+
+		/** Contention table index for this processor (does not include query requirements) */
+		int32 ContentionIndex = INDEX_NONE;
+
+		/** Queries will follow te owning processor in the contention table, count is used to get the correct index */
+		int32 QueryCount = 1;
+
+		/** Explicit dependencies remaining (ExecuteBefore/ExecuteAfter) */
+		std::atomic<int32> RemainingDependencyCount{ 0 };
+
+		/** Total explicit dependency count (ExecuteBefore/ExecuteAfter) */
+		int32 DependencyCount = 0;
+
+		/** Indices of items that depend explicitly on this processor (within the FProcessingQueue's processor array) */
+		TArray<int32, TInlineAllocator<8>> DependentProcessorIndices;
+
+		/**
+		 * The length of the longest dependency chain starting from this processor.
+		 * 0 = Leaf node (nothing depends on this).
+		*/
+		int32 GraphDepth = 0;
+
+		FDataAccess DataAccess;
+
+		TArray<FQueryInfo> QueryInfos;
+		int32 CurrentQueryIndex = INDEX_NONE;
+
+		bool bRequiresGameThread = true;
+
+		void InitThreadRequirement();
+		void InitWriteContention(FContentionMatrix& ContentionMatrix, FQueuedProcessor& Other);
+		void InitWriteContention(FContentionMatrix& ContentionMatrix, int32 MyQueryIndex, FQueuedProcessor& Other);
+		void InitCrossChunkWriteContention(FContentionMatrix& ContentionMatrix, int32 MyQueryIndex, FQueuedProcessor& Other, int32 OtherQueryIndex);
+		void InitChunkWriteContention(FContentionMatrix& ContentionMatrix, int32 MyQueryIndex, FQueuedProcessor& Other, int32 OtherQueryIndex);
+
+		explicit FQueuedProcessor(TNotNull<UMassProcessor*> InProcessor);
+
+		void PrepareForExecute();
+
+		const FQueryInfo& GetCurrentQueryInfo() const;
+		void SetCurrentQuery(const FMassEntityQuery* SourceQuery);
+
+		/** Multithreading access detector to ensure ChunkQueue is not altered concurrently */
+		UE_MT_DECLARE_RW_ACCESS_DETECTOR(ChunkQueueAccessDetector);
+		
+		/** Chunks that could not be executed immediately during the query and were queued. */
+		TArray<FQueuedChunk> ChunkQueue;
+
+		/** List of indices into ChunkQueue that have not been processed yet. */
+		TArray<int32> PendingChunkIndices;
+#if WITH_MASSENTITY_DEBUG
+		FProcessorWorker* DebugCurrentProcessorWorker = nullptr;
+#endif
+
+	};
+
+	/** A worker to execute queued processors */
+	class FProcessorWorker
+	{
+	public:
+		FProcessorWorker(FProcessingQueue& InQueue, const FMassExecutionContext& InExecutionContext, bool bInInlineExecution = false);
+		~FProcessorWorker();
+
+		void AssignProcessor(FQueuedProcessor* QueuedProcessor);
+
+		void QueueChunk(FMassChunkProcessingQueueParams&& Params);
+		const FQueryInfo& GetCurrentQueryInfo() const
+		{
+			return RunningProcessor->GetCurrentQueryInfo();
+		}
+
+		void ExecuteProcessor();
+		void WaitForQueryResources();
+		void ProcessChunks();
+		void Cleanup();
+		void ChunkDone(FQueuedChunk* InChunk);
+		
+		FProcessingQueue& OwningQueue;
+		const FMassExecutionContext& GlobalExecutionContext;
+		bool bInlineExecution = false;
+		FQueuedProcessor* RunningProcessor = nullptr;
+		std::atomic<int32> ChunksRemaining{ 0 };
+		TSharedPtr<FMassCommandBuffer> CommandBuffer;
+		int32 GlobalChunkContentionIndex = INDEX_NONE;
+
+		/** Total number of chunks this processor currently wants to process */
+		std::atomic<int32> DesiredChunkCount{ 0 };
+
+		/** Number of chunk workers explicitly granted to this processor by the Queue */
+		std::atomic<int32> GrantedChunkSlots{ 0 };
+
+		/** Number of slots currently effectively used */
+		std::atomic<int32> ActiveChunkSlots{ 0 };
+
+#if WITH_MASSENTITY_DEBUG
+		FMassExecutionContext* DebugCurrentExecutionContext = nullptr;
+#endif
+	};
+
+	/** A worker to process queued chunks */
+	class FChunkWorker
+	{
+	public:
+
+		FChunkWorker(FProcessingQueue& InQueue, FQueuedChunk* InChunk, bool UseChunkCommandBuffer = false);
+		~FChunkWorker();
+		void ExecuteChunk(bool bInlineExecution);
+		void Cleanup();
+		void AssignChunk(FQueuedChunk* InChunk)
+		{
+			Chunk = InChunk;
+		}
+		TSharedPtr<FMassCommandBuffer>& GetOrCreateCommandBuffer();
+
+		FProcessingQueue& OwningQueue;
+		FQueuedChunk* Chunk = nullptr;
+		bool bUseContextCommandBuffer = false;
+		/** Each worker owns a command buffer that is reused for all the chunks it processes */
+		/** The buffer is created in the thread this worker uses, and registered with the queue to be merged/flushed */
+		TSharedPtr<FMassCommandBuffer> CommandBuffer;
+	};
+} // namespace UE::Mass
+
+
+
+#undef UE_API
diff --git a/Engine/Source/Runtime/MassEntity/Public/MassProcessor.h b/Engine/Source/Runtime/MassEntity/Public/MassProcessor.h
--- a/Engine/Source/Runtime/MassEntity/Public/MassProcessor.h
+++ b/Engine/Source/Runtime/MassEntity/Public/MassProcessor.h
@@ -26,6 +26,8 @@ struct FMassSubsystemRequirements;
 namespace UE::Mass
 {
 	struct FQueryExecutor;
+	struct FProcessingQueue;
+	struct FQueuedProcessor;
 }
 
 enum class EProcessorCompletionStatus : uint8
@@ -130,6 +132,7 @@ public:
 	UE_API virtual void ExportRequirements(FMassExecutionRequirements& OutRequirements) const;
 
 	const FMassSubsystemRequirements& GetProcessorRequirements() const;
+	const FMassEntityCreationRequirements& GetProcessorEntityCreationRequirements() const;
 
 	/**
 	 * @return the current value of ExecutionPriority
@@ -255,6 +258,11 @@ protected:
 	UPROPERTY(EditDefaultsOnly, Category = Processor, config)
 	uint8 bRequiresGameThreadExecution : 1 = false;
 
+	/** Combine all query and processor requirements when resolving contention in FProcessingQueue
+	 *  A value of true is legacy mode, any multi-query processor forces aggregate requirements */
+	UPROPERTY(EditDefaultsOnly, Category = Processor, config)
+	uint8 bAggregateRequirements : 1 = true;
+
 #if WITH_EDITORONLY_DATA
 	/** Used to permanently remove a given processor class from PipeSetting's listing. Used primarily for test-time 
 	 *  processor classes, but can also be used by project-specific code to prune the processor list. */
@@ -315,6 +323,7 @@ protected:
 
 	friend UMassCompositeProcessor;
 	friend FMassDebugger;
+	friend UE::Mass::FQueuedProcessor; // needs access to OwnedQueries
 
 	/** A query representing elements this processor is accessing in Execute function outside of query execution */
 	FMassSubsystemRequirements ProcessorRequirements;
@@ -508,6 +517,11 @@ inline const FMassSubsystemRequirements& UMassProcessor::GetProcessorRequirement
 	return ProcessorRequirements;
 }
 
+inline const FMassEntityCreationRequirements& UMassProcessor::GetProcessorEntityCreationRequirements() const
+{
+	return ProcessorEntityCreationRequirements;
+}
+
 inline int16 UMassProcessor::GetExecutionPriority() const
 {
 	return ExecutionPriority; 
diff --git a/Engine/Source/Runtime/MassEntity/Public/MassRequirements.h b/Engine/Source/Runtime/MassEntity/Public/MassRequirements.h
--- a/Engine/Source/Runtime/MassEntity/Public/MassRequirements.h
+++ b/Engine/Source/Runtime/MassEntity/Public/MassRequirements.h
@@ -699,6 +699,8 @@ struct FMassFragmentRequirements
 	const FMassFragmentBitSet& GetRequiredAnyFragments() const;
 	const FMassFragmentBitSet& GetRequiredOptionalFragments() const;
 	const FMassFragmentBitSet& GetRequiredNoneFragments() const;
+	const FMassFragmentBitSet& GetIndirectReadOnlyFragments() const;
+	const FMassFragmentBitSet& GetIndirectReadWriteFragments() const;
 	const FMassTagBitSet& GetRequiredAllTags() const;
 	const FMassTagBitSet& GetRequiredAnyTags() const;
 	const FMassTagBitSet& GetRequiredNoneTags() const;
@@ -899,6 +901,27 @@ public:
 		return Compositions.Num() == 0;
 	}
 
+	const TArray<FMassElementBitSet>& GetCompositions() const
+	{
+		return Compositions;
+	}
+
+	bool HasAny(const FMassEntityCreationRequirements& Other) const
+	{
+		if (!Compositions.IsEmpty())
+		{
+			for (const FMassElementBitSet& Composition : Compositions)
+			{
+				if (Other.Contains(Composition))
+				{
+					return true;
+				}
+			}
+		}
+		return false;
+	}
+
+
 private:
 	TArray<FMassArchetypeHandle> ArchetypeHandles;
 	TArray<FMassElementBitSet> Compositions;
@@ -1121,6 +1144,16 @@ inline const FMassElementBitSet& FMassFragmentRequirements::GetRequiredNoneSpars
 	return RequiredNoneSparseElements; 
 }
 
+inline const FMassFragmentBitSet& FMassFragmentRequirements::GetIndirectReadOnlyFragments() const
+{
+	return IndirectReadOnlyFragments;
+}
+
+inline const FMassFragmentBitSet& FMassFragmentRequirements::GetIndirectReadWriteFragments() const
+{
+	return IndirectReadWriteFragments;
+}
+
 inline bool FMassFragmentRequirements::IsInitialized() const 
 { 
 	return bInitialized; 
