[Mass] fixed Clang Analyzer errors in MassExecutor.cpp
Issue 1220511

#rnx

#Synced-CL 50764032

diff --git a/Engine/Source/Runtime/MassEntity/Public/MassProcessingContext.h b/Engine/Source/Runtime/MassEntity/Public/MassProcessingContext.h
--- a/Engine/Source/Runtime/MassEntity/Public/MassProcessingContext.h
+++ b/Engine/Source/Runtime/MassEntity/Public/MassProcessingContext.h
@@ -54,7 +54,7 @@ namespace UE::Mass
 		TSharedPtr<FMassCommandBuffer> CommandBuffer;
 
 	protected:
-		uint8 ExecutionContextBuffer[sizeof(FMassExecutionContext)];
+		alignas(FMassExecutionContext) uint8 ExecutionContextBuffer[sizeof(FMassExecutionContext)];
 		FMassExecutionContext* ExecutionContextPtr = nullptr;
 	};
 
