[MassEntity] used transactionally safe thread access detector for archetype chunks

#Synced-CL 50837704

diff --git a/Engine/Source/Runtime/MassEntity/Internal/MassArchetypeData.h b/Engine/Source/Runtime/MassEntity/Internal/MassArchetypeData.h
--- a/Engine/Source/Runtime/MassEntity/Internal/MassArchetypeData.h
+++ b/Engine/Source/Runtime/MassEntity/Internal/MassArchetypeData.h
@@ -5,6 +5,7 @@
 #include "MassEntityHandle.h"
 #include "MassArchetypeTypes.h"
 #include "MassArchetypeGroup.h"
+#include "Misc/MTTransactionallySafeAccessDetector.h"
 #if UE_ENABLE_INCLUDE_ORDER_DEPRECATED_IN_5_6
 #include "MassEntityManager.h"
 #endif // UE_ENABLE_INCLUDE_ORDER_DEPRECATED_IN_5_6
@@ -385,7 +386,7 @@ private:
 	 */
 	UE::Mass::FSparseElementsStorage* SparseElementsStorage = nullptr;
 
-	UE_MT_DECLARE_RW_ACCESS_DETECTOR(ChunkIterationDetector);
+	UE_MT_DECLARE_TS_RW_ACCESS_DETECTOR(ChunkIterationDetector);
 
 	mutable bool bSparseElementsBitSetDirty = false;
 
