[Mass] Additional archetype information retrieval.

The following information can now be retrieved from an archetype:
- list all elements (fragments, tags, etc.).
- list all entities

#Synced-CL 50781450

diff --git a/Engine/Source/Runtime/MassEntity/Internal/MassArchetypeData.h b/Engine/Source/Runtime/MassEntity/Internal/MassArchetypeData.h
--- a/Engine/Source/Runtime/MassEntity/Internal/MassArchetypeData.h
+++ b/Engine/Source/Runtime/MassEntity/Internal/MassArchetypeData.h
@@ -592,6 +592,8 @@ public:
 
 	void ExportEntityHandles(TArray<FMassEntityHandle>& InOutHandles) const;
 
+	void ListEntitiesForEachChunk(TFunctionRef<void(TConstArrayView<FMassEntityHandle>)> Function) const;
+
 	// Converts the list of fragments into a user-readable debug string
 	FString DebugGetDescription() const;
 
diff --git a/Engine/Source/Runtime/MassEntity/Private/MassArchetypeData.cpp b/Engine/Source/Runtime/MassEntity/Private/MassArchetypeData.cpp
--- a/Engine/Source/Runtime/MassEntity/Private/MassArchetypeData.cpp
+++ b/Engine/Source/Runtime/MassEntity/Private/MassArchetypeData.cpp
@@ -1232,6 +1232,17 @@ void FMassArchetypeData::ExportEntityHandles(TArray<FMassEntityHandle>& InOutHan
 	}
 }
 
+void FMassArchetypeData::ListEntitiesForEachChunk(TFunctionRef<void(TConstArrayView<FMassEntityHandle>)> Function) const
+{
+	for (const FMassArchetypeChunk& Chunk : Chunks)
+	{
+		if (int32 ChuckInstanceCount = Chunk.GetNumInstances(); ChuckInstanceCount > 0)
+		{
+			Function(TConstArrayView<FMassEntityHandle>(Chunk.GetEntityArray(EntityListOffsetWithinChunk), ChuckInstanceCount));
+		}
+	}
+}
+
 FString FMassArchetypeData::DebugGetDescription() const
 {
 #if WITH_MASSENTITY_DEBUG
diff --git a/Engine/Source/Runtime/MassEntity/Private/MassEntityManager.cpp b/Engine/Source/Runtime/MassEntity/Private/MassEntityManager.cpp
--- a/Engine/Source/Runtime/MassEntity/Private/MassEntityManager.cpp
+++ b/Engine/Source/Runtime/MassEntity/Private/MassEntityManager.cpp
@@ -553,12 +553,24 @@ void FMassEntityManager::GetMatchingArchetypes(const FMassFragmentRequirements&
 	GetMatchingArchetypes(Requirements, OutValidArchetypes, 0);
 }
 
-void FMassEntityManager::ForEachArchetypeFragmentType(const FMassArchetypeHandle& ArchetypeHandle, TFunction< void(const UScriptStruct* /*FragmentType*/)> Function)
+void FMassEntityManager::ForEachArchetypeElementType(const FMassArchetypeHandle& ArchetypeHandle, TFunctionRef< bool(const UScriptStruct*)> Function)
+{
+	const FMassArchetypeData& ArchetypeData = FMassArchetypeHelper::ArchetypeDataFromHandleChecked(ArchetypeHandle);
+	ArchetypeData.GetCompositionBitSet().ExportTypes(Function);
+}
+
+void FMassEntityManager::ForEachArchetypeFragmentType(const FMassArchetypeHandle& ArchetypeHandle, TFunction< void(const UScriptStruct*)> Function)
 {
 	const FMassArchetypeData& ArchetypeData = FMassArchetypeHelper::ArchetypeDataFromHandleChecked(ArchetypeHandle);
 	ArchetypeData.ForEachFragmentType(Function);
 }
 
+void FMassEntityManager::ForEachArchetypeChunkEntityList(const FMassArchetypeHandle& ArchetypeHandle, TFunctionRef<void(TConstArrayView<FMassEntityHandle>)> Function)
+{
+	const FMassArchetypeData& ArchetypeData = FMassArchetypeHelper::ArchetypeDataFromHandleChecked(ArchetypeHandle);
+	ArchetypeData.ListEntitiesForEachChunk(Function);
+}
+
 void FMassEntityManager::DoEntityCompaction(const double TimeAllowed)
 {
 	int32 TotalEntitiesMoved = 0;
diff --git a/Engine/Source/Runtime/MassEntity/Public/MassEntityManager.h b/Engine/Source/Runtime/MassEntity/Public/MassEntityManager.h
--- a/Engine/Source/Runtime/MassEntity/Public/MassEntityManager.h
+++ b/Engine/Source/Runtime/MassEntity/Public/MassEntityManager.h
@@ -237,9 +237,18 @@ public:
 	 */
 	UE_API void GetMatchingArchetypes(const FMassFragmentRequirements& Requirements, TArray<FMassArchetypeHandle>& OutValidArchetypes) const;
 
+	/** 
+	 * Method to iterate on all the element (fragments, tags, etc.) types of an archetype.
+	 * Note that sparse elements are not archetype characteristics and therefore not included in this call.
+	 */
+	static UE_API void ForEachArchetypeElementType(const FMassArchetypeHandle& ArchetypeHandle, TFunctionRef< bool(const UScriptStruct* /*ElementType*/)> Function);
+
 	/** Method to iterate on all the fragment types of an archetype */
 	static UE_API void ForEachArchetypeFragmentType(const FMassArchetypeHandle& ArchetypeHandle, TFunction< void(const UScriptStruct* /*FragmentType*/)> Function);
 
+	/** List all entities in an archetype per chunk. */
+	static UE_API void ForEachArchetypeChunkEntityList(const FMassArchetypeHandle& ArchetypeHandle, TFunctionRef<void(TConstArrayView<FMassEntityHandle> /*Entities*/)> Function);
+
 	/**
 	 * Go through all archetypes and compact entities
 	 * @param TimeAllowed to do entity compaction, once it reach that time it will stop and return
