[MassEntity] added a few methods to add/remove tags/fragments

#rnx

#Synced-CL 50939189

diff --git a/Engine/Source/Runtime/MassEntity/Public/MassEntityTypes.h b/Engine/Source/Runtime/MassEntity/Public/MassEntityTypes.h
--- a/Engine/Source/Runtime/MassEntity/Public/MassEntityTypes.h
+++ b/Engine/Source/Runtime/MassEntity/Public/MassEntityTypes.h
@@ -219,6 +219,11 @@ struct FMassArchetypeCompositionDescriptor
 
 	void Add(TNotNull<const UScriptStruct*> ElementType);
 
+	template<typename T>
+	void Remove();
+
+	void Remove(TNotNull<const UScriptStruct*> ElementType);
+
 	FMassFragmentBitSet DebugGetFragments() const;
 	FMassTagBitSet DebugGetTags() const;
 	FMassChunkFragmentBitSet DebugGetChunkFragments() const;
@@ -952,6 +957,17 @@ inline void FMassArchetypeCompositionDescriptor::Add(TNotNull<const UScriptStruc
 	return ElementsBitSet.Add(*ElementType);
 }
 
+template<typename T>
+void FMassArchetypeCompositionDescriptor::Remove()
+{
+	return ElementsBitSet.Remove(*T::StaticStruct());
+}
+
+inline void FMassArchetypeCompositionDescriptor::Remove(TNotNull<const UScriptStruct*> ElementType)
+{
+	return ElementsBitSet.Remove(*ElementType);
+}
+
 inline FMassFragmentBitSet FMassArchetypeCompositionDescriptor::DebugGetFragments() const 
 { 
 	return ElementsBitSet.Get<FMassFragmentBitSet>();
