[Mass] Added move-semantics version of FMassEntityQuery::ForEachEntityChunk

#rnx

diff --git a/Engine/Source/Runtime/MassEntity/Private/MassEntityQuery.cpp b/Engine/Source/Runtime/MassEntity/Private/MassEntityQuery.cpp
--- a/Engine/Source/Runtime/MassEntity/Private/MassEntityQuery.cpp
+++ b/Engine/Source/Runtime/MassEntity/Private/MassEntityQuery.cpp
@@ -242,12 +242,19 @@ void FMassEntityQuery::ForEachEntityChunkInCollections(TConstArrayView<FMassArch
 void FMassEntityQuery::ForEachEntityChunk(const FMassArchetypeEntityCollection& EntityCollection, FMassExecutionContext& ExecutionContext
 	, const FMassExecuteFunction& ExecuteFunction)
 {
-	// mz@todo I don't like that we're copying data here.
 	ExecutionContext.SetEntityCollection(EntityCollection);
 	ForEachEntityChunk(ExecutionContext, ExecuteFunction);
 	ExecutionContext.ClearEntityCollection();
 }
 
+void FMassEntityQuery::ForEachEntityChunk(FMassArchetypeEntityCollection&& EntityCollection, FMassExecutionContext& ExecutionContext
+	, const FMassExecuteFunction& ExecuteFunction)
+{
+	ExecutionContext.SetEntityCollection(Forward<FMassArchetypeEntityCollection>(EntityCollection));
+	ForEachEntityChunk(ExecutionContext, ExecuteFunction);
+	ExecutionContext.ClearEntityCollection();
+}
+
 void FMassEntityQuery::ForEachEntityChunk(FMassExecutionContext& ExecutionContext, const FMassExecuteFunction& ExecuteFunction)
 {
 	UE_TRACE_SCOPED_MASS_QUERY_FOR_EACH()
diff --git a/Engine/Source/Runtime/MassEntity/Public/MassEntityQuery.h b/Engine/Source/Runtime/MassEntity/Public/MassEntityQuery.h
--- a/Engine/Source/Runtime/MassEntity/Public/MassEntityQuery.h
+++ b/Engine/Source/Runtime/MassEntity/Public/MassEntityQuery.h
@@ -97,10 +97,16 @@ public:
 	 */
 	UE_API void ForEachEntityChunk(FMassExecutionContext& ExecutionContext, UE::Mass::FExecutionLimiter& Limiter, const FMassExecuteFunction& ExecuteFunction);
 	
-	/** Will first verify that the archetype given with Collection matches the query's requirements, and if so will run the other, more generic ForEachEntityChunk implementation */
+	/**
+	 * Executes provided ExecuteFunction for entities in EntityCollection. If EntityCollection's archetype
+	 * doesn't match the query's requirements then ExecuteFunction won't be executed.
+	 */
 	UE_API void ForEachEntityChunk(const FMassArchetypeEntityCollection& EntityCollection
 		, FMassExecutionContext& ExecutionContext, const FMassExecuteFunction& ExecuteFunction);
 
+	UE_API void ForEachEntityChunk(FMassArchetypeEntityCollection&& EntityCollection
+		, FMassExecutionContext& ExecutionContext, const FMassExecuteFunction& ExecuteFunction);
+
 	/**
 	 * Attempts to process every chunk of every affected archetype in parallel.
 	 */
