[MassEntity] add observer support for shared fragments and swaptag

#Synced-CL 50989336

diff --git a/Engine/Source/Runtime/MassEntity/Private/MassEntityManager.cpp b/Engine/Source/Runtime/MassEntity/Private/MassEntityManager.cpp
--- a/Engine/Source/Runtime/MassEntity/Private/MassEntityManager.cpp
+++ b/Engine/Source/Runtime/MassEntity/Private/MassEntityManager.cpp
@@ -1102,9 +1102,7 @@ void FMassEntityManager::AddFragmentToEntity(FMassEntityHandle Entity, const USc
 	void* FragmentData = CurrentArchetype->GetFragmentDataForEntity(FragmentType, Entity.Index);
 	Initializer(FragmentData, *FragmentType);
 
-	const FMassArchetypeCompositionDescriptor Descriptor(MoveTemp(Fragments));
-	
-	ObserverManager.OnPostCompositionAdded(Entity, Descriptor);
+	ObserverManager.OnPostCompositionAdded(Entity, FMassElementBitSet(MoveTemp(Fragments)));
 }
 
 void FMassEntityManager::AddFragmentListToEntity(FMassEntityHandle Entity, TConstArrayView<const UScriptStruct*> FragmentList)
@@ -1295,7 +1293,7 @@ void FMassEntityManager::RemoveFragmentListFromEntity(FMassEntityHandle Entity,
 	CHECK_SYNC_API();
 
 	CheckIfEntityIsActive(Entity);
-	
+
 	FMassArchetypeData* OldArchetype = GetEntityStorageInterface().GetArchetype(Entity.Index);
 	check(OldArchetype);
 
@@ -1307,15 +1305,14 @@ void FMassEntityManager::RemoveFragmentListFromEntity(FMassEntityHandle Entity,
 		const FMassArchetypeCompositionDescriptor NewComposition(OldArchetype->GetCompositionBitSet() - FragmentsToRemove);
 		const FMassArchetypeHandle NewArchetypeHandle = CreateArchetype(NewComposition, FMassArchetypeCreationParams(*OldArchetype));
 
-		// Find overlap.  It isn't guaranteed that the old archetype has all of the fragments being removed.
-		FMassArchetypeCompositionDescriptor CompositionDelta(OldArchetype->GetCompositionBitSet().GetOverlap(FragmentsToRemove));
-		ObserverManager.OnPreCompositionRemoved(Entity, CompositionDelta);
+		// Find overlap.  It isn't guaranteed that the old archetype has all the fragments being removed.
+		ObserverManager.OnPreCompositionRemoved(Entity, FMassElementBitSet(OldArchetype->GetCompositionBitSet().GetOverlap(FragmentsToRemove)));
 
 		// Move the entity over
 		FMassArchetypeData& NewArchetype = FMassArchetypeHelper::ArchetypeDataFromHandleChecked(NewArchetypeHandle);
 		NewArchetype.CopyDebugNamesFrom(*OldArchetype);
 		OldArchetype->MoveEntityToAnotherArchetype(Entity, NewArchetype);
-		
+
 		GetEntityStorageInterface().SetArchetypeFromShared(Entity.Index, NewArchetypeHandle.DataPtr);
 	}
 }
@@ -1342,16 +1339,21 @@ void FMassEntityManager::SwapTagsForEntity(FMassEntityHandle Entity, TNotNull<co
 	FMassElementBitSet NewCompositionBitSet = CurrentArchetype->GetCompositionBitSet();
 	NewCompositionBitSet.Remove(*OldTagType);
 	NewCompositionBitSet.Add(*NewTagType);
-	
+
 	ensureMsgf(NewCompositionBitSet.IsEquivalent(CurrentArchetype->GetCompositionBitSet()) == false, TEXT("Failing this test at this point is unexpected."));
 
 	const FMassArchetypeHandle NewArchetypeHandle = InternalCreateSimilarArchetype(CurrentArchetype, MoveTemp(NewCompositionBitSet));
 	checkSlow(NewArchetypeHandle.IsValid());
 
+	ObserverManager.OnPreCompositionRemoved(Entity, FMassElementBitSet(*OldTagType));
+
 	// Move the entity over
 	CurrentArchetype->MoveEntityToAnotherArchetype(Entity, *NewArchetypeHandle.DataPtr.Get());
 
+	// Change the entity archetype
 	GetEntityStorageInterface().SetArchetypeFromShared(Entity.Index, NewArchetypeHandle.DataPtr);
+
+	ObserverManager.OnPostCompositionAdded(Entity, FMassElementBitSet(*NewTagType));
 }
 
 void FMassEntityManager::AddTagToEntity(FMassEntityHandle Entity, TNotNull<const UScriptStruct*> TagType)
@@ -1376,8 +1378,8 @@ void FMassEntityManager::AddTagToEntity(FMassEntityHandle Entity, TNotNull<const
 			// Move the entity over
 			CurrentArchetype->MoveEntityToAnotherArchetype(Entity, *NewArchetypeHandle.DataPtr.Get());
 			GetEntityStorageInterface().SetArchetypeFromShared(Entity.Index, NewArchetypeHandle.DataPtr);
-		
-			ObserverManager.OnPostCompositionAdded(Entity, FMassArchetypeCompositionDescriptor(FMassTagBitSet(*TagType)));
+
+			ObserverManager.OnPostCompositionAdded(Entity, FMassElementBitSet(*TagType));
 		}
 	}
 	else
@@ -1403,7 +1405,7 @@ void FMassEntityManager::RemoveTagFromEntity(FMassEntityHandle Entity, TNotNull<
 			FMassElementBitSet TagDelta(*TagType);
 			FMassElementBitSet NewComposition = CurrentArchetype->GetCompositionBitSet() - TagDelta;
 
-			ObserverManager.OnPreCompositionRemoved(Entity, FMassArchetypeCompositionDescriptor(MoveTemp(TagDelta)));
+			ObserverManager.OnPreCompositionRemoved(Entity, MoveTemp(TagDelta));
 
 			const FMassArchetypeHandle NewArchetypeHandle = InternalCreateSimilarArchetype(CurrentArchetype, MoveTemp(NewComposition));
 			checkSlow(NewArchetypeHandle.IsValid());
@@ -1600,11 +1602,14 @@ bool FMassEntityManager::AddConstSharedFragmentToEntity(const FMassEntityHandle
 	NewSharedFragmentValues.Add(InConstSharedFragment);
 	NewSharedFragmentValues.Sort();
 
+	// Move the entity over
 	CurrentArchetype->MoveEntityToAnotherArchetype(Entity, *NewArchetype, &NewSharedFragmentValues);
 
 	// Change the entity archetype
 	GetEntityStorageInterface().SetArchetypeFromShared(Entity.Index, NewArchetypeHandle.DataPtr);
 
+	ObserverManager.OnPostCompositionAdded(Entity, FMassElementBitSet(*StructType));
+
 	return true;
 }
 
@@ -1642,7 +1647,10 @@ bool FMassEntityManager::RemoveConstSharedFragmentFromEntity(const FMassEntityHa
 	const FMassConstSharedFragmentBitSet ToRemove(ConstSharedFragmentType);
 	NewSharedFragmentValues.Remove(ToRemove);
 	NewSharedFragmentValues.Sort();
-	
+
+	ObserverManager.OnPreCompositionRemoved(Entity, FMassElementBitSet(ToRemove));
+
+	// Move the entity over
 	CurrentArchetype->MoveEntityToAnotherArchetype(Entity, *NewArchetype, &NewSharedFragmentValues);
 
 	// Change the entity archetype
@@ -1693,11 +1701,14 @@ bool FMassEntityManager::AddSharedFragmentToEntity(const FMassEntityHandle Entit
 	NewSharedFragmentValues.Add(InSharedFragment);
 	NewSharedFragmentValues.Sort();
 
+	// Move the entity over
 	CurrentArchetype->MoveEntityToAnotherArchetype(Entity, *NewArchetype, &NewSharedFragmentValues);
 
 	// Change the entity archetype
 	GetEntityStorageInterface().SetArchetypeFromShared(Entity.Index, NewArchetypeHandle.DataPtr);
 
+	ObserverManager.OnPostCompositionAdded(Entity, FMassElementBitSet(*StructType));
+
 	return true;
 }
 
@@ -1727,20 +1738,23 @@ bool FMassEntityManager::RemoveSharedFragmentFromEntity(const FMassEntityHandle
 	check(NewArchetypeHandle.IsValid());
 	FMassArchetypeData* NewArchetype = NewArchetypeHandle.DataPtr.Get();
 	check(NewArchetype);
-	
+
 	const FMassArchetypeSharedFragmentValues& OldSharedFragmentValues = CurrentArchetype->GetSharedFragmentValues(Entity.Index);
 	check(OldSharedFragmentValues.ContainsType(&SharedFragmentType));
 	FMassArchetypeSharedFragmentValues NewSharedFragmentValues(OldSharedFragmentValues);
-	
+
 	const FMassSharedFragmentBitSet ToRemove(SharedFragmentType);
 	NewSharedFragmentValues.Remove(ToRemove);
 	NewSharedFragmentValues.Sort();
-	
+
+	ObserverManager.OnPreCompositionRemoved(Entity, FMassElementBitSet(ToRemove));
+
+	// Move the entity over
 	CurrentArchetype->MoveEntityToAnotherArchetype(Entity, *NewArchetype, &NewSharedFragmentValues);
 
 	// Change the entity archetype
 	GetEntityStorageInterface().SetArchetypeFromShared(Entity.Index, NewArchetypeHandle.DataPtr);
-	
+
 	return true;
 }
 
