[Mass] Modified FMassArchetypeSharedFragmentValues's hash calculation to skip one scoped lock (change makes it unnecessary).

#rnx

#Synced-CL 51046751

diff --git a/Engine/Source/Runtime/MassEntity/Public/MassEntityTypes.h b/Engine/Source/Runtime/MassEntity/Public/MassEntityTypes.h
--- a/Engine/Source/Runtime/MassEntity/Public/MassEntityTypes.h
+++ b/Engine/Source/Runtime/MassEntity/Public/MassEntityTypes.h
@@ -585,29 +585,27 @@ struct FMassArchetypeSharedFragmentValues
 		bSorted = (SharedFragments.Num() + ConstSharedFragments.Num() <= 1) ;
 	}
 
-	inline void CacheHash() const
+	/** Returns the cached-or-calculated hash */
+	inline uint32 CacheHash() const
 	{
 		{
 			UE::TReadScopeLock ReadLock{ HashLock };
 			if (HashCache != UINT32_MAX)
 			{
-				return;
+				return HashCache;
 			}
 		}
 
 		{
 			UE::TWriteScopeLock WriteLock{ HashLock };
 			HashCache = CalculateHash();
+			return HashCache;
 		}
 	}
 
 	friend inline uint32 GetTypeHash(const FMassArchetypeSharedFragmentValues& SharedFragmentValues)
 	{
-		SharedFragmentValues.CacheHash();
-		{
-			UE::TReadScopeLock ReadLock{ SharedFragmentValues.HashLock };
-			return SharedFragmentValues.HashCache;
-		}
+		return SharedFragmentValues.CacheHash();
 	}
 
 	MASSENTITY_API uint32 CalculateHash() const;
