[Mass] Sparse Elements taken into consideration when processor dependencies get calculated.
Caveat: all sparse element requirements are treated as ReadWrite

#Synced-CL 51046751
#Synced-CL 51046751

diff --git a/Engine/Source/Runtime/MassEntity/Private/MassProcessorDependencySolver.cpp b/Engine/Source/Runtime/MassEntity/Private/MassProcessorDependencySolver.cpp
--- a/Engine/Source/Runtime/MassEntity/Private/MassProcessorDependencySolver.cpp
+++ b/Engine/Source/Runtime/MassEntity/Private/MassProcessorDependencySolver.cpp
@@ -141,6 +141,7 @@ void FMassExecutionRequirements::Append(const FMassExecutionRequirements& Other)
 		ChunkFragments[i] += Other.ChunkFragments[i];
 		SharedFragments[i] += Other.SharedFragments[i];
 		RequiredSubsystems[i] += Other.RequiredSubsystems[i];
+		SparseElements[i] += Other.SparseElements[i];
 	}
 	ConstSharedFragments.Read += Other.ConstSharedFragments.Read;
 
@@ -164,6 +165,7 @@ void FMassExecutionRequirements::CountResourcesUsed()
 		ResourcesUsedCount += ChunkFragments[i].CountStoredTypes();
 		ResourcesUsedCount += SharedFragments[i].CountStoredTypes();
 		ResourcesUsedCount += RequiredSubsystems[i].CountStoredTypes();
+		ResourcesUsedCount += SparseElements[i].CountStoredTypes();
 	}
 }
 
@@ -179,6 +181,7 @@ bool FMassExecutionRequirements::IsEmpty() const
 {
 	return Fragments.IsEmpty() && IndirectFragments.IsEmpty() && ChunkFragments.IsEmpty()
 		&& SharedFragments.IsEmpty() && ConstSharedFragments.IsEmpty() && RequiredSubsystems.IsEmpty()
+		&& SparseElements.IsEmpty()
 		&& RequiredAllTags.IsEmpty() && RequiredAnyTags.IsEmpty() && RequiredNoneTags.IsEmpty();
 }
 
@@ -204,6 +207,7 @@ FMassProcessorDependencySolver::FResourceUsage::FResourceUsage(const TArray<FNod
 		ChunkFragmentsAccess[i].Access.AddZeroed(FMassChunkFragmentBitSet::GetMaxNum());
 		SharedFragmentsAccess[i].Access.AddZeroed(FMassSharedFragmentBitSet::GetMaxNum());
 		RequiredSubsystemsAccess[i].Access.AddZeroed(FMassExternalSubsystemBitSet::GetMaxNum());
+		SparseElementsAccess[i].Access.AddZeroed(FMassElementBitSet::GetMaxNum());
 	}
 }
 
@@ -363,8 +367,11 @@ bool FMassProcessorDependencySolver::FResourceUsage::CanAccessRequirements(
 	const bool bSubsystemsOK =
 		CanAccess<FMassExternalSubsystemBitSet>(Requirements.RequiredSubsystems, Tested.RequiredSubsystems);
 
+	const bool bSparseElementsOK =
+		CanAccess<FMassElementBitSet>(Requirements.SparseElements, Tested.SparseElements);
+
 	return bDirectOK && bIndirectOK && bIndirectVsDirectOK && bDirectVsIndirectOK
-		&& bChunkOK && bSharedOK && bSubsystemsOK;
+		&& bChunkOK && bSharedOK && bSubsystemsOK && bSparseElementsOK;
 }
 
 
@@ -385,6 +392,8 @@ void FMassProcessorDependencySolver::FResourceUsage::SubmitNode(const int32 Node
 	HandleElementType<FMassSharedFragmentBitSet>(SharedFragmentsAccess, InOutNode.Requirements.SharedFragments, InOutNode, NodeIndex);
 	// Subsystems:
 	HandleElementType<FMassExternalSubsystemBitSet>(RequiredSubsystemsAccess, InOutNode.Requirements.RequiredSubsystems, InOutNode, NodeIndex);
+	// Sparse Elements (archetype-agnostic like indirect fragments):
+	HandleElementType<FMassElementBitSet>(SparseElementsAccess, InOutNode.Requirements.SparseElements, InOutNode, NodeIndex, /*bSkipArchetypeOverlapTest*/ true);
 
 	// note that on purpose we're not pushing ConstSharedFragments - those are always only read, no danger of conflicting access
 	// so there's no point in tracking them
diff --git a/Engine/Source/Runtime/MassEntity/Private/MassRequirements.cpp b/Engine/Source/Runtime/MassEntity/Private/MassRequirements.cpp
--- a/Engine/Source/Runtime/MassEntity/Private/MassRequirements.cpp
+++ b/Engine/Source/Runtime/MassEntity/Private/MassRequirements.cpp
@@ -315,6 +315,11 @@ void FMassFragmentRequirements::ExportRequirements(FMassExecutionRequirements& O
 	OutRequirements.RequiredAnyTags += RequiredAnyTags;
 	OutRequirements.RequiredNoneTags += RequiredNoneTags;
 	// not exporting optional tags by design
+
+	// we treat all sparse elements access as ReadWrite, at the moment. 
+	OutRequirements.SparseElements.Write += RequiredAllSparseElements;
+	OutRequirements.SparseElements.Write += RequiredAnySparseElements;
+	OutRequirements.SparseElements.Write += RequiredOptionalSparseElements;
 }
 
 void FMassFragmentRequirements::Reset()
diff --git a/Engine/Source/Runtime/MassEntity/Public/MassProcessorDependencySolver.h b/Engine/Source/Runtime/MassEntity/Public/MassProcessorDependencySolver.h
--- a/Engine/Source/Runtime/MassEntity/Public/MassProcessorDependencySolver.h
+++ b/Engine/Source/Runtime/MassEntity/Public/MassProcessorDependencySolver.h
@@ -31,13 +31,13 @@ struct TMassExecutionAccess
 
 	T& operator[](const uint32 OpIndex)
 	{
-		check(OpIndex <= EMassAccessOperation::MAX);
+		check(OpIndex < EMassAccessOperation::MAX);
 		return OpIndex == EMassAccessOperation::Read ? Read : Write;
 	}
 
 	const T& operator[](const uint32 OpIndex) const
 	{
-		check(OpIndex <= EMassAccessOperation::MAX);
+		check(OpIndex < EMassAccessOperation::MAX);
 		return OpIndex == EMassAccessOperation::Read ? Read : Write;
 	}
 
@@ -78,6 +78,8 @@ struct FMassExecutionRequirements
 	TMassExecutionAccess<FMassSharedFragmentBitSet> SharedFragments;
 	TMassExecutionAccess<FMassConstSharedFragmentBitSet> ConstSharedFragments;
 	TMassExecutionAccess<FMassExternalSubsystemBitSet> RequiredSubsystems;
+	/** Sparse elements - treated like indirect fragments (archetype-agnostic access) */
+	TMassExecutionAccess<FMassElementBitSet> SparseElements;
 	FMassTagBitSet RequiredAllTags;
 	FMassTagBitSet RequiredAnyTags;
 	FMassTagBitSet RequiredNoneTags;
@@ -161,6 +163,7 @@ private:
 		TMassExecutionAccess<FResourceAccess> ChunkFragmentsAccess;
 		TMassExecutionAccess<FResourceAccess> SharedFragmentsAccess;
 		TMassExecutionAccess<FResourceAccess> RequiredSubsystemsAccess;
+		TMassExecutionAccess<FResourceAccess> SparseElementsAccess;
 		TConstArrayView<FNode> AllNodesView;
 
 		template<typename TBitSet>
