[Mass] Fixed FMassArchetypeData::SetSharedFragmentsData to work as advertised. Also made it support both Const and non-const Shared Fragments.

#rnx

#Synced-CL 51046751

diff --git a/Engine/Source/Runtime/MassEntity/Internal/MassArchetypeData.h b/Engine/Source/Runtime/MassEntity/Internal/MassArchetypeData.h
--- a/Engine/Source/Runtime/MassEntity/Internal/MassArchetypeData.h
+++ b/Engine/Source/Runtime/MassEntity/Internal/MassArchetypeData.h
@@ -710,7 +710,39 @@ public:
 	/** Works with all sparse element types (fragments and tags). */
 	void BatchRemoveSparseElementFromEntities(const FMassArchetypeEntityCollection& EntityCollection, TNotNull<const UScriptStruct*> ElementType);
 
+	template<typename TSharedStruct>
+	requires std::is_same_v<typename TRemoveReference<TSharedStruct>::Type, FSharedStruct> || std::is_same_v<typename TRemoveReference<TSharedStruct>::Type, FConstSharedStruct>
+	void SetSharedFragmentsData(const FMassEntityHandle Entity, const TArray<TSharedStruct>& SharedFragmentValueOverrides)
+	{
+		SetSharedFragmentsData<TSharedStruct>(Entity, MakeArrayView(SharedFragmentValueOverrides));
+	}
+
+	/**
+	 * The function first creates new FMassArchetypeSharedFragmentValues instance combining existing values
+	 * and the contents of SharedFragmentValueOverrides. Then that is used to find the target chunk for Entity,
+	 * and if one cannot be found a new one will be created.
+	 * Note that if SharedFragmentValueOverrides contains types not already present in archetype's composition,
+	 * then those elements will be ignored (there will be an ensure failing).
+	 * @param SharedFragmentValueOverrides is expected to contain only instance of types already
+	 *    present in given archetypes FMassArchetypeSharedFragmentValues
+	 */
+	template<typename TSharedStruct>
+	requires std::is_same_v<TSharedStruct, FSharedStruct> || std::is_same_v<TSharedStruct, FConstSharedStruct>
+	void SetSharedFragmentsData(const FMassEntityHandle Entity, TArrayView<TSharedStruct> SharedFragmentValueOverrides)
+	{
+		FMassEntityInChunkDataHandle EntityInChunkHandle = MakeEntityHandle(Entity.Index);
+
+		// Gets or adds a new chunk that will hold the new entity with the new shared values
+		FMassArchetypeSharedFragmentValues NewSharedFragmentValues(Chunks[EntityInChunkHandle.ChunkIndex].GetSharedFragmentValues());
+		NewSharedFragmentValues.ReplaceSharedFragments<TSharedStruct>(SharedFragmentValueOverrides);
+		NewSharedFragmentValues.Sort();
+
+		SetSharedFragmentsData(Entity, EntityInChunkHandle, MoveTemp(NewSharedFragmentValues));
+	}
+
 protected:
+	MASSENTITY_API void SetSharedFragmentsData(FMassEntityHandle Entity, FMassEntityInChunkDataHandle EntityInChunkHandle, FMassArchetypeSharedFragmentValues&& SharedFragmentValues);
+
 	FMassArchetypeEntityCollection::FArchetypeEntityRange PrepareNextEntitiesSpanInternal(TConstArrayView<FMassEntityHandle> Entities, const FMassArchetypeSharedFragmentValues& InSharedFragmentValues, const int32 StartingChunk = 0);
 	void BatchRemoveEntitiesInternal(const int32 ChunkIndex, const int32 StartIndexWithinChunk, const int32 NumberToRemove);
 
@@ -735,15 +767,6 @@ protected:
 	void BindConstSharedFragmentRequirements(FMassExecutionContext& RunContext, const FMassArchetypeSharedFragmentValues& SharedFragmentValues, const FMassFragmentIndicesMapping& ChunkFragmentsMapping);
 	void BindSharedFragmentRequirements(FMassExecutionContext& RunContext, FMassArchetypeSharedFragmentValues& SharedFragmentValues, const FMassFragmentIndicesMapping& ChunkFragmentsMapping);
 
-	/**
-	 * The function first creates new FMassArchetypeSharedFragmentValues instance combining existing values
-	 * and the contents of SharedFragmentValueOverrides. Then that is used to find the target chunk for Entity,
-	 * and if one cannot be found a new one will be created. 
-	 * @param SharedFragmentValueOverrides is expected to contain only instance of types already
-	 *    present in given archetypes FMassArchetypeSharedFragmentValues
-	 */
-	void SetSharedFragmentsData(const FMassEntityHandle Entity, TConstArrayView<FSharedStruct> SharedFragmentValueOverrides);
-
 	FMassArchetypeChunk& GetOrAddChunk(const FMassArchetypeSharedFragmentValues& SharedFragmentValues, int32& OutAbsoluteIndex, int32& OutIndexWithinChunk);
 	
 private:
diff --git a/Engine/Source/Runtime/MassEntity/Private/MassArchetypeData.cpp b/Engine/Source/Runtime/MassEntity/Private/MassArchetypeData.cpp
--- a/Engine/Source/Runtime/MassEntity/Private/MassArchetypeData.cpp
+++ b/Engine/Source/Runtime/MassEntity/Private/MassArchetypeData.cpp
@@ -675,22 +675,12 @@ bool FMassArchetypeData::DoesEntityMatchSparseComposition(const FMassEntityInChu
 	return RequiredAllSparseElements.IsEmpty() && RequiredAnySparseElements.IsEmpty();
 }
 
-void FMassArchetypeData::SetSharedFragmentsData(const FMassEntityHandle Entity, TConstArrayView<FSharedStruct> SharedFragmentValueOverrides)
+void FMassArchetypeData::SetSharedFragmentsData(const FMassEntityHandle Entity, const FMassEntityInChunkDataHandle EntityInChunkHandle, FMassArchetypeSharedFragmentValues&& SharedFragmentValues)
 {
-	// Gets the current chunk where the entity is located
-	const int32 OldAbsoluteIndex = EntityMap.FindChecked(Entity.Index);
-	const int32 OldChunkIndex = OldAbsoluteIndex / NumEntitiesPerChunk;
-	const int32 OldIndexWithinChunk = OldAbsoluteIndex - (NumEntitiesPerChunk * OldChunkIndex);
-	FMassArchetypeChunk& OldChunk = Chunks[OldChunkIndex];
-
-	// Gets or adds a new chunk that will hold the new entity with the new shared values
-	FMassArchetypeSharedFragmentValues NewSharedFragmentValues(OldChunk.GetSharedFragmentValues());
-	NewSharedFragmentValues.ReplaceSharedFragments(SharedFragmentValueOverrides);
-	NewSharedFragmentValues.Sort();
-
 	int32 NewAbsoluteIndex = 0;
 	int32 NewIndexWithinChunk = 0;
-	FMassArchetypeChunk& NewChunk = GetOrAddChunk(NewSharedFragmentValues, NewAbsoluteIndex, NewIndexWithinChunk);
+	FMassArchetypeChunk& NewChunk = GetOrAddChunk(SharedFragmentValues, NewAbsoluteIndex, NewIndexWithinChunk);
+	FMassArchetypeChunk& OldChunk = Chunks[EntityInChunkHandle.ChunkIndex];
 
 	if (ensureMsgf(&NewChunk != &OldChunk, TEXT("Found target chunk is the same as the source chunk. Probably "
 		"caused by setting shared fragment values resulted in no change, meaning the target values equal the source values")))
@@ -702,18 +692,20 @@ void FMassArchetypeData::SetSharedFragmentsData(const FMassEntityHandle Entity,
 		NewChunk.GetEntityArrayElementRef(EntityListOffsetWithinChunk, NewIndexWithinChunk) = Entity;
 		
 		// Move the current entity fragments into the new chunk
-		MoveFragmentsToNewLocationInternal({ OldChunk.GetRawMemory(), OldIndexWithinChunk }, { NewChunk.GetRawMemory(), NewIndexWithinChunk }, 1);
+		MoveFragmentsToNewLocationInternal({ NewChunk.GetRawMemory(), NewIndexWithinChunk }, { EntityInChunkHandle.ChunkRawMemory, EntityInChunkHandle.IndexWithinChunk }, 1);
 
 		if (OldChunk.HasSparseElements())
 		{
 			UE::Mass::FChunkSparseElements& OldSparseElements = OldChunk.GetSparseElementsUnsafe();
 			UE::Mass::FChunkSparseElements& NewSparseElements = NewChunk.GetSparseElements(NumEntitiesPerChunk);
-			const bool bElementsPresentDirty = NewSparseElements.MoveElementsFrom(OldSparseElements, OldIndexWithinChunk, NewIndexWithinChunk, 1);
-			bSparseElementsBitSetDirty |= bElementsPresentDirty;
+			// note that we don't have to update bSparseElementsBitSetDirty since we're moving elements
+			// between chunks of the same archeype so the archetype's cumulative CachedSparseElementsBitSet 
+			// won't be affected
+			NewSparseElements.MoveElementsFrom(OldSparseElements, EntityInChunkHandle.IndexWithinChunk, NewIndexWithinChunk, 1);
 		}
 
 		// Clean up the old chunk
-		RemoveEntityInternal(OldAbsoluteIndex);
+		RemoveEntityInternal(EntityInChunkHandle.GetAbsoluteIndex(NumEntitiesPerChunk));
 	}
 }
 
@@ -1726,11 +1718,10 @@ FORCEINLINE void FMassArchetypeData::MoveFragmentsToNewLocationInternal(FMassArc
 {
 	for (const FMassArchetypeFragmentConfig& FragmentConfig : FragmentConfigs)
 	{
-		void* DyingFragmentPtr = FragmentConfig.GetFragmentData(Target.RawChunkMemory, Target.IndexWithinChunk);
-		void* MovingFragmentPtr = FragmentConfig.GetFragmentData(Source.RawChunkMemory, Source.IndexWithinChunk); 
+		void* TargetPtr = FragmentConfig.GetFragmentData(Target.RawChunkMemory, Target.IndexWithinChunk);
+		const void* SourcePtr = FragmentConfig.GetFragmentData(Source.RawChunkMemory, Source.IndexWithinChunk); 
 
-		// Swap fragments to the empty space just created.
-		FMemory::Memcpy(DyingFragmentPtr, MovingFragmentPtr, FragmentConfig.FragmentType->GetStructureSize() * NumberToMove);
+		FMemory::Memcpy(TargetPtr, SourcePtr, FragmentConfig.FragmentType->GetStructureSize() * NumberToMove);
 	}
 }
 
diff --git a/Engine/Source/Runtime/MassEntity/Private/MassEntityTypes.cpp b/Engine/Source/Runtime/MassEntity/Private/MassEntityTypes.cpp
--- a/Engine/Source/Runtime/MassEntity/Private/MassEntityTypes.cpp
+++ b/Engine/Source/Runtime/MassEntity/Private/MassEntityTypes.cpp
@@ -136,28 +136,6 @@ FSharedStruct FMassArchetypeSharedFragmentValues::Add_GetRef(const FSharedStruct
 	return StructInstance;
 }
 
-void FMassArchetypeSharedFragmentValues::ReplaceSharedFragments(TConstArrayView<FSharedStruct> Fragments)
-{
-	DirtyHashCache();
-	for (const FSharedStruct& NewFragment : Fragments)
-	{
-		const UScriptStruct* NewFragScriptStruct = NewFragment.GetScriptStruct();
-		check(NewFragScriptStruct);
-
-		bool bEntryFound = false;
-		for (FSharedStruct& MyFragment : SharedFragments)
-		{
-			if (MyFragment.GetScriptStruct() == NewFragScriptStruct)
-			{
-				MyFragment = NewFragment;
-				bEntryFound = true;
-				break;
-			}
-		}
-		ensureMsgf(bEntryFound, TEXT("Existing fragment of type %s could not be found"), *GetNameSafe(NewFragScriptStruct));
-	}
-}
-
 uint32 FMassArchetypeSharedFragmentValues::CalculateHash() const
 {
 	if (!testableEnsureMsgf(bSorted, TEXT("Expecting the containers to be sorted for the hash caluclation to be consistent")))
diff --git a/Engine/Source/Runtime/MassEntity/Public/MassArchetypeTypes.h b/Engine/Source/Runtime/MassEntity/Public/MassArchetypeTypes.h
--- a/Engine/Source/Runtime/MassEntity/Public/MassArchetypeTypes.h
+++ b/Engine/Source/Runtime/MassEntity/Public/MassArchetypeTypes.h
@@ -321,6 +321,8 @@ struct FMassEntityInChunkDataHandle : FMassRawEntityInChunkData
 	MASSENTITY_API bool IsValid(const FMassArchetypeHandle& ArchetypeHandle) const;
 	bool operator==(const FMassEntityInChunkDataHandle& Other) const;
 
+	int32 GetAbsoluteIndex(const int32 NumEntitiesPerChunk) const;
+
 	const int32 ChunkIndex = INDEX_NONE;
 	const int32 ChunkSerialNumber = INDEX_NONE;
 };
@@ -497,4 +499,9 @@ inline bool FMassEntityInChunkDataHandle::operator==(const FMassEntityInChunkDat
 		&& ChunkSerialNumber == Other.ChunkSerialNumber;
 }
 
+inline int32 FMassEntityInChunkDataHandle::GetAbsoluteIndex(const int32 NumEntitiesPerChunk) const
+{
+	return ChunkIndex * NumEntitiesPerChunk + IndexWithinChunk;
+}
+
 #undef UE_API 
\ No newline at end of file
diff --git a/Engine/Source/Runtime/MassEntity/Public/MassEntityTypes.h b/Engine/Source/Runtime/MassEntity/Public/MassEntityTypes.h
--- a/Engine/Source/Runtime/MassEntity/Public/MassEntityTypes.h
+++ b/Engine/Source/Runtime/MassEntity/Public/MassEntityTypes.h
@@ -483,7 +483,28 @@ struct FMassArchetypeSharedFragmentValues
 	 * Note that it's callers responsibility to ensure every fragment type in Fragments already has an instance in
 	 * this FMassArchetypeSharedFragmentValues instance. Failing that assumption will result in ensure failure. 
 	 */
-	MASSENTITY_API void ReplaceSharedFragments(TConstArrayView<FSharedStruct> Fragments);
+	template<typename TSharedStruct>
+	void ReplaceSharedFragments(TArrayView<TSharedStruct> Fragments)
+	{
+		DirtyHashCache();
+		for (const TSharedStruct& NewFragment : Fragments)
+		{
+			const UScriptStruct* NewFragScriptStruct = NewFragment.GetScriptStruct();
+			check(NewFragScriptStruct);
+
+			bool bEntryFound = false;
+			for (TSharedStruct& MyFragment : GetMutableFragmentsContainer<TSharedStruct>())
+			{
+				if (MyFragment.GetScriptStruct() == NewFragScriptStruct)
+				{
+					MyFragment = NewFragment;
+					bEntryFound = true;
+					break;
+				}
+			}
+			ensureMsgf(bEntryFound, TEXT("Existing fragment of type %s could not be found"), *GetNameSafe(NewFragScriptStruct));
+		}
+	}
 
 	/** 
 	 * Appends contents of Other to `this` instance. All common fragments will get overridden with values in Other.
@@ -536,6 +557,10 @@ struct FMassArchetypeSharedFragmentValues
 	{
 		return SharedFragments;
 	}
+
+	/** depending on TSharedStruct returns Shared or ConstShared fragment container */
+	template<typename TSharedStruct>
+	TArray<TSharedStruct>& GetMutableFragmentsContainer();
 	
 	inline const TArray<FSharedStruct>& GetSharedFragments() const
 	{
@@ -1074,3 +1099,15 @@ inline void FMassArchetypeSharedFragmentValues::Reset()
 	ConstSharedFragments.Reset();
 	SharedFragments.Reset();
 }
+
+template<>
+inline TArray<FSharedStruct>& FMassArchetypeSharedFragmentValues::GetMutableFragmentsContainer<FSharedStruct>()
+{
+	return SharedFragments;
+}
+
+template<>
+inline TArray<FConstSharedStruct>& FMassArchetypeSharedFragmentValues::GetMutableFragmentsContainer<FConstSharedStruct>()
+{
+	return ConstSharedFragments;
+}
\ No newline at end of file
