[EntityFramework] replaced sunsetting FMassArchetypeCompositionDescriptor in favor of FMassElementBitset
Added new overload of CreateArchetype for FMassElementBitSet

#Synced-CL 51050949

diff --git a/Engine/Source/Runtime/MassEntity/Private/MassEntityManager.cpp b/Engine/Source/Runtime/MassEntity/Private/MassEntityManager.cpp
--- a/Engine/Source/Runtime/MassEntity/Private/MassEntityManager.cpp
+++ b/Engine/Source/Runtime/MassEntity/Private/MassEntityManager.cpp
@@ -396,9 +396,14 @@ FMassArchetypeHandle FMassEntityManager::GetOrCreateSuitableArchetype(const FMas
 }
 
 FMassArchetypeHandle FMassEntityManager::CreateArchetype(const FMassArchetypeCompositionDescriptor& Composition, const FMassArchetypeCreationParams& CreationParams)
+{
+	return CreateArchetype(Composition.GetElementsBitSet(), CreationParams);
+}
+
+FMassArchetypeHandle FMassEntityManager::CreateArchetype(const FMassElementBitSet& ElementBitSet, const FMassArchetypeCreationParams& CreationParams)
 {
 	LLM_SCOPE_BYNAME(TEXT("Mass/EntityManager"));
-	FMassElementBitSet NonSparseComposition = Composition.GetElementsBitSet().AndNot(FMassElementBitSet::GetAllSparseElements());
+	const FMassElementBitSet NonSparseComposition = ElementBitSet.AndNot(FMassElementBitSet::GetAllSparseElements());
 	const uint32 TypeHash = HashCombine(GetTypeHash(NonSparseComposition), GetTypeHash(UE::Mass::FArchetypeGroups()));
 
 	TArray<TSharedPtr<FMassArchetypeData>>& HashRow = FragmentHashToArchetypeMap.FindOrAdd(TypeHash);
diff --git a/Engine/Source/Runtime/MassEntity/Public/MassEntityManager.h b/Engine/Source/Runtime/MassEntity/Public/MassEntityManager.h
--- a/Engine/Source/Runtime/MassEntity/Public/MassEntityManager.h
+++ b/Engine/Source/Runtime/MassEntity/Public/MassEntityManager.h
@@ -180,6 +180,15 @@ public:
 	 */
 	UE_API FMassArchetypeHandle CreateArchetype(const FMassArchetypeCompositionDescriptor& Composition, const FMassArchetypeCreationParams& CreationParams = FMassArchetypeCreationParams());
 
+	/**
+	 * CreateArchetype from a mass element bitset and initial values
+	 *
+	 * @param ElementBitSet Bitset of fragment, tag and chunk fragment types
+	 * @param CreationParams Parameters used during archetype construction
+	 * @return a handle of a new archetype
+	 */
+	UE_API FMassArchetypeHandle CreateArchetype(const FMassElementBitSet& ElementBitSet, const FMassArchetypeCreationParams& CreationParams = FMassArchetypeCreationParams());
+
 	/**
 	 *  Creates an archetype like SourceArchetype + InFragments.
 	 *  @param SourceArchetype the archetype used to initially populate the list of fragments of the archetype being created.
