[Mass] Made FMassArchetypeCompositionDescriptor(const FMassElementBitSet& InElements) constructor non-explicit to support near-ish-future removal of FMassArchetypeCompositionDescriptor

#Synced-CL 51050668

diff --git a/Engine/Source/Runtime/MassEntity/Public/MassEntityTypes.h b/Engine/Source/Runtime/MassEntity/Public/MassEntityTypes.h
--- a/Engine/Source/Runtime/MassEntity/Public/MassEntityTypes.h
+++ b/Engine/Source/Runtime/MassEntity/Public/MassEntityTypes.h
@@ -126,7 +126,11 @@ struct FMassArchetypeCompositionDescriptor
 		: ElementsBitSet(MoveTemp(InElements))
 	{}
 
-	explicit FMassArchetypeCompositionDescriptor(const FMassElementBitSet& InElements)
+	/**
+	 * Constructor deliberately not marked as eplicit to support near-future change where we get rid of FMassArchetypeCompositionDescriptor
+	 * since it's just a wrapper for FMassElementBitSet now.
+	 */
+	FMassArchetypeCompositionDescriptor(const FMassElementBitSet& InElements)
 		: ElementsBitSet(InElements)
 	{}
 
