[MassEntity] added SwapConstSharedFragmentForEntity to FMassEntityManager API
Added support for const shared fragments in:
- template version of FMassArchetypeData::SetSharedFragmentsData
- FMassArchetypeSharedFragmentValues::ReplaceSharedFragments

#Synced-CL 51051399

diff --git a/Engine/Source/Runtime/MassEntity/Internal/MassArchetypeData.h b/Engine/Source/Runtime/MassEntity/Internal/MassArchetypeData.h
--- a/Engine/Source/Runtime/MassEntity/Internal/MassArchetypeData.h
+++ b/Engine/Source/Runtime/MassEntity/Internal/MassArchetypeData.h
@@ -557,7 +557,7 @@ public:
 	 * @param NewArchetype the archetype to move to
 	 * @param SharedFragmentValuesOverride if provided will override all given Entity's shared fragment values
 	 */
-	void MoveEntityToAnotherArchetype(const FMassEntityHandle Entity, FMassArchetypeData& NewArchetype, const FMassArchetypeSharedFragmentValues* SharedFragmentValuesOverride = nullptr);
+	void MoveEntityToAnotherArchetype(FMassEntityHandle Entity, FMassArchetypeData& NewArchetype, const FMassArchetypeSharedFragmentValues* SharedFragmentValuesOverride = nullptr);
 
 	/**
 	 * Set all fragment sources data on specified entity, will check if there are fragment sources type that does not exist in the archetype
@@ -711,7 +711,7 @@ public:
 	void BatchRemoveSparseElementFromEntities(const FMassArchetypeEntityCollection& EntityCollection, TNotNull<const UScriptStruct*> ElementType);
 
 	template<typename TSharedStruct>
-	requires std::is_same_v<typename TRemoveReference<TSharedStruct>::Type, FSharedStruct> || std::is_same_v<typename TRemoveReference<TSharedStruct>::Type, FConstSharedStruct>
+	requires std::is_same_v<typename TDecay<TSharedStruct>::Type, FSharedStruct> || std::is_same_v<typename TDecay<TSharedStruct>::Type, FConstSharedStruct>
 	void SetSharedFragmentsData(const FMassEntityHandle Entity, const TArray<TSharedStruct>& SharedFragmentValueOverrides)
 	{
 		SetSharedFragmentsData<TSharedStruct>(Entity, MakeArrayView(SharedFragmentValueOverrides));
@@ -727,7 +727,7 @@ public:
 	 *    present in given archetypes FMassArchetypeSharedFragmentValues
 	 */
 	template<typename TSharedStruct>
-	requires std::is_same_v<TSharedStruct, FSharedStruct> || std::is_same_v<TSharedStruct, FConstSharedStruct>
+	requires std::is_same_v<typename TDecay<TSharedStruct>::Type, FSharedStruct> || std::is_same_v<typename TDecay<TSharedStruct>::Type, FConstSharedStruct>
 	void SetSharedFragmentsData(const FMassEntityHandle Entity, TArrayView<TSharedStruct> SharedFragmentValueOverrides)
 	{
 		FMassEntityInChunkDataHandle EntityInChunkHandle = MakeEntityHandle(Entity.Index);
diff --git a/Engine/Source/Runtime/MassEntity/Private/MassEntityManager.cpp b/Engine/Source/Runtime/MassEntity/Private/MassEntityManager.cpp
--- a/Engine/Source/Runtime/MassEntity/Private/MassEntityManager.cpp
+++ b/Engine/Source/Runtime/MassEntity/Private/MassEntityManager.cpp
@@ -1565,6 +1565,42 @@ void FMassEntityManager::RemoveElementFromEntity(FMassEntityHandle Entity, TNotN
 	}
 }
 
+bool FMassEntityManager::SwapConstSharedFragmentForEntity(const FMassEntityHandle Entity, const FConstSharedStruct& InConstSharedFragment)
+{
+	CHECK_SYNC_API_RETURN(return false);
+
+	if (!ensureMsgf(InConstSharedFragment.IsValid(), TEXT("%hs parameter Fragment is expected to be valid"), __FUNCTION__))
+	{
+		return false;
+	}
+
+	checkf(ConstSharedFragmentsContainer.Find(InConstSharedFragment.GetScriptStruct())
+		, TEXT("Provided shared struct requires to be registered first (i.e., GetGetOrCreateConstSharedFragment"));
+
+	CheckIfEntityIsActive(Entity);
+
+	FMassArchetypeData* CurrentArchetype = GetEntityStorageInterface().GetArchetypeAsShared(Entity.Index).Get();
+	check(CurrentArchetype);
+
+	const UScriptStruct* FragmentType = InConstSharedFragment.GetScriptStruct();
+	if (!CurrentArchetype->GetCompositionBitSet().Contains(*FragmentType))
+	{
+		// Observers will get notify by the call to AddConstSharedFragmentToEntity
+		AddConstSharedFragmentToEntity(Entity, InConstSharedFragment);
+	}
+	else
+	{
+		// SetSharedFragmentsData won't notify observers so we have to do it here to allow observers to handle the change
+		ObserverManager.OnPreCompositionRemoved(Entity, FMassArchetypeCompositionDescriptor(FMassElementBitSet{*FragmentType}));
+
+		CurrentArchetype->SetSharedFragmentsData(Entity, MakeArrayView<FConstSharedStruct>({ InConstSharedFragment }));
+
+		ObserverManager.OnPostCompositionAdded(Entity, FMassElementBitSet(*FragmentType));
+	}
+
+	return true;
+}
+
 bool FMassEntityManager::AddConstSharedFragmentToEntity(const FMassEntityHandle Entity, const FConstSharedStruct& InConstSharedFragment)
 {
 	CHECK_SYNC_API_RETURN(return false);
@@ -1574,6 +1610,9 @@ bool FMassEntityManager::AddConstSharedFragmentToEntity(const FMassEntityHandle
 		return false;
 	}
 
+	checkf(ConstSharedFragmentsContainer.Find(InConstSharedFragment.GetScriptStruct())
+		, TEXT("Provided shared struct requires to be registered first (i.e., GetGetOrCreateConstSharedFragment"));
+
 	CheckIfEntityIsActive(Entity);
 	
 	FMassArchetypeData* CurrentArchetype = GetEntityStorageInterface().GetArchetypeAsShared(Entity.Index).Get();
@@ -1673,6 +1712,9 @@ bool FMassEntityManager::AddSharedFragmentToEntity(const FMassEntityHandle Entit
 		return false;
 	}
 
+	checkf(SharedFragmentsContainer.Find(InSharedFragment.GetScriptStruct())
+		, TEXT("Provided shared struct requires to be registered first (i.e., GetGetOrCreateSharedFragment"));
+
 	CheckIfEntityIsActive(Entity);
 	
 	FMassArchetypeData* CurrentArchetype = GetEntityStorageInterface().GetArchetypeAsShared(Entity.Index).Get();
diff --git a/Engine/Source/Runtime/MassEntity/Private/MassEntityTypes.cpp b/Engine/Source/Runtime/MassEntity/Private/MassEntityTypes.cpp
--- a/Engine/Source/Runtime/MassEntity/Private/MassEntityTypes.cpp
+++ b/Engine/Source/Runtime/MassEntity/Private/MassEntityTypes.cpp
@@ -138,7 +138,7 @@ FSharedStruct FMassArchetypeSharedFragmentValues::Add_GetRef(const FSharedStruct
 
 uint32 FMassArchetypeSharedFragmentValues::CalculateHash() const
 {
-	if (!testableEnsureMsgf(bSorted, TEXT("Expecting the containers to be sorted for the hash caluclation to be consistent")))
+	if (!testableEnsureMsgf(bSorted, TEXT("Expecting the containers to be sorted for the hash calculation to be consistent")))
 	{
 		return 0;
 	}
diff --git a/Engine/Source/Runtime/MassEntity/Public/MassEntityManager.h b/Engine/Source/Runtime/MassEntity/Public/MassEntityManager.h
--- a/Engine/Source/Runtime/MassEntity/Public/MassEntityManager.h
+++ b/Engine/Source/Runtime/MassEntity/Public/MassEntityManager.h
@@ -495,6 +495,15 @@ public:
 		return DoesEntityHaveElement(Entity, T::StaticStruct());
 	}
 
+	/**
+	 * Adds a new const shared fragment to the given entity if it doesn't have a shared fragment of the given type, 
+	 * moves it to its new archetype and notify observers.
+	 * In case a fragment of that type is already present, then it replaces it,
+	 * moves the entity to its new chunk within the archetype and notifies the observers.
+	 * @return whether the Fragment value is valid and successfully assigned to the entity (added or replaced)
+	 */
+	UE_API bool SwapConstSharedFragmentForEntity(const FMassEntityHandle EntityHandle, const FConstSharedStruct& InConstSharedFragment);
+
 	/** 
 	 * Adds a new const shared fragment to the given entity. Note that it only works if the given entity doesn't have
 	 * a shared fragment of the given type. The function will give a soft "pass" if the entity has the shared fragment
diff --git a/Engine/Source/Runtime/MassEntity/Public/MassEntityTypes.h b/Engine/Source/Runtime/MassEntity/Public/MassEntityTypes.h
--- a/Engine/Source/Runtime/MassEntity/Public/MassEntityTypes.h
+++ b/Engine/Source/Runtime/MassEntity/Public/MassEntityTypes.h
@@ -490,6 +490,8 @@ struct FMassArchetypeSharedFragmentValues
 	template<typename TSharedStruct>
 	void ReplaceSharedFragments(TArrayView<TSharedStruct> Fragments)
 	{
+		using FDecayedSharedStruct = TDecay<TSharedStruct>::Type;
+
 		DirtyHashCache();
 		for (const TSharedStruct& NewFragment : Fragments)
 		{
@@ -497,7 +499,7 @@ struct FMassArchetypeSharedFragmentValues
 			check(NewFragScriptStruct);
 
 			bool bEntryFound = false;
-			for (TSharedStruct& MyFragment : GetMutableFragmentsContainer<TSharedStruct>())
+			for (FDecayedSharedStruct& MyFragment : GetMutableFragmentsContainer<FDecayedSharedStruct>())
 			{
 				if (MyFragment.GetScriptStruct() == NewFragScriptStruct)
 				{
