[Mass] Modified UE::Mass::Utils::CreateEntityCollections to handle the common case of single-entity input, and behave better (in terms of performance) when multiple entities of the same archetype are processed.

#Synced-CL 51090036

diff --git a/Engine/Source/Runtime/MassEntity/Private/MassEntityUtils.cpp b/Engine/Source/Runtime/MassEntity/Private/MassEntityUtils.cpp
--- a/Engine/Source/Runtime/MassEntity/Private/MassEntityUtils.cpp
+++ b/Engine/Source/Runtime/MassEntity/Private/MassEntityUtils.cpp
@@ -74,15 +74,48 @@ void CreateEntityCollections(const FMassEntityManager& EntityManager, const TCon
 {
 	TRACE_CPUPROFILER_EVENT_SCOPE_STR("Mass_CreateSparseChunks");
 
+	// find first valid entity
+	int32 EntityIndex = 0;
+	for (; EntityIndex < Entities.Num(); ++EntityIndex)
+	{
+		if (EntityManager.IsEntityValid(Entities[EntityIndex]))
+		{
+			break;
+		}
+	}
+
+	if (EntityIndex >= Entities.Num())
+	{
+		// no valid entities
+		return;
+	}
+
+	FMassArchetypeHandle Archetype = EntityManager.GetArchetypeForEntityUnsafe(Entities[EntityIndex]);
+	if (Entities.Num() - EntityIndex == 1)
+	{
+		// common, trivial case, so we can handle it directly
+		OutEntityCollections.Add(FMassArchetypeEntityCollection(Archetype, Entities[EntityIndex]));
+		return;
+	}
+
 	TMap<const FMassArchetypeHandle, TArray<FMassEntityHandle>> ArchetypeToEntities;
 
-	for (const FMassEntityHandle& Entity : Entities)
+	TArray<FMassEntityHandle>* PerArchetypeEntities = &ArchetypeToEntities.Add(Archetype);
+	PerArchetypeEntities->Add(Entities[EntityIndex++]);
+
+	for (; EntityIndex < Entities.Num(); ++EntityIndex)
 	{
+		const FMassEntityHandle& Entity = Entities[EntityIndex];
 		if (EntityManager.IsEntityValid(Entity))
 		{
-			FMassArchetypeHandle Archetype = EntityManager.GetArchetypeForEntityUnsafe(Entity);
-			TArray<FMassEntityHandle>& PerArchetypeEntities = ArchetypeToEntities.FindOrAdd(Archetype);
-			PerArchetypeEntities.Add(Entity);
+			FMassArchetypeHandle LocalArchetype = EntityManager.GetArchetypeForEntityUnsafe(Entity);
+			if (Archetype != LocalArchetype)
+			{
+				Archetype = LocalArchetype;
+				PerArchetypeEntities = &ArchetypeToEntities.FindOrAdd(Archetype);
+			}
+			
+			PerArchetypeEntities->Add(Entity);
 		}
 	}
 
