[Mass] fixed Mass FMassCommandAddFragmentInstances  to support move semantics when adding fragment instances.

#Synced-CL 51208877

diff --git a/Engine/Source/Runtime/MassEntity/Public/MassCommands.h b/Engine/Source/Runtime/MassEntity/Public/MassCommands.h
--- a/Engine/Source/Runtime/MassEntity/Public/MassCommands.h
+++ b/Engine/Source/Runtime/MassEntity/Public/MassCommands.h
@@ -411,10 +411,11 @@ struct FMassCommandAddFragmentInstances : public FMassBatchedEntityCommand
 		, FragmentsAffected(UE::Mass::Utils::ConstructFragmentBitSet<EMassCommandCheckTime::CompileTimeCheck, TOthers...>())
 	{}
 
-	void Add(FMassEntityHandle Entity, TOthers... InFragments)
+	template<typename... TInstances>
+	void Add(FMassEntityHandle Entity, TInstances... InFragments)
 	{
 		Super::Add(Entity);
-		Fragments.Add(InFragments...);
+		Fragments.Add(Forward<TInstances>(InFragments)...);
 	}
 
 #if WITH_MASSENTITY_DEBUG
diff --git a/Engine/Source/Runtime/MassEntity/Public/MassEntityTypes.h b/Engine/Source/Runtime/MassEntity/Public/MassEntityTypes.h
--- a/Engine/Source/Runtime/MassEntity/Public/MassEntityTypes.h
+++ b/Engine/Source/Runtime/MassEntity/Public/MassEntityTypes.h
@@ -849,12 +849,25 @@ namespace UE::Mass
 			return FragmentInstances.GetAllocatedSize() + Super::GetAllocatedSize();
 		}
 
-		int GetNumArrays() const { return Ordinal + 1; }
+		int GetNumArrays() const
+		{
+			return Ordinal + 1;
+		}
 
-		void Add(const FType& Item, TOthers... Rest)
+		/** TInstances might be different from TOthers if move semantics were used */
+		template<typename... TInstances>
+		void Add(const FType& Item, TInstances... Rest)
 		{
 			FragmentInstances.Add(Item);
-			Super::Add(Rest...);
+			Super::Add(Forward<TInstances>(Rest)...);
+		}
+
+		/** TInstances might be different from TOthers if move semantics were used */
+		template<typename... TInstances>
+		void Add(FType&& Item, TInstances... Rest)
+		{
+			FragmentInstances.Emplace(Forward<FType>(Item));
+			Super::Add(Forward<TInstances>(Rest)...);
 		}
 
 		void GetAsGenericMultiArray(TArray<FStructArrayView>& A) /*const*/
@@ -878,7 +891,7 @@ namespace UE::Mass
 		TArray<FType> FragmentInstances;
 	};
 
-	/**TMultiArray single-type specialization */
+	/** TMultiArray single-type specialization */
 	template<typename T>
 	struct TMultiArray<T>
 	{
@@ -890,9 +903,20 @@ namespace UE::Mass
 			return FragmentInstances.GetAllocatedSize();
 		}
 
-		int GetNumArrays() const { return Ordinal + 1; }
+		int GetNumArrays() const
+		{
+			return Ordinal + 1;
+		}
 
-		void Add(const FType& Item) { FragmentInstances.Add(Item); }
+		void Add(const FType& Item)
+		{
+			FragmentInstances.Add(Item);
+		}
+
+		void Add(FType&& Item)
+		{
+			FragmentInstances.Emplace(Forward<FType>(Item));
+		}
 
 		void GetAsGenericMultiArray(TArray<FStructArrayView>& A) /*const*/
 		{
