[Mass]
Removed erroneous const specifiers from FMutableOptionalFragmentAccess

#Synced-CL 50938174

diff --git a/Engine/Source/Runtime/MassEntity/Public/MassQueryExecutor.h b/Engine/Source/Runtime/MassEntity/Public/MassQueryExecutor.h
--- a/Engine/Source/Runtime/MassEntity/Public/MassQueryExecutor.h
+++ b/Engine/Source/Runtime/MassEntity/Public/MassQueryExecutor.h
@@ -300,12 +300,12 @@ struct FMutableOptionalFragmentAccess
 
 	FMutableOptionalFragmentAccess() = default;
 
-	inline TArrayView<TFragment>& Get() const
+	inline TArrayView<TFragment>& Get()
 	{
 		return View;
 	}
 
-	inline TFragment& operator[](int32 Index) const
+	inline TFragment& operator[](int32 Index)
 	{
 		return View[Index];
 	}
@@ -369,7 +369,7 @@ struct FConstOptionalFragmentAccess
 
 	inline void ConfigureQuery(FMassEntityQuery& EntityQuery, FMassSubsystemRequirements& ProcessorRequirements) const
 	{
-		EntityQuery.AddRequirement<TFragment>(EMassFragmentAccess::ReadOnly);
+		EntityQuery.AddRequirement<TFragment>(EMassFragmentAccess::ReadOnly, EMassFragmentPresence::Optional);
 	}
 
 	inline void SetupForExecute(FMassExecutionContext& Context)
