[Mass] Fixed compilation issues with FMassArchetypeData::SetSharedFragmentsData(const FMassEntityHandle, const TArray<>)

#rnx

#Synced-CL 51211040

diff --git a/Engine/Source/Runtime/MassEntity/Internal/MassArchetypeData.h b/Engine/Source/Runtime/MassEntity/Internal/MassArchetypeData.h
--- a/Engine/Source/Runtime/MassEntity/Internal/MassArchetypeData.h
+++ b/Engine/Source/Runtime/MassEntity/Internal/MassArchetypeData.h
@@ -714,7 +714,7 @@ public:
 	requires std::is_same_v<typename TDecay<TSharedStruct>::Type, FSharedStruct> || std::is_same_v<typename TDecay<TSharedStruct>::Type, FConstSharedStruct>
 	void SetSharedFragmentsData(const FMassEntityHandle Entity, const TArray<TSharedStruct>& SharedFragmentValueOverrides)
 	{
-		SetSharedFragmentsData<TSharedStruct>(Entity, MakeArrayView(SharedFragmentValueOverrides));
+		SetSharedFragmentsData(Entity, MakeArrayView(SharedFragmentValueOverrides));
 	}
 
 	/**
