[Mass] API update to replace UObject& to TNotNull<*>

#rnx

#Synced-CL 51260715

diff --git a/Engine/Plugins/AI/MassAI/Source/MassAIBehavior/Private/MassLookAtSubsystem.cpp b/Engine/Plugins/AI/MassAI/Source/MassAIBehavior/Private/MassLookAtSubsystem.cpp
--- a/Engine/Plugins/AI/MassAI/Source/MassAIBehavior/Private/MassLookAtSubsystem.cpp
+++ b/Engine/Plugins/AI/MassAI/Source/MassAIBehavior/Private/MassLookAtSubsystem.cpp
@@ -423,7 +423,7 @@ void UMassLookAtSubsystem::Initialize(FSubsystemCollectionBase& Collection)
 	FMassEntityManager& EntityManager = Collection.InitializeDependency<UMassEntitySubsystem>()->GetMutableEntityManager();
 
 	// Create Mass archetype for entities representing requests
-	const FMassArchetypeCompositionDescriptor RequestComposition{ FMassFragmentBitSet(*FMassLookAtRequestFragment::StaticStruct()) };
+	const FMassArchetypeCompositionDescriptor RequestComposition{ FMassFragmentBitSet(FMassLookAtRequestFragment::StaticStruct()) };
 	RequestArchetype = EntityManager.CreateArchetype(RequestComposition);
 
 	// Create Mass archetype for entities representing targets
diff --git a/Engine/Plugins/AI/MassAI/Source/MassAIBehavior/Private/MassStateTreeTrait.cpp b/Engine/Plugins/AI/MassAI/Source/MassAIBehavior/Private/MassStateTreeTrait.cpp
--- a/Engine/Plugins/AI/MassAI/Source/MassAIBehavior/Private/MassStateTreeTrait.cpp
+++ b/Engine/Plugins/AI/MassAI/Source/MassAIBehavior/Private/MassStateTreeTrait.cpp
@@ -92,30 +92,30 @@ bool UMassStateTreeTrait::ValidateTemplate(const FMassEntityTemplateBuildContext
 			}
 			else if (UE::Mass::IsA<FMassFragment>(ItemDesc.Struct))
 			{
-				const UScriptStruct& FragmentType = *CastChecked<UScriptStruct>(ItemDesc.Struct);
+				const UScriptStruct* FragmentType = CastChecked<UScriptStruct>(ItemDesc.Struct);
 				if (BuildContext.HasFragment(FragmentType) == false)
 				{
-					OutTraitRequirements.Add(&FragmentType);
+					OutTraitRequirements.Add(FragmentType);
 					UE_VLOG_UELOG(MassStateTreeSubsystem, LogMassBehavior, Error, TEXT("StateTree %s: Could not find required fragment %s"), *GetNameSafe(StateTree), *GetNameSafe(ItemDesc.Struct));
 					bIssuesFound = true;
 				}
 			}
 			else if (UE::Mass::IsA<FMassSharedFragment>(ItemDesc.Struct))
 			{
-				const UScriptStruct& FragmentType = *CastChecked<UScriptStruct>(ItemDesc.Struct);
+				const UScriptStruct* FragmentType = CastChecked<UScriptStruct>(ItemDesc.Struct);
 				if (BuildContext.HasSharedFragment(FragmentType) == false)
 				{
-					OutTraitRequirements.Add(&FragmentType);
+					OutTraitRequirements.Add(FragmentType);
 					UE_VLOG_UELOG(MassStateTreeSubsystem, LogMassBehavior, Error, TEXT("StateTree %s: Could not find required shared fragment %s"), *GetNameSafe(StateTree), *GetNameSafe(ItemDesc.Struct));
 					bIssuesFound = true;
 				}
 			}
 			else if (UE::Mass::IsA<FMassConstSharedFragment>(ItemDesc.Struct))
 			{
-				const UScriptStruct& FragmentType = *CastChecked<UScriptStruct>(ItemDesc.Struct);
+				const UScriptStruct* FragmentType = CastChecked<UScriptStruct>(ItemDesc.Struct);
 				if (BuildContext.HasConstSharedFragment(FragmentType) == false)
 				{
-					OutTraitRequirements.Add(&FragmentType);
+					OutTraitRequirements.Add(FragmentType);
 					UE_VLOG_UELOG(MassStateTreeSubsystem, LogMassBehavior, Error, TEXT("StateTree %s: Could not find required const shared fragment %s"), *GetNameSafe(StateTree), *GetNameSafe(ItemDesc.Struct));
 					bIssuesFound = true;
 				}
diff --git a/Engine/Plugins/Experimental/EditorDataStorage/Source/TedsCore/Private/Processors/TypedElementProcessorAdaptors.cpp b/Engine/Plugins/Experimental/EditorDataStorage/Source/TedsCore/Private/Processors/TypedElementProcessorAdaptors.cpp
--- a/Engine/Plugins/Experimental/EditorDataStorage/Source/TedsCore/Private/Processors/TypedElementProcessorAdaptors.cpp
+++ b/Engine/Plugins/Experimental/EditorDataStorage/Source/TedsCore/Private/Processors/TypedElementProcessorAdaptors.cpp
@@ -86,11 +86,11 @@ namespace UE::Editor::DataStorage
 			{
 				if (UE::Mass::IsA<FMassTag>(ColumnType))
 				{
-					return Context.DoesArchetypeHaveTag(*ColumnType);
+					return Context.DoesArchetypeHaveTag(ColumnType);
 				}
 				if (UE::Mass::IsA<FMassFragment>(ColumnType))
 				{
-					return Context.DoesArchetypeHaveFragment(*ColumnType);
+					return Context.DoesArchetypeHaveFragment(ColumnType);
 				}
 				const bool bIsTagOrFragment = false;
 				checkf(bIsTagOrFragment, TEXT("Attempting to check for a column type that is not a column or tag."));
@@ -106,7 +106,7 @@ namespace UE::Editor::DataStorage
 
 				if (UE::Mass::IsA<FMassTag>(ColumnType) || UE::Mass::IsA<FMassFragment>(ColumnType))
 				{
-					return Composition.GetElementsBitSet().Contains(*ColumnType);
+					return Composition.GetElementsBitSet().Contains(ColumnType);
 				}
 				const bool bIsTagOrFragment = false;
 				checkf(bIsTagOrFragment, TEXT("Attempting to check for a column type that is not a column or tag."));
@@ -207,7 +207,7 @@ namespace UE::Editor::DataStorage
 				{
 					if (UE::Mass::IsA<FMassTag>(ColumnType))
 					{
-						if (this->Context.DoesArchetypeHaveTag(*ColumnType))
+						if (this->Context.DoesArchetypeHaveTag(ColumnType))
 						{
 							Descriptor.Add(ColumnType);
 						}
@@ -216,7 +216,7 @@ namespace UE::Editor::DataStorage
 					{
 						checkf(UE::Mass::IsA<FMassFragment>(ColumnType),
 							TEXT("Given struct type is not a valid fragment or tag type."));
-						if (this->Context.DoesArchetypeHaveFragment(*ColumnType))
+						if (this->Context.DoesArchetypeHaveFragment(ColumnType))
 						{
 							Descriptor.Add(ColumnType);
 						}
@@ -358,7 +358,7 @@ namespace UE::Editor::DataStorage
 						for (FAddValueColumn* const AddColumn : AddedColumnWithRelocationQueue)
 						{
 							FRowOperations& Operations = RowOperations.FindOrAdd(AddColumn->Entity);
-							Operations.ColumnsChange.AddedFragments.Add(*AddColumn->FragmentType);
+							Operations.ColumnsChange.AddedFragments.Add(AddColumn->FragmentType);
 							Operations.Relocators.Add(FRelocator{
 									.CopyOrMove = AddColumn->Relocator,
 									.FragmentType = AddColumn->FragmentType,
@@ -467,7 +467,7 @@ namespace UE::Editor::DataStorage
 								{
 									for (const FDynamicColumnDescription& DynamicColumnAdded : BatchedModification.Key.AddedDynamicColumns)
 									{
-										BatchedModification.Key.AddedFragments.Add(*PtrToEnvironment->GenerateDynamicColumn(*DynamicColumnAdded.TemplateType, DynamicColumnAdded.Identifier));
+										BatchedModification.Key.AddedFragments.Add(PtrToEnvironment->GenerateDynamicColumn(*DynamicColumnAdded.TemplateType, DynamicColumnAdded.Identifier));
 									}
 
 									const FRowColumnModifications& ColumnModification = BatchedModification.Key;
diff --git a/Engine/Plugins/Experimental/EditorDataStorage/Source/TedsCore/Private/Queries/TypedElementQueryContext_MassForwarder.cpp b/Engine/Plugins/Experimental/EditorDataStorage/Source/TedsCore/Private/Queries/TypedElementQueryContext_MassForwarder.cpp
--- a/Engine/Plugins/Experimental/EditorDataStorage/Source/TedsCore/Private/Queries/TypedElementQueryContext_MassForwarder.cpp
+++ b/Engine/Plugins/Experimental/EditorDataStorage/Source/TedsCore/Private/Queries/TypedElementQueryContext_MassForwarder.cpp
@@ -28,11 +28,11 @@ namespace UE::Editor::DataStorage::Queries
 			const UScriptStruct* ColumnType = *It;
 			if (UE::Mass::IsA<FMassTag>(ColumnType))
 			{
-				Result = Context.DoesArchetypeHaveTag(*ColumnType);
+				Result = Context.DoesArchetypeHaveTag(ColumnType);
 			}
 			else if (UE::Mass::IsA<FMassFragment>(ColumnType))
 			{
-				Result = Context.DoesArchetypeHaveFragment(*ColumnType);
+				Result = Context.DoesArchetypeHaveFragment(ColumnType);
 			}
 			else
 			{
diff --git a/Engine/Plugins/Experimental/EditorDataStorage/Source/TedsCore/Private/TypedElementDatabase.cpp b/Engine/Plugins/Experimental/EditorDataStorage/Source/TedsCore/Private/TypedElementDatabase.cpp
--- a/Engine/Plugins/Experimental/EditorDataStorage/Source/TedsCore/Private/TypedElementDatabase.cpp
+++ b/Engine/Plugins/Experimental/EditorDataStorage/Source/TedsCore/Private/TypedElementDatabase.cpp
@@ -118,7 +118,7 @@ namespace UE::Editor::DataStorage::Private
 	{
 		for (const UScriptStruct* ColumnType : Columns)
 		{
-			Elements.Add(*ColumnType);
+			Elements.Add(ColumnType);
 		}
 	}
 
@@ -137,12 +137,12 @@ namespace UE::Editor::DataStorage::Private
 		{
 			if (UE::Mass::IsA<FMassFragment>(ColumnType))
 			{
-				Fragments.Add(*ColumnType);
+				Fragments.Add(ColumnType);
 				Result.bMustUpdateFragments = true;
 			}
 			else if (UE::Mass::IsA<FMassTag>(ColumnType))
 			{
-				Tags.Add(*ColumnType);
+				Tags.Add(ColumnType);
 				Result.bMustUpdateTags = true;
 			}
 		}
@@ -979,7 +979,7 @@ bool UEditorDataStorage::HasColumns(UE::Editor::DataStorage::RowHandle Row, TCon
 			{
 				if (UE::Mass::IsA<FMassFragment>(*ColumnType) || UE::Mass::IsA<FMassTag>(*ColumnType))
 				{
-					bHasAllColumns = Composition.GetElementsBitSet().Contains(**ColumnType);
+					bHasAllColumns = Composition.GetElementsBitSet().Contains(*ColumnType);
 				}
 				else
 				{
@@ -1015,13 +1015,13 @@ bool UEditorDataStorage::HasColumns(UE::Editor::DataStorage::RowHandle Row, TCon
 			const FMassArchetypeCompositionDescriptor& Composition = ActiveEditorEntityManager->GetArchetypeComposition(Archetype);
 
 			const TWeakObjectPtr<const UScriptStruct>* ColumnTypesEnd = ColumnTypes.end();
-			for (const TWeakObjectPtr<const UScriptStruct>* ColumnType = ColumnTypes.begin(); ColumnType != ColumnTypesEnd && bHasAllColumns; ++ColumnType)
+			for (const TWeakObjectPtr<const UScriptStruct>* WeakColumnType = ColumnTypes.begin(); WeakColumnType != ColumnTypesEnd && bHasAllColumns; ++WeakColumnType)
 			{
-				if (ColumnType->IsValid())
+				if (const UScriptStruct* ColumnType = WeakColumnType->Get())
 				{
-					if (UE::Mass::IsA<FMassFragment>(ColumnType->Get()) || UE::Mass::IsA<FMassTag>(ColumnType->Get()))
+					if (UE::Mass::IsA<FMassFragment>(ColumnType) || UE::Mass::IsA<FMassTag>(ColumnType))
 					{
-						bHasAllColumns = Composition.GetElementsBitSet().Contains(**ColumnType);
+						bHasAllColumns = Composition.GetElementsBitSet().Contains(ColumnType);
 						continue;
 					}
 				}
diff --git a/Engine/Plugins/Experimental/EditorDataStorage/Source/TedsCore/Private/TypedElementDatabaseCommandBuffer.cpp b/Engine/Plugins/Experimental/EditorDataStorage/Source/TedsCore/Private/TypedElementDatabaseCommandBuffer.cpp
--- a/Engine/Plugins/Experimental/EditorDataStorage/Source/TedsCore/Private/TypedElementDatabaseCommandBuffer.cpp
+++ b/Engine/Plugins/Experimental/EditorDataStorage/Source/TedsCore/Private/TypedElementDatabaseCommandBuffer.cpp
@@ -352,7 +352,7 @@ namespace UE::Editor::DataStorage::Legacy
 		if (ColumnType.IsChildOf(FTedsSharedColumn::StaticStruct()) && Execute_IsRowAssigned(MassEntityManager, Row))
 		{
 			FMassEntityHandle Entity = FMassEntityHandle::FromNumber(Row);
-			MassEntityManager.RemoveConstSharedFragmentFromEntity(Entity, ColumnType);
+			MassEntityManager.RemoveConstSharedFragmentFromEntity(Entity, &ColumnType);
 		}
 	}
 
diff --git a/Engine/Plugins/Runtime/MassGameplay/Source/MassEQS/Private/Processors/MassEnvQueryTestProcessor_MassEntityTags.cpp b/Engine/Plugins/Runtime/MassGameplay/Source/MassEQS/Private/Processors/MassEnvQueryTestProcessor_MassEntityTags.cpp
--- a/Engine/Plugins/Runtime/MassGameplay/Source/MassEQS/Private/Processors/MassEnvQueryTestProcessor_MassEntityTags.cpp
+++ b/Engine/Plugins/Runtime/MassGameplay/Source/MassEQS/Private/Processors/MassEnvQueryTestProcessor_MassEntityTags.cpp
@@ -114,7 +114,7 @@ bool Mass::EQS::Utils::TestChunkForAnyTags(const TArray<FInstancedStruct>& Tags,
 			continue;
 		}
 #endif
-		if (Context.DoesArchetypeHaveTag(*Tag.GetScriptStruct()))
+		if (Context.DoesArchetypeHaveTag(TagScriptStruct))
 		{
 			return true;
 		}
@@ -134,7 +134,7 @@ bool Mass::EQS::Utils::TestChunkForAllTags(const TArray<FInstancedStruct>& Tags,
 			continue;
 		}
 #endif
-		if (!Context.DoesArchetypeHaveTag(*Tag.GetScriptStruct()))
+		if (!Context.DoesArchetypeHaveTag(TagScriptStruct))
 		{
 			return false;
 		}
@@ -154,7 +154,7 @@ bool Mass::EQS::Utils::TestChunkForNoTags(const TArray<FInstancedStruct>& Tags,
 			continue;
 		}
 #endif
-		if (Context.DoesArchetypeHaveTag(*Tag.GetScriptStruct()))
+		if (Context.DoesArchetypeHaveTag(Tag.GetScriptStruct()))
 		{
 			return false;
 		}
diff --git a/Engine/Plugins/Runtime/MassGameplay/Source/MassGameplayTestSuite/Private/MassEntityTemplateTest.cpp b/Engine/Plugins/Runtime/MassGameplay/Source/MassGameplayTestSuite/Private/MassEntityTemplateTest.cpp
--- a/Engine/Plugins/Runtime/MassGameplay/Source/MassGameplayTestSuite/Private/MassEntityTemplateTest.cpp
+++ b/Engine/Plugins/Runtime/MassGameplay/Source/MassGameplayTestSuite/Private/MassEntityTemplateTest.cpp
@@ -214,7 +214,7 @@ struct FEntityTest_BatchSetMissingFragment : FEntityTemplateBase
 
 		// forcing archetype change for the created entities. 
 		EntityManager->BatchChangeFragmentCompositionForEntities(CreationContext->GetEntityCollections(*EntityManager.Get())
-			, {}, FMassFragmentBitSet(*FTestFragment_Float::StaticStruct()));
+			, {}, FMassFragmentBitSet(FTestFragment_Float::StaticStruct()));
 
 		{
 			AITEST_SCOPED_CHECK("Attempting to set value of fragment of type", 1);
diff --git a/Engine/Plugins/Runtime/MassGameplay/Source/MassSpawner/Private/MassAssortedFragmentsTrait.cpp b/Engine/Plugins/Runtime/MassGameplay/Source/MassSpawner/Private/MassAssortedFragmentsTrait.cpp
--- a/Engine/Plugins/Runtime/MassGameplay/Source/MassSpawner/Private/MassAssortedFragmentsTrait.cpp
+++ b/Engine/Plugins/Runtime/MassGameplay/Source/MassSpawner/Private/MassAssortedFragmentsTrait.cpp
@@ -30,7 +30,7 @@ void UMassAssortedFragmentsTrait::BuildTemplate(FMassEntityTemplateBuildContext&
 		const UScriptStruct* Type = Tag.GetScriptStruct();
 		if (UE::Mass::IsA<FMassTag>(Type))
 		{
-			BuildContext.AddTag(*Type);
+			BuildContext.AddTag(Type);
 		}
 		else
 		{
diff --git a/Engine/Plugins/Runtime/MassGameplay/Source/MassSpawner/Private/MassEntityTemplateRegistry.cpp b/Engine/Plugins/Runtime/MassGameplay/Source/MassSpawner/Private/MassEntityTemplateRegistry.cpp
--- a/Engine/Plugins/Runtime/MassGameplay/Source/MassSpawner/Private/MassEntityTemplateRegistry.cpp
+++ b/Engine/Plugins/Runtime/MassGameplay/Source/MassSpawner/Private/MassEntityTemplateRegistry.cpp
@@ -132,7 +132,7 @@ bool FMassEntityTemplateBuildContext::BuildFromTraits(TConstArrayView<UMassEntit
 	for (FRemovedType& Removed : RemovedTypes)
 	{
 		check(Removed.TypeRemoved);
-		TemplateData.RemoveTag(*CastChecked<UScriptStruct>(Removed.TypeRemoved));
+		TemplateData.RemoveTag(CastChecked<UScriptStruct>(Removed.TypeRemoved));
 	}
 
 	bBuildInProgress = false;
diff --git a/Engine/Plugins/Runtime/MassGameplay/Source/MassSpawner/Public/MassEntityTemplate.h b/Engine/Plugins/Runtime/MassGameplay/Source/MassSpawner/Public/MassEntityTemplate.h
--- a/Engine/Plugins/Runtime/MassGameplay/Source/MassSpawner/Public/MassEntityTemplate.h
+++ b/Engine/Plugins/Runtime/MassGameplay/Source/MassSpawner/Public/MassEntityTemplate.h
@@ -133,21 +133,28 @@ struct FMassEntityTemplateData
 		Composition.GetElementsBitSet().Add<T>();
 	}
 
-	void AddFragment(const UScriptStruct& FragmentType)
+	void AddFragment(TNotNull<const UScriptStruct*> FragmentType)
+
 	{
-		checkf(UE::Mass::IsA<FMassFragment>(&FragmentType), TEXT(MASS_INVALID_FRAGMENT_MSG_F), *FragmentType.GetName());
+		checkf(UE::Mass::IsA<FMassFragment>(FragmentType), TEXT(MASS_INVALID_FRAGMENT_MSG_F), *FragmentType->GetName());
 		Composition.GetElementsBitSet().Add(FragmentType);
 	}
 
+	UE_DEPRECATED(5.8, "Use overload taking TNotNull")
+	void AddFragment(const UScriptStruct& FragmentType)
+	{
+		AddFragment(TNotNull<const UScriptStruct*>(&FragmentType));
+	}
+
 	// @todo this function is doing nothing if a given fragment's initial value has already been created. This seems inconsistent with the other AddFragment functions (especially AddFragment_GetRef).
 	void AddFragment(FConstStructView Fragment)
 	{
 		const UScriptStruct* FragmentType = Fragment.GetScriptStruct();
 		checkf(UE::Mass::IsA<FMassFragment>(FragmentType), TEXT(MASS_INVALID_FRAGMENT_MSG_F), *GetNameSafe(FragmentType));
 
-		if (!Composition.GetElementsBitSet().Contains(*FragmentType))
+		if (!Composition.GetElementsBitSet().Contains(FragmentType))
 		{
-			Composition.GetElementsBitSet().Add(*FragmentType);
+			Composition.GetElementsBitSet().Add(FragmentType);
 			InitialFragmentValues.Emplace(Fragment);
 		}
 		else if (!InitialFragmentValues.ContainsByPredicate(FStructTypeEqualOperator(FragmentType)))
@@ -188,12 +195,18 @@ struct FMassEntityTemplateData
 		Composition.GetElementsBitSet().Add<T>();
 	}
 	
-	void AddTag(const UScriptStruct& TagType)
+	void AddTag(TNotNull<const UScriptStruct*> TagType)
 	{
-		checkf(UE::Mass::IsA<FMassTag>(&TagType), TEXT("Given struct doesn't represent a valid mass tag type. Make sure to inherit from FMassTag or one of its child-types."));
+		checkf(UE::Mass::IsA<FMassTag>(TagType), TEXT("Given struct doesn't represent a valid mass tag type. Make sure to inherit from FMassTag or one of its child-types."));
 		Composition.GetElementsBitSet().Add(TagType);
 	}
 
+	UE_DEPRECATED(5.8, "Use overload taking TNotNull")
+	void AddTag(const UScriptStruct& TagType)
+	{
+		AddTag(&TagType);
+	}
+
 	template<typename T>
 	void RemoveTag()
 	{
@@ -201,12 +214,18 @@ struct FMassEntityTemplateData
 		Composition.GetElementsBitSet().Remove<T>();
 	}
 
-	void RemoveTag(const UScriptStruct& TagType)
+	void RemoveTag(TNotNull<const UScriptStruct*> TagType)
 	{
-		checkf(UE::Mass::IsA<FMassTag>(&TagType), TEXT("Given struct doesn't represent a valid mass tag type. Make sure to inherit from FMassTag or one of its child-types."));
+		checkf(UE::Mass::IsA<FMassTag>(TagType), TEXT("Given struct doesn't represent a valid mass tag type. Make sure to inherit from FMassTag or one of its child-types."));
 		Composition.GetElementsBitSet().Remove(TagType);
 	}
 
+	UE_DEPRECATED(5.8, "Use overload taking TNotNull")
+	void RemoveTag(const UScriptStruct& TagType)
+	{
+		RemoveTag(&TagType);
+	}
+
 	FMassTagBitSet GetTags() const;
 	FMassElementBitSet& GetMutableTags();
 
@@ -222,9 +241,9 @@ struct FMassEntityTemplateData
 		const UScriptStruct* FragmentType = SharedFragment.GetScriptStruct();
 		if(ensureMsgf(UE::Mass::IsA<FMassConstSharedFragment>(FragmentType), TEXT("Given struct doesn't represent a valid const shared fragment type. Make sure to inherit from FMassConstSharedFragment or one of its child-types.")))
 		{
-			if (!Composition.GetElementsBitSet().Contains(*FragmentType))
+			if (!Composition.GetElementsBitSet().Contains(FragmentType))
 			{
-				Composition.GetElementsBitSet().Add(*FragmentType);
+				Composition.GetElementsBitSet().Add(FragmentType);
 				SharedFragmentValues.Add(SharedFragment);
 			}
 #if DO_ENSURE
@@ -243,9 +262,9 @@ struct FMassEntityTemplateData
 		const UScriptStruct* FragmentType = SharedFragment.GetScriptStruct();
 		if(ensureMsgf(UE::Mass::IsA<FMassSharedFragment>(FragmentType), TEXT("Given struct doesn't represent a valid shared fragment type. Make sure to inherit from FMassSharedFragment or one of its child-types.")))
 		{
-			if (!Composition.GetElementsBitSet().Contains(*FragmentType))
+			if (!Composition.GetElementsBitSet().Contains(FragmentType))
 			{
-				Composition.GetElementsBitSet().Add(*FragmentType);
+				Composition.GetElementsBitSet().Add(FragmentType);
 				SharedFragmentValues.Add(SharedFragment);
 			}
 	#if DO_ENSURE
@@ -265,12 +284,18 @@ struct FMassEntityTemplateData
 		return Composition.GetElementsBitSet().Contains<T>();
 	}
 	
-	bool HasFragment(const UScriptStruct& ScriptStruct) const
+	bool HasFragment(TNotNull<const UScriptStruct*> ScriptStruct) const
 	{
-		check(UE::Mass::IsA<FMassFragment>(&ScriptStruct))
+		check(UE::Mass::IsA<FMassFragment>(ScriptStruct))
 		return Composition.GetElementsBitSet().Contains(ScriptStruct);
 	}
 
+	UE_DEPRECATED(5.8, "Use overload taking TNotNull")
+	void HasFragment(const UScriptStruct& ScriptStruct)
+	{
+		HasFragment(&ScriptStruct);
+	}
+
 	template<UE::Mass::CTag T>
 	bool HasTag() const
 	{
@@ -289,24 +314,36 @@ struct FMassEntityTemplateData
 		return Composition.GetElementsBitSet().Contains<T>();
 	}
 
-	bool HasSharedFragment(const UScriptStruct& ScriptStruct) const
+	bool HasSharedFragment(TNotNull<const UScriptStruct*> ScriptStruct) const
 	{
-		check(UE::Mass::IsA<FMassSharedFragment>(&ScriptStruct))
+		check(UE::Mass::IsA<FMassSharedFragment>(ScriptStruct))
 		return Composition.GetElementsBitSet().Contains(ScriptStruct);
 	}
 
+	UE_DEPRECATED(5.8, "Use overload taking TNotNull")
+	void HasSharedFragment(const UScriptStruct& ScriptStruct)
+	{
+		HasSharedFragment(&ScriptStruct);
+	}
+
 	template<UE::Mass::CConstSharedFragment T>
 	bool HasConstSharedFragment() const
 	{
 		return Composition.GetElementsBitSet().Contains<T>();
 	}
 
-	bool HasConstSharedFragment(const UScriptStruct& ScriptStruct) const
+	bool HasConstSharedFragment(TNotNull<const UScriptStruct*> ScriptStruct) const
 	{
-		check(UE::Mass::IsA<FMassConstSharedFragment>(&ScriptStruct))
+		check(UE::Mass::IsA<FMassConstSharedFragment>(ScriptStruct))
 		return Composition.GetElementsBitSet().Contains(ScriptStruct);
 	}
 
+	UE_DEPRECATED(5.8, "Use overload taking TNotNull")
+	void HasConstSharedFragment(const UScriptStruct& ScriptStruct)
+	{
+		HasConstSharedFragment(&ScriptStruct);
+	}
+
 	void Sort()
 	{
 		SharedFragmentValues.Sort();
diff --git a/Engine/Plugins/Runtime/MassGameplay/Source/MassSpawner/Public/MassEntityTemplateRegistry.h b/Engine/Plugins/Runtime/MassGameplay/Source/MassSpawner/Public/MassEntityTemplateRegistry.h
--- a/Engine/Plugins/Runtime/MassGameplay/Source/MassSpawner/Public/MassEntityTemplateRegistry.h
+++ b/Engine/Plugins/Runtime/MassGameplay/Source/MassSpawner/Public/MassEntityTemplateRegistry.h
@@ -80,21 +80,21 @@ struct FMassEntityTemplateBuildContext
 	template<typename T>
 	T& AddFragment_GetRef()
 	{
-		TypeAdded(*T::StaticStruct());
+		TypeAdded(T::StaticStruct());
 		return TemplateData.AddFragment_GetRef<T>();
 	}
 
 	template<typename T>
 	void AddFragment()
 	{
-		TypeAdded(*T::StaticStruct());
+		TypeAdded(T::StaticStruct());
 		TemplateData.AddFragment<T>();
 	}
 
 	void AddFragment(FConstStructView InFragment)
 	{ 
 		checkf(InFragment.GetScriptStruct(), TEXT("Expecting a valid fragment type"));
-		TypeAdded(*InFragment.GetScriptStruct());
+		TypeAdded(InFragment.GetScriptStruct());
 		TemplateData.AddFragment(InFragment);
 	}
 
@@ -103,34 +103,40 @@ struct FMassEntityTemplateBuildContext
 	{
 		// Tags can be added by multiple traits, so they do not follow the same rules as fragments
 		TemplateData.AddTag<T>();
-		TypeAdded(*T::StaticStruct());
+		TypeAdded(T::StaticStruct());
 	}
 
-	void AddTag(const UScriptStruct& TagType)
+	void AddTag(TNotNull<const UScriptStruct*> TagType)
 	{
 		// Tags can be added by multiple traits, so they do not follow the same rules as fragments
 		TemplateData.AddTag(TagType);
 		TypeAdded(TagType);
 	}
 
+	UE_DEPRECATED(5.8, "Use overload taking TNotNull")
+	void AddTag(const UScriptStruct& TagType)
+	{
+		AddTag(&TagType);
+	}
+
 	template<typename T>
 	void AddChunkFragment()
 	{
-		TypeAdded(*T::StaticStruct());
+		TypeAdded(T::StaticStruct());
 		TemplateData.AddChunkFragment<T>();
 	}
 
 	void AddConstSharedFragment(const FConstSharedStruct& InSharedFragment)
 	{
 		checkf(InSharedFragment.GetScriptStruct(), TEXT("Expecting a valid shared fragment type"));
-		TypeAdded(*InSharedFragment.GetScriptStruct());
+		TypeAdded(InSharedFragment.GetScriptStruct());
 		TemplateData.AddConstSharedFragment(InSharedFragment);
 	}
 
 	void AddSharedFragment(const FSharedStruct& InSharedFragment)
 	{
 		checkf(InSharedFragment.GetScriptStruct(), TEXT("Expecting a valid shared fragment type"));
-		TypeAdded(*InSharedFragment.GetScriptStruct());
+		TypeAdded(InSharedFragment.GetScriptStruct());
 		TemplateData.AddSharedFragment(InSharedFragment);
 	}
 
@@ -140,20 +146,26 @@ struct FMassEntityTemplateBuildContext
 	 *		Using this functionality makes most sense for removing tags that specifically mean that entities having it are to be
 	 *		processed by a given processor.
 	 */
-	void RemoveTag(const UScriptStruct& TagType)
+	void RemoveTag(TNotNull<const UScriptStruct*> TagType)
 	{
-		checkf(UE::Mass::IsA<FMassTag>(&TagType), TEXT("Given struct doesn't represent a valid mass tag type. Make sure to inherit from FMassTag or one of its child-types."));
-		RemovedTypes.Add({&TagType
+		checkf(UE::Mass::IsA<FMassTag>(TagType), TEXT("Given struct doesn't represent a valid mass tag type. Make sure to inherit from FMassTag or one of its child-types."));
+		RemovedTypes.Add({TagType
 #if WITH_EDITORONLY_DATA
 			, TraitsData.Last().Trait
 #endif // WITH_EDITORONLY_DATA
 		});
 	}
 
+	UE_DEPRECATED(5.8, "Use overload taking TNotNull")
+	void RemoveTag(const UScriptStruct& TagType)
+	{
+		RemoveTag(&TagType);
+	}
+
 	template<typename T>
 	void RemoveTag()
 	{
-		RemoveTag(*T::StaticStruct());
+		RemoveTag(T::StaticStruct());
 	}
 
 	template<typename T>
@@ -169,12 +181,18 @@ struct FMassEntityTemplateBuildContext
 		return TemplateData.HasFragment<T>();
 	}
 	
-	bool HasFragment(const UScriptStruct& ScriptStruct) const
+	bool HasFragment(TNotNull<const UScriptStruct*> ScriptStruct) const
 	{
 		ENSURE_SUPPORTED_TRAIT_OPERATION();
 		return TemplateData.HasFragment(ScriptStruct);
 	}
 
+	UE_DEPRECATED(5.8, "Use overload taking TNotNull")
+	void HasFragment(const UScriptStruct& ScriptStruct)
+	{
+		HasFragment(&ScriptStruct);
+	}
+
 	template<typename T>
 	bool HasTag() const
 	{
@@ -195,12 +213,18 @@ struct FMassEntityTemplateBuildContext
 		return TemplateData.HasSharedFragment<T>();
 	}
 
-	bool HasSharedFragment(const UScriptStruct& ScriptStruct) const
+	bool HasSharedFragment(TNotNull<const UScriptStruct*> ScriptStruct) const
 	{
 		ENSURE_SUPPORTED_TRAIT_OPERATION();
 		return TemplateData.HasSharedFragment(ScriptStruct);
 	}
 
+	UE_DEPRECATED(5.8, "Use overload taking TNotNull")
+	void HasSharedFragment(const UScriptStruct& ScriptStruct)
+	{
+		HasSharedFragment(&ScriptStruct);
+	}
+
 	template<typename T>
 	bool HasConstSharedFragment() const
 	{
@@ -208,19 +232,25 @@ struct FMassEntityTemplateBuildContext
 		return TemplateData.HasConstSharedFragment<T>();
 	}
 
-	bool HasConstSharedFragment(const UScriptStruct& ScriptStruct) const
+	bool HasConstSharedFragment(TNotNull<const UScriptStruct*> ScriptStruct) const
 	{
 		ENSURE_SUPPORTED_TRAIT_OPERATION();
 		return TemplateData.HasConstSharedFragment(ScriptStruct);
 	}
 
+	UE_DEPRECATED(5.8, "Use overload taking TNotNull")
+	void HasConstSharedFragment(const UScriptStruct& ScriptStruct)
+	{
+		HasConstSharedFragment(&ScriptStruct);
+	}
+
 	//----------------------------------------------------------------------//
 	// Translators
 	//----------------------------------------------------------------------//
 	template<typename T>
 	void AddTranslator()
 	{
-		TypeAdded(*T::StaticClass());
+		TypeAdded(T::StaticClass());
 		TemplateData.GetMutableTags() += GetDefault<T>()->GetRequiredTags();
 	}
 
@@ -307,10 +337,16 @@ protected:
 	 */
 	bool ValidateBuildContext(const UWorld& World);
 
-	void TypeAdded(const UStruct& Type)
+	void TypeAdded(TNotNull<const UStruct*> Type)
 	{
 		checkf(TraitsData.Num(), TEXT("Adding elements to the build context before BuildFromTraits or SetTraitBeingProcessed was called is unsupported"));
-		TraitsData.Last().TypesAdded.Add(&Type);
+		TraitsData.Last().TypesAdded.Add(Type);
+	}
+
+	UE_DEPRECATED(5.8, "Use overload taking TNotNull")
+	void TypeAdded(const UScriptStruct& Type)
+	{
+		TypeAdded(&Type);
 	}
 
 	/** 
diff --git a/Engine/Source/Developer/MassEntityTestSuite/Private/MassChunkCollectionTest.cpp b/Engine/Source/Developer/MassEntityTestSuite/Private/MassChunkCollectionTest.cpp
--- a/Engine/Source/Developer/MassEntityTestSuite/Private/MassChunkCollectionTest.cpp
+++ b/Engine/Source/Developer/MassEntityTestSuite/Private/MassChunkCollectionTest.cpp
@@ -471,7 +471,7 @@ struct FArchetypeEntityCollection_ZeroLength_BatchMove : FEntityTestBase
 		AITEST_EQUAL("Entity creation generated expected number of entities", EntityCount, FMassArchetypeHelper::ArchetypeDataFromHandleChecked(SourceArchetype).GetNumEntities());
 		check(EntityCount == Entities.Num());
 
-		EntityManager->BatchChangeTagsForEntities(MakeArrayView(&Collection, 1), FMassTagBitSet(*AddedTag), {});
+		EntityManager->BatchChangeTagsForEntities(MakeArrayView(&Collection, 1), FMassTagBitSet(AddedTag), {});
 
 		const FMassArchetypeHandle NewArchetype = EntityManager->GetArchetypeForEntity(Entities[0]);
 		AITEST_NOT_EQUAL("Entities have changed their host archetype", NewArchetype, SourceArchetype);
@@ -502,7 +502,7 @@ struct FArchetypeEntityCollection_ZeroLength_BatchSetFragmentValues : FEntityTes
 
 		FMassArchetypeEntityCollectionWithPayload EntityCollectionWithEmptyPayload{FMassArchetypeEntityCollection(SourceArchetype)};
 
-		const FMassFragmentBitSet AddedFragmentsBitSet(*FTestFragment_Int::StaticStruct());
+		const FMassFragmentBitSet AddedFragmentsBitSet(FTestFragment_Int::StaticStruct());
 		EntityManager->BatchAddFragmentInstancesForEntities(MakeArrayView(&EntityCollectionWithEmptyPayload, 1), AddedFragmentsBitSet);
 
 		const FMassArchetypeHandle NewArchetype = EntityManager->GetArchetypeForEntity(Entities[0]);
diff --git a/Engine/Source/Developer/MassEntityTestSuite/Private/MassCommandsTest.cpp b/Engine/Source/Developer/MassEntityTestSuite/Private/MassCommandsTest.cpp
--- a/Engine/Source/Developer/MassEntityTestSuite/Private/MassCommandsTest.cpp
+++ b/Engine/Source/Developer/MassEntityTestSuite/Private/MassCommandsTest.cpp
@@ -326,7 +326,7 @@ struct FCommands_MoveHandleArrays : FEntityTestBase
 		
 		TArray<FMassArchetypeEntityCollection> EntityCollections;
 		UE::Mass::Utils::CreateEntityCollections(*EntityManager.Get(), Entities, FMassArchetypeEntityCollection::NoDuplicates, EntityCollections);
-		EntityManager->BatchChangeTagsForEntities(EntityCollections, FMassTagBitSet(*FTestTag_A::StaticStruct()), FMassTagBitSet());
+		EntityManager->BatchChangeTagsForEntities(EntityCollections, FMassTagBitSet(FTestTag_A::StaticStruct()), FMassTagBitSet());
 
 		// verify that original archetypes no longer host any entities
 		AITEST_TRUE("Original archetypes are empty after adding a tag to all entities"
diff --git a/Engine/Source/Developer/MassEntityTestSuite/Private/MassCreationContextTest.cpp b/Engine/Source/Developer/MassEntityTestSuite/Private/MassCreationContextTest.cpp
--- a/Engine/Source/Developer/MassEntityTestSuite/Private/MassCreationContextTest.cpp
+++ b/Engine/Source/Developer/MassEntityTestSuite/Private/MassCreationContextTest.cpp
@@ -97,15 +97,15 @@ struct FManualBuild : FEntityTestBase
 		TArray<FMassEntityHandle> Entities;
 		EntityManager->BatchReserveEntities(FloatEntitiesToSpawnCount, Entities);
 
-		FStructArrayView PaloadView(Payload);
+		FStructArrayView PayloadView(Payload);
 		TArray<FMassArchetypeEntityCollectionWithPayload> EntityCollections;
 		FMassArchetypeEntityCollectionWithPayload::CreateEntityRangesWithPayload(*EntityManager, Entities, FMassArchetypeEntityCollection::NoDuplicates
-			, FMassGenericPayloadView(MakeArrayView(&PaloadView, 1)), EntityCollections);
+			, FMassGenericPayloadView(MakeArrayView(&PayloadView, 1)), EntityCollections);
 
-		checkf(EntityCollections.Num() <= 1, TEXT("We expect TargetEntities to only contain archetype-less entities, ones that need to be built"));
+		checkf(EntityCollections.Num() == 1, TEXT("We expect TargetEntities to only contain archetype-less entities, ones that need to be built"));
 
 		{
-			TSharedRef<FMassEntityManager::FEntityCreationContext> CreationContext = EntityManager->BatchBuildEntities(EntityCollections[0], FMassFragmentBitSet(*FTestFragment_Float::StaticStruct()));
+			TSharedRef<FMassEntityManager::FEntityCreationContext> CreationContext = EntityManager->BatchBuildEntities(EntityCollections[0], FMassFragmentBitSet(FTestFragment_Float::StaticStruct()));
 			AITEST_EQUAL(TEXT("Creating entities should return the original context"), ObtainedContext, CreationContext);
 			++NumDifferentArchetypesUsed;
 		}
diff --git a/Engine/Source/Developer/MassEntityTestSuite/Private/MassEntityCollectionTest.cpp b/Engine/Source/Developer/MassEntityTestSuite/Private/MassEntityCollectionTest.cpp
--- a/Engine/Source/Developer/MassEntityTestSuite/Private/MassEntityCollectionTest.cpp
+++ b/Engine/Source/Developer/MassEntityTestSuite/Private/MassEntityCollectionTest.cpp
@@ -244,7 +244,7 @@ struct FCreateWithOutdatedCollections : FEntityCollectionTestBase
 		AITEST_TRUE("Initially the created archetype collection is up to date", InitialCollection.IsUpToDate());
 
 		// we now move the entities to another archetype, to force InitialCollection's invalidation
-		EntityManager->BatchChangeFragmentCompositionForEntities(MakeArrayView(&InitialCollection, 1), FMassFragmentBitSet(*FTestFragment_Int::StaticStruct()), {});
+		EntityManager->BatchChangeFragmentCompositionForEntities(MakeArrayView(&InitialCollection, 1), FMassFragmentBitSet(FTestFragment_Int::StaticStruct()), {});
 
 		AITEST_FALSE("(NOT) After entities are moved to another archetype the InitialCollection is up to date", InitialCollection.IsUpToDate());
 		{
diff --git a/Engine/Source/Developer/MassEntityTestSuite/Private/MassObserverTest.cpp b/Engine/Source/Developer/MassEntityTestSuite/Private/MassObserverTest.cpp
--- a/Engine/Source/Developer/MassEntityTestSuite/Private/MassObserverTest.cpp
+++ b/Engine/Source/Developer/MassEntityTestSuite/Private/MassObserverTest.cpp
@@ -357,11 +357,11 @@ struct FEntityCreation_Batched : FTagBaseOperation
 		{
 			TSharedRef<FMassEntityManager::FEntityCreationContext> CreationContext = EntityManager->BatchCreateEntities(IntsArchetype, EntitiesToSpawnCount, EntitiesInt);
 
-			EntityManager->BatchChangeTagsForEntities(CreationContext->GetEntityCollections(*EntityManager.Get()), FMassTagBitSet(*FTagStruct::StaticStruct()), FMassTagBitSet());
+			EntityManager->BatchChangeTagsForEntities(CreationContext->GetEntityCollections(*EntityManager.Get()), FMassTagBitSet(FTagStruct::StaticStruct()), FMassTagBitSet());
 			AITEST_TRUE(TEXT("The tag observer is not expected to run yet"), AffectedEntities.Num() == 0);
 			AITEST_FALSE(TEXT("CreationContext's entity collection should be invalidated at this moment"), CreationContext->DebugAreEntityCollectionsUpToDate());
 
-			EntityManager->BatchChangeTagsForEntities(CreationContext->GetEntityCollections(*EntityManager.Get()), FMassTagBitSet(*FTagStruct::StaticStruct()), FMassTagBitSet());
+			EntityManager->BatchChangeTagsForEntities(CreationContext->GetEntityCollections(*EntityManager.Get()), FMassTagBitSet(FTagStruct::StaticStruct()), FMassTagBitSet());
 			AITEST_TRUE(TEXT("The tag observer is still not expected to run"), AffectedEntities.Num() == 0);
 		}
 		AITEST_TRUE(TEXT("The tag observer is expected to run just after FEntityCreationContext's destruction"), AffectedEntities.Num() > 0);
@@ -662,7 +662,7 @@ struct FObserverChangingComposition_Sync : FFragmentTestBase
 			{
 				FMassArchetypeEntityCollection EntityCollection(OriginalArchetype, Context.GetEntities(), FMassArchetypeEntityCollection::NoDuplicates);
 				EntityManager->BatchChangeFragmentCompositionForEntities(MakeArrayView(&EntityCollection, 1)
-					, FMassFragmentBitSet(*FTestFragment_Int::StaticStruct()), {});
+					, FMassFragmentBitSet(FTestFragment_Int::StaticStruct()), {});
 			};
 
 			EntityManager->BatchCreateEntities(OriginalArchetype, EntitiesToSpawn, EntitiesInt);
@@ -734,7 +734,7 @@ struct FModificationsWhileCreationContextActive : FTagBaseOperation
 				// we're testing only the behavior related to the previously created entities here
 			}
 			EntityManager->BatchChangeTagsForEntities(MakeArrayView(&InitialEntitiesCollection, 1)
-				, FMassTagBitSet(*FTagStruct::StaticStruct()), FMassTagBitSet());
+				, FMassTagBitSet(FTagStruct::StaticStruct()), FMassTagBitSet());
 
 			AITEST_TRUE(TEXT("The tag observer is not expected to run yet"), AffectedEntities.Num() == 0);
 		}
@@ -814,7 +814,7 @@ struct FCreationOperationOrder_Batch : FCreationOperationOrder
 		}
 		{
 			FMassArchetypeEntityCollection Collection(IntsArchetype, MakeArrayView(&PreExistingEntities[0], 1), FMassArchetypeEntityCollection::NoDuplicates);
-			EntityManager->BatchChangeFragmentCompositionForEntities(MakeArrayView(&Collection, 1), FMassFragmentBitSet(*FTestFragment_Float::StaticStruct()), {});
+			EntityManager->BatchChangeFragmentCompositionForEntities(MakeArrayView(&Collection, 1), FMassFragmentBitSet(FTestFragment_Float::StaticStruct()), {});
 		}
 		{
 			// creating two separate entities, that should end up in the same creation context
@@ -823,7 +823,7 @@ struct FCreationOperationOrder_Batch : FCreationOperationOrder
 		}
 		{
 			FMassArchetypeEntityCollection Collection(IntsArchetype, MakeArrayView(&PreExistingEntities[1], 1), FMassArchetypeEntityCollection::NoDuplicates);
-			EntityManager->BatchChangeFragmentCompositionForEntities(MakeArrayView(&Collection, 1), FMassFragmentBitSet(*FTestFragment_Float::StaticStruct()), {});
+			EntityManager->BatchChangeFragmentCompositionForEntities(MakeArrayView(&Collection, 1), FMassFragmentBitSet(FTestFragment_Float::StaticStruct()), {});
 		}
 	}
 };
@@ -870,7 +870,7 @@ struct FCreatedEntitiesUnobservedCompositionChange : FFragmentTestBase
 			TSharedRef<FMassEntityManager::FEntityCreationContext> CreationContext = EntityManager->BatchCreateEntities(FloatsArchetype, EntitiesToSpawnCount, EntitiesFloats);
 
 			// add unobserved tag
-			EntityManager->BatchChangeTagsForEntities(CreationContext->GetEntityCollections(*EntityManager.Get()), FMassTagBitSet(*FTestTag_A::StaticStruct()), FMassTagBitSet());
+			EntityManager->BatchChangeTagsForEntities(CreationContext->GetEntityCollections(*EntityManager.Get()), FMassTagBitSet(FTestTag_A::StaticStruct()), FMassTagBitSet());
 
 			AITEST_EQUAL(TEXT("The fragment observer is not expected to run yet"), AffectedEntities.Num(), 0);
 		}
diff --git a/Engine/Source/Developer/MassEntityTestSuite/Private/MassQueryTest.cpp b/Engine/Source/Developer/MassEntityTestSuite/Private/MassQueryTest.cpp
--- a/Engine/Source/Developer/MassEntityTestSuite/Private/MassQueryTest.cpp
+++ b/Engine/Source/Developer/MassEntityTestSuite/Private/MassQueryTest.cpp
@@ -924,7 +924,7 @@ struct FQueryTest_ExportHandles : FEntityTestBase
 
 		TArray<FMassArchetypeEntityCollection> EntityCollections;
 		UE::Mass::Utils::CreateEntityCollections(*EntityManager.Get(), Entities, FMassArchetypeEntityCollection::NoDuplicates, EntityCollections);
-		EntityManager->BatchChangeTagsForEntities(EntityCollections, FMassTagBitSet(*FTestTag_A::StaticStruct()), FMassTagBitSet());
+		EntityManager->BatchChangeTagsForEntities(EntityCollections, FMassTagBitSet(FTestTag_A::StaticStruct()), FMassTagBitSet());
 
 		FMassEntityQuery Query(EntityManager.ToSharedRef());
 		Query.AddTagRequirement<FTestTag_A>(EMassFragmentPresence::All);
diff --git a/Engine/Source/Developer/MassEntityTestSuite/Private/MassSharedFragmentsTest.cpp b/Engine/Source/Developer/MassEntityTestSuite/Private/MassSharedFragmentsTest.cpp
--- a/Engine/Source/Developer/MassEntityTestSuite/Private/MassSharedFragmentsTest.cpp
+++ b/Engine/Source/Developer/MassEntityTestSuite/Private/MassSharedFragmentsTest.cpp
@@ -611,7 +611,7 @@ struct FConstSharedFragment_RemoveFromEntity : FEntityTestBase
 		AITEST_NULL("Initially the entity is not expected to have the shared fragment", EntitySharedFragment);
 
 		AITEST_FALSE("Attempt to remove shared fragment from entity that doesn't have shared fragment should return false and do nothing",
-			EntityManager->RemoveConstSharedFragmentFromEntity(EntityHandle, *FTestConstSharedFragment_Int::StaticStruct()));
+			EntityManager->RemoveConstSharedFragmentFromEntity(EntityHandle, FTestConstSharedFragment_Int::StaticStruct()));
 
 		AITEST_TRUE("Adding shared fragment to entity should succeed",
 			EntityManager->AddConstSharedFragmentToEntity(EntityHandle, SharedFragmentInstance));
@@ -622,7 +622,7 @@ struct FConstSharedFragment_RemoveFromEntity : FEntityTestBase
 
 		AITEST_TRUE(
 			"Removing shared fragment from entity that has the shared fragment should succeed",
-			EntityManager->RemoveConstSharedFragmentFromEntity(EntityHandle, *FTestConstSharedFragment_Int::StaticStruct()));
+			EntityManager->RemoveConstSharedFragmentFromEntity(EntityHandle, FTestConstSharedFragment_Int::StaticStruct()));
 
 		EntitySharedFragment = EntityManager->GetConstSharedFragmentDataPtr<FTestConstSharedFragment_Int>(EntityHandle);
 		AITEST_NULL("The entity is not expected to have the shared fragment after the operation", EntitySharedFragment);
@@ -648,7 +648,7 @@ struct FSharedFragment_RemoveFromEntity : FEntityTestBase
 		AITEST_NULL("Initially the entity is not expected to have the shared fragment", EntitySharedFragment);
 
 		AITEST_FALSE("Attempt to remove shared fragment from entity that doesn't have shared fragment should return false and do nothing",
-			EntityManager->RemoveSharedFragmentFromEntity(EntityHandle, *FTestSharedFragment_Int::StaticStruct()));
+			EntityManager->RemoveSharedFragmentFromEntity(EntityHandle, FTestSharedFragment_Int::StaticStruct()));
 
 		AITEST_TRUE("Adding shared fragment to entity should succeed",
 			EntityManager->AddSharedFragmentToEntity(EntityHandle, SharedFragmentInstance));
@@ -659,7 +659,7 @@ struct FSharedFragment_RemoveFromEntity : FEntityTestBase
 
 		AITEST_TRUE(
 			"Removing shared fragment from entity that has the shared fragment should succeed",
-			EntityManager->RemoveSharedFragmentFromEntity(EntityHandle, *FTestSharedFragment_Int::StaticStruct()));
+			EntityManager->RemoveSharedFragmentFromEntity(EntityHandle, FTestSharedFragment_Int::StaticStruct()));
 
 		EntitySharedFragment = EntityManager->GetSharedFragmentDataPtr<FTestSharedFragment_Int>(EntityHandle);
 		AITEST_NULL("The entity is not expected to have the shared fragment after the operation", EntitySharedFragment);
diff --git a/Engine/Source/Developer/MassEntityTestSuite/Private/MassSparseElementsTest.cpp b/Engine/Source/Developer/MassEntityTestSuite/Private/MassSparseElementsTest.cpp
--- a/Engine/Source/Developer/MassEntityTestSuite/Private/MassSparseElementsTest.cpp
+++ b/Engine/Source/Developer/MassEntityTestSuite/Private/MassSparseElementsTest.cpp
@@ -244,7 +244,7 @@ namespace UE::Mass::Test::SparseElements
 			}
 
 			InitialCollection = FMassArchetypeEntityCollection(FEntityTestBase::IntsArchetype, EntitiesCreated, FMassArchetypeEntityCollection::NoDuplicates);
-			EntityManager->BatchChangeFragmentCompositionForEntities({InitialCollection}, FMassFragmentBitSet(*FTestFragment_Float::StaticStruct()), {});
+			EntityManager->BatchChangeFragmentCompositionForEntities({InitialCollection}, FMassFragmentBitSet(FTestFragment_Float::StaticStruct()), {});
 
 			AITEST_TRUE(TEXT("The original archetype no longer has the element")
 				, TTestOperations::HasSparseElement(OriginalArchetypeHandle, SparseElementType) == false);
diff --git a/Engine/Source/Developer/MassEntityTestSuite/Private/MassTypeBitSetBuilderTest.cpp b/Engine/Source/Developer/MassEntityTestSuite/Private/MassTypeBitSetBuilderTest.cpp
--- a/Engine/Source/Developer/MassEntityTestSuite/Private/MassTypeBitSetBuilderTest.cpp
+++ b/Engine/Source/Developer/MassEntityTestSuite/Private/MassTypeBitSetBuilderTest.cpp
@@ -63,8 +63,8 @@ namespace UE::Mass::Test::BitSetBuilder
 			BitSetBuilder.template Add<FTestElement_B>();
 
 			typename Traits::FBitSetReader BitSetReader = BitSetBuilder;
-			AITEST_EQUAL("The reading result is the same regardless of the method, existing element", BitSetReader.template Contains<FTestElement_B>(), BitSetReader.Contains(*FTestElement_B::StaticStruct()));
-			AITEST_EQUAL("The reading result is the same regardless of the method, non-existing element", BitSetReader.template Contains<FTestElement_A>(), BitSetReader.Contains(*FTestElement_A::StaticStruct()));
+			AITEST_EQUAL("The reading result is the same regardless of the method, existing element", BitSetReader.template Contains<FTestElement_B>(), BitSetReader.Contains(FTestElement_B::StaticStruct()));
+			AITEST_EQUAL("The reading result is the same regardless of the method, non-existing element", BitSetReader.template Contains<FTestElement_A>(), BitSetReader.Contains(FTestElement_A::StaticStruct()));
 			AITEST_TRUE("The resulting bitset has the right bit set", BitSetReader.template Contains<FTestElement_B>());
 			AITEST_FALSE("(NOT) The resulting bitset has the wrong bit set", BitSetReader.template Contains<FTestElement_A>());
 
diff --git a/Engine/Source/Editor/MassEntityDebugger/Private/MassDebuggerModel.cpp b/Engine/Source/Editor/MassEntityDebugger/Private/MassDebuggerModel.cpp
--- a/Engine/Source/Editor/MassEntityDebugger/Private/MassDebuggerModel.cpp
+++ b/Engine/Source/Editor/MassEntityDebugger/Private/MassDebuggerModel.cpp
@@ -1758,7 +1758,7 @@ namespace UE::MassDebugger
 
 		for (const TSharedPtr<FMassDebuggerFragmentData>& TagData : DebuggerModel.CachedTagData)
 		{
-			const UScriptStruct& StructType = *TagData->Fragment.Get();
+			const UScriptStruct* StructType = TagData->Fragment.Get();
 			EMassFragmentPresence Presence = EMassFragmentPresence::All;
 
 			if (AllTags.Contains(StructType))
@@ -1783,7 +1783,7 @@ namespace UE::MassDebugger
 			}
 
 			TSharedPtr<FFragmentEntry> Entry = MakeShared<FFragmentEntry>();
-			Entry->StructType = &StructType;
+			Entry->StructType = StructType;
 			Entry->Presence = Presence;
 			TagRequirements.Add(Entry);
 		}
diff --git a/Engine/Source/Runtime/CoreUObject/Private/StructUtils/StructTypeBitSet.cpp b/Engine/Source/Runtime/CoreUObject/Private/StructUtils/StructTypeBitSet.cpp
--- a/Engine/Source/Runtime/CoreUObject/Private/StructUtils/StructTypeBitSet.cpp
+++ b/Engine/Source/Runtime/CoreUObject/Private/StructUtils/StructTypeBitSet.cpp
@@ -85,10 +85,10 @@ FORCEINLINE const UStruct* FStructTracker::GetBaseType() const
 	return BaseType ? BaseType : (BaseStructGetter ? (BaseType = BaseStructGetter()) : nullptr);
 }
 
-int32 FStructTracker::FindOrAddStructTypeIndex(const UStruct& InStructType)
+int32 FStructTracker::FindOrAddStructTypeIndex(TNotNull<const UStruct*> InStructType)
 {
 	// Get existing index...
-	const uint32 Hash = PointerHash(&InStructType);
+	const uint32 Hash = PointerHash(InStructType);
 	const FSetElementId ElementId = StructTypeToIndexSet.FindIdByHash(Hash, Hash);
 
 	if (LIKELY(ElementId.IsValidId()))
@@ -98,16 +98,16 @@ int32 FStructTracker::FindOrAddStructTypeIndex(const UStruct& InStructType)
 	return RegisterImplementation(InStructType, /*bCheckExisting=*/false);
 }
 
-int32 FStructTracker::RegisterImplementation(const UStruct& InStructType, const bool bCheckExisting)
+int32 FStructTracker::RegisterImplementation(TNotNull<const UStruct*> InStructType, const bool bCheckExisting)
 {
-	if (!ensureMsgf(TypeVerification(&InStructType)
+	if (!ensureMsgf(TypeVerification(InStructType)
 		, TEXT("Unable to register %s singe it fails the validation test.")
-		, *InStructType.GetName()))
+		, *InStructType->GetName()))
 	{
 		return INDEX_NONE;
 	}
 
-	const uint32 Hash = PointerHash(&InStructType);
+	const uint32 Hash = PointerHash(InStructType);
 	if (bCheckExisting)
 	{
 		const FSetElementId ElementId = StructTypeToIndexSet.FindIdByHash(Hash, Hash);
@@ -123,9 +123,9 @@ int32 FStructTracker::RegisterImplementation(const UStruct& InStructType, const
 
 	const int32 NewIndex = ElementId.AsInteger();
 	check(StructTypesList.Num() == NewIndex);
-	StructTypesList.Add(&InStructType);
+	StructTypesList.Add(InStructType);
 
-	OnTypeRegistered.Broadcast(&InStructType, NewIndex);
+	OnTypeRegistered.Broadcast(InStructType, NewIndex);
 
 	if (bIsSerializable)
 	{
@@ -137,7 +137,7 @@ int32 FStructTracker::RegisterImplementation(const UStruct& InStructType, const
 			CA_ASSUME(LocalBaseType);
 			SerializationHash = GetTypeHash(LocalBaseType->GetFullName());
 		}
-		SerializationHash = HashCombine(SerializationHash, GetTypeHash(InStructType.GetFullName()));
+		SerializationHash = HashCombine(SerializationHash, GetTypeHash(InStructType->GetFullName()));
 
 		// it's worth pointing out that we're registering a given tracker for all the hashes created along the way.
 		// This will help with loading bitsets from serialized data.
@@ -145,7 +145,7 @@ int32 FStructTracker::RegisterImplementation(const UStruct& InStructType, const
 	}
 
 #if WITH_STRUCTUTILS_DEBUG
-	DebugStructTypeNamesList.Add(InStructType.GetFName());
+	DebugStructTypeNamesList.Add(InStructType->GetFName());
 	ensure(StructTypeToIndexSet.Num() == DebugStructTypeNamesList.Num());
 #endif // WITH_STRUCTUTILS_DEBUG
 
@@ -153,9 +153,9 @@ int32 FStructTracker::RegisterImplementation(const UStruct& InStructType, const
 }
 
 
-int32 FStructTracker::FindStructTypeIndex(const UStruct& InStructType) const
+int32 FStructTracker::FindStructTypeIndex(TNotNull<const UStruct*> InStructType) const
 {
-	const uint32 Hash = PointerHash(&InStructType);
+	const uint32 Hash = PointerHash(InStructType);
 	const FSetElementId ElementId = StructTypeToIndexSet.FindIdByHash(Hash, Hash);
 
 	return ElementId.IsValidId() ? ElementId.AsInteger() : INDEX_NONE;
@@ -258,7 +258,7 @@ void FStructTracker::Serialize(FArchive& Ar, FStructTypeBitSet::FBitSetContainer
 						if (UStruct* StructType = FindObject<UStruct>(TypePath))
 						{
 							checkSlow(StructType->IsChildOf(GetBaseType()));
-							BitMapping[Index] = FindOrAddStructTypeIndex(*StructType);
+							BitMapping[Index] = FindOrAddStructTypeIndex(StructType);
 						}
 						else
 						{
diff --git a/Engine/Source/Runtime/CoreUObject/Public/StructUtils/StructTypeBitSet.h b/Engine/Source/Runtime/CoreUObject/Public/StructUtils/StructTypeBitSet.h
--- a/Engine/Source/Runtime/CoreUObject/Public/StructUtils/StructTypeBitSet.h
+++ b/Engine/Source/Runtime/CoreUObject/Public/StructUtils/StructTypeBitSet.h
@@ -401,23 +401,41 @@ struct FStructTracker
 	/** 
 	 * @return the index for the given type, or INDEX_NONE if it hasn't been registered with this struct tracker
 	 */
-	UE_API int32 FindStructTypeIndex(const UStruct& InStructType) const;
+	UE_API int32 FindStructTypeIndex(TNotNull<const UStruct*> InStructType) const;
+
+	UE_DEPRECATED(5.8, "Use overload taking TNotNull")
+	int32 FindStructTypeIndex(const UStruct& InStructType) const
+	{
+		return FindStructTypeIndex(&InStructType);
+	}
 
 	/**
 	 * Fetches the internal index representing the given UStruct instance. If it hasn't been registered yet, it will be automatically added.
 	 * @return the index for the given type.
 	 */
-	UE_API int32 FindOrAddStructTypeIndex(const UStruct& InStructType);
+	UE_API int32 FindOrAddStructTypeIndex(TNotNull<const UStruct*> InStructType);
+
+	UE_DEPRECATED(5.8, "Use overload taking TNotNull")
+	int32 FindOrAddStructTypeIndex(const UStruct& InStructType)
+	{
+		return FindOrAddStructTypeIndex(&InStructType);
+	}
 
 	/**
 	 * Registers the given UStruct with the struct tracker.
 	 * @return the index for the registered type.
 	 */
-	int32 Register(const UStruct& InStructType)
+	int32 Register(TNotNull<const UStruct*> InStructType)
 	{
 		return RegisterImplementation(InStructType, /*bCheckPrevious=*/true);
 	}
 
+	UE_DEPRECATED(5.8, "Use overload taking TNotNull")
+	int32 Register(const UStruct& InStructType)
+	{
+		return Register(&InStructType);
+	}
+
 	/**
 	 * Retrieves the UStruct type associated with a given index
 	 * @param StructTypeIndex - The index of the struct type
@@ -481,7 +499,13 @@ protected:
 #endif // WITH_STRUCTUTILS_DEBUG
 
 protected:
-	UE_API int32 RegisterImplementation(const UStruct& InStructType, const bool bCheckPrevious);
+	UE_API int32 RegisterImplementation(TNotNull<const UStruct*> InStructType, const bool bCheckPrevious);
+
+	UE_DEPRECATED(5.8, "Use overload taking TNotNull")
+	int32 RegisterImplementation(const UStruct& InStructType, const bool bCheckPrevious)
+	{
+		return RegisterImplementation(&InStructType, bCheckPrevious);
+	}
 
 	/** Set mapping struct type hashes to their indices */
 	TSet<uint32> StructTypeToIndexSet;
@@ -506,14 +530,13 @@ private:
  *
  * @param TImplementation - The derived type implementing this base class.
  * @param TBaseStruct - The base struct type for type validation.
- * @param TStructType - Unreal's struct type, typically UScriptStruct or UClass.
+ * @param TUStructType - Unreal's struct type, typically UScriptStruct or UClass.
  * @param TBitSetContainer - The container type for storing bitsets.
  * @param bTestInheritanceAtRuntime - Flag to enable runtime inheritance checks.
  */
-template<typename TImplementation, typename TBaseStruct, typename TStructType, typename TBitSetContainer, bool bTestInheritanceAtRuntime=WITH_STRUCTUTILS_DEBUG>
+template<typename TImplementation, typename TBaseStruct, typename TUStructType, typename TBitSetContainer, bool bTestInheritanceAtRuntime=WITH_STRUCTUTILS_DEBUG>
 struct TTypeBitSetBase
 {
-	using FUStructType = TStructType;
 	using FBaseStruct = TBaseStruct;
 	using FContainer = TBitSetContainer;
 	using FBitArrayBase = FStructTypeBitSet::FConstBitSetContainer::FBitArrayBase;
@@ -620,13 +643,13 @@ struct TTypeBitSetBase
 	 * Adds a struct type to the bitset.
 	 * @param InStructType - The struct type to add.
 	 */
-	void Add(const FUStructType& InStructType)
+	void Add(TNotNull<const TUStructType*> InStructType)
 	{
 		if constexpr (bTestInheritanceAtRuntime)
 		{
-			if (UNLIKELY(!ensureMsgf(InStructType.IsChildOf(GetBaseUStruct()),
+			if (UNLIKELY(!ensureMsgf(InStructType->IsChildOf(GetBaseUStruct()),
 				TEXT("Registering '%s' with FStructTracker while it doesn't derive from the expected struct type %s"),
-				*InStructType.GetPathName(), *GetBaseUStruct()->GetName())))
+				*InStructType->GetPathName(), *GetBaseUStruct()->GetName())))
 			{
 				return;
 			}
@@ -636,6 +659,12 @@ struct TTypeBitSetBase
 		StructTypesBitArray.AddAtIndex(StructTypeIndex);
 	}
 
+	UE_DEPRECATED(5.8, "Use overload taking TNotNull")
+	void Add(const TUStructType& InStructType)
+	{
+		Add(&InStructType);
+	}
+
 	/** Sets the bit at index StructTypeIndex */
 	void AddAtIndex(const int32 StructTypeIndex)
 	{
@@ -646,13 +675,13 @@ struct TTypeBitSetBase
 	 * Removes a struct type from the bitset.
 	 * @param InStructType - The struct type to remove.
 	 */
-	void Remove(const FUStructType& InStructType)
+	void Remove(TNotNull<const TUStructType*> InStructType)
 	{
 		if constexpr (bTestInheritanceAtRuntime)
 		{
-			if (UNLIKELY(!ensureMsgf(InStructType.IsChildOf(GetBaseUStruct()),
+			if (UNLIKELY(!ensureMsgf(InStructType->IsChildOf(GetBaseUStruct()),
 				TEXT("Registering '%s' with FStructTracker while it doesn't derive from the expected struct type %s"),
-				*InStructType.GetPathName(), *GetBaseUStruct()->GetName())))
+				*InStructType->GetPathName(), *GetBaseUStruct()->GetName())))
 			{
 				return;
 			}
@@ -662,6 +691,12 @@ struct TTypeBitSetBase
 		StructTypesBitArray.RemoveAtIndex(StructTypeIndex);
 	}
 
+	UE_DEPRECATED(5.8, "Use overload taking TNotNull")
+	void Remove(const TUStructType& InStructType)
+	{
+		Remove(&InStructType);
+	}
+
 	/** Clears the bit at index StructTypeIndex */
 	void RemoveAtIndex(const int32 StructTypeIndex)
 	{
@@ -681,13 +716,13 @@ struct TTypeBitSetBase
 	 * @param InStructType - The struct type to check for.
 	 * @return True if the struct type is in the bitset; false otherwise.
 	 */
-	bool Contains(const FUStructType& InStructType) const
+	bool Contains(TNotNull<const TUStructType*> InStructType) const
 	{
 		if constexpr (bTestInheritanceAtRuntime)
 		{
-			if (UNLIKELY(!ensureMsgf(InStructType.IsChildOf(GetBaseUStruct()),
+			if (UNLIKELY(!ensureMsgf(InStructType->IsChildOf(GetBaseUStruct()),
 				TEXT("Registering '%s' with FStructTracker while it doesn't derive from the expected struct type %s"),
-				*InStructType.GetPathName(), *GetBaseUStruct()->GetName())))
+				*InStructType->GetPathName(), *GetBaseUStruct()->GetName())))
 			{
 				return false;
 			}
@@ -697,6 +732,12 @@ struct TTypeBitSetBase
 		return StructTypeIndex != INDEX_NONE && StructTypesBitArray.Contains(StructTypeIndex);
 	}
 
+	UE_DEPRECATED(5.8, "Use overload taking TNotNull")
+	void Contains(const TUStructType& InStructType)
+	{
+		Contains(&InStructType);
+	}
+
 	/**
 	 * Performs a bitwise AND operation with another bitset.
 	 * @param Other - The other bitset to intersect with.
@@ -830,25 +871,37 @@ struct TTypeBitSetBase
 	 * @param NewElement - The struct type to add.
 	 * @return A new bitset with the added struct type.
 	 */
-	TImplementation operator+(const FUStructType& NewElement) const
+	TImplementation operator+(TNotNull<const TUStructType*> NewElement) const
 	{
 		TImplementation Result = GetImplementation();
 		Result.Add(NewElement);
 		return MoveTemp(Result);
 	}
 
+	UE_DEPRECATED(5.8, "Use overload taking TNotNull")
+	TImplementation operator+(const TUStructType& NewElement) const
+	{
+		return operator+(&NewElement);
+	}
+
 	/**
 	 * Removes a struct type from the bitset, returning a new bitset.
 	 * @param NewElement - The struct type to remove.
 	 * @return A new bitset with the struct type removed.
 	 */
-	TImplementation operator-(const FUStructType& NewElement) const
+	TImplementation operator-(TNotNull<const TUStructType*> NewElement) const
 	{
 		TImplementation Result = GetImplementation();
 		Result.Remove(NewElement);
 		return MoveTemp(Result);
 	}
 
+	UE_DEPRECATED(5.8, "Use overload taking TNotNull")
+	TImplementation operator-(const TUStructType& NewElement) const
+	{
+		return operator-(&NewElement);
+	}
+
 	/**
 	 * Counts the number of set bits in the bitset.
 	 * @return The number of set bits.
@@ -980,7 +1033,6 @@ struct TStructTypeBitSet : TTypeBitSetBase<TStructTypeBitSet<TBaseStruct, TStruc
 {
 	using Super = TTypeBitSetBase<TStructTypeBitSet, TBaseStruct, TUStructType, FStructTypeBitSet::FBitSetContainer>;
 	using FStructTrackerWrapper = TStructTrackerWrapper;
-	using FUStructType = TUStructType;
 	using FBaseStruct = TBaseStruct;
 	using FBitSetContainer = FStructTypeBitSet::FBitSetContainer;
 	using FBitArrayBase = FStructTypeBitSet::FConstBitSetContainer::FBitArrayBase;
@@ -1007,22 +1059,29 @@ struct TStructTypeBitSet : TTypeBitSetBase<TStructTypeBitSet<TBaseStruct, TStruc
 	 * Constructor that initializes the bitset with a single struct type.
 	 * @param StructType - The struct type to add.
 	 */
-	explicit TStructTypeBitSet(const FUStructType& StructType)
+	explicit TStructTypeBitSet(TNotNull<const TUStructType*> StructType)
 	{
 		Add(StructType);
 	}
 
+	UE_DEPRECATED(5.8, "Use overload taking TNotNull")
+	explicit TStructTypeBitSet(const TUStructType& StructType)
+		: TStructTypeBitSet(&StructType)
+	{
+	}
+
 	/**
 	 * Constructor that initializes the bitset with an initializer list of struct types.
 	 * @param InitList - The initializer list of struct type pointers.
 	 */
-	explicit TStructTypeBitSet(std::initializer_list<const FUStructType*> InitList)
+	explicit TStructTypeBitSet(std::initializer_list<const TUStructType*> InitList)
+
 	{
-		for (const FUStructType* StructType : InitList)
+		for (const TUStructType* StructType : InitList)
 		{
 			if (StructType)
 			{
-				Add(*StructType);
+				Add(StructType);
 			}
 		}
 	}
@@ -1031,13 +1090,13 @@ struct TStructTypeBitSet : TTypeBitSetBase<TStructTypeBitSet<TBaseStruct, TStruc
 	 * Constructor that initializes the bitset with an array view of struct types.
 	 * @param InitList - The array view of struct type pointers.
 	 */
-	explicit TStructTypeBitSet(TConstArrayView<const FUStructType*> InitList)
+	explicit TStructTypeBitSet(TConstArrayView<const TUStructType*> InitList)
 	{
-		for (const FUStructType* StructType : InitList)
+		for (const TUStructType* StructType : InitList)
 		{
 			if (StructType)
 			{
-				Add(*StructType);
+				Add(StructType);
 			}
 		}
 	}
@@ -1054,7 +1113,7 @@ struct TStructTypeBitSet : TTypeBitSetBase<TStructTypeBitSet<TBaseStruct, TStruc
 		{
 			if (StructInstance.GetScriptStruct())
 			{
-				Add(*StructInstance.GetScriptStruct());
+				Add(StructInstance.GetScriptStruct());
 			}
 		}
 	}
@@ -1104,17 +1163,25 @@ public:
 	 * @param InStructType - The struct type to get the index for.
 	 * @return The index of the struct type.
 	 */
-	static int32 GetTypeIndex(const FUStructType& InStructType)
+	static int32 GetTypeIndex(TNotNull<const TUStructType*> InStructType)
+
 	{
 #if WITH_STRUCTUTILS_DEBUG
-		ensureMsgf(InStructType.IsChildOf(GetBaseUStruct())
+		ensureMsgf(InStructType->IsChildOf(GetBaseUStruct())
 			, TEXT("Creating index for '%s' while it doesn't derive from the expected struct type %s")
-			, *InStructType.GetPathName(), *GetBaseUStruct()->GetName());
+			, *InStructType->GetPathName(), *GetBaseUStruct()->GetName());
 #endif // WITH_STRUCTUTILS_DEBUG
 
 		return TStructTrackerWrapper::StructTracker.FindOrAddStructTypeIndex(InStructType);
 	}
 
+	UE_DEPRECATED(5.8, "Use overload taking TNotNull")
+	static int32 GetTypeIndex(const TUStructType& InStructType)
+	{
+		return GetTypeIndex(&InStructType);
+	}
+
+
 	/**
 	 * Gets the index of a struct type within the tracker, adding it if not already present.
 	 * @param T - The struct type to get the index for.
@@ -1124,7 +1191,7 @@ public:
 	static int32 GetTypeIndex()
 	{
 		static_assert(TIsDerivedFrom<T, TBaseStruct>::IsDerived, "Given struct type doesn't match the expected base struct type.");
-		static const int32 TypeIndex = GetTypeIndex(*UE::StructUtils::GetAsUStruct<T>());
+		static const int32 TypeIndex = GetTypeIndex(UE::StructUtils::GetAsUStruct<T>());
 		return TypeIndex;
 	}
 	
@@ -1146,9 +1213,9 @@ public:
 	 * @param Index - The index of the struct type.
 	 * @return Pointer to the struct type if found; nullptr otherwise.
 	 */
-	static const FUStructType* GetTypeAtIndex(const int32 Index)
+	static const TUStructType* GetTypeAtIndex(const int32 Index)
 	{
-		return Cast<const FUStructType>(TStructTrackerWrapper::StructTracker.GetStructType(Index));
+		return Cast<const TUStructType>(TStructTrackerWrapper::StructTracker.GetStructType(Index));
 	}
 
 	/**
@@ -1256,7 +1323,7 @@ public:
 	 * @todo To be improved.
 	 * @param Callback - The callback function to invoke for each type.
 	 */
-	void ExportTypes(TFunctionRef<bool(const FUStructType*)> Callback) const
+	void ExportTypes(TFunctionRef<bool(const TUStructType*)> Callback) const
 	{
 		FBitArrayBase::FConstIterator It(Super::StructTypesBitArray);
 		bool bKeepGoing = true;
diff --git a/Engine/Source/Runtime/MassEntity/Internal/MassArchetypeData.h b/Engine/Source/Runtime/MassEntity/Internal/MassArchetypeData.h
--- a/Engine/Source/Runtime/MassEntity/Internal/MassArchetypeData.h
+++ b/Engine/Source/Runtime/MassEntity/Internal/MassArchetypeData.h
@@ -61,7 +61,7 @@ namespace UE::Mass
 		/** @return whether anything in the chunk potentially has the element */
 		bool Contains(TNotNull<const UScriptStruct*> ElementType) const
 		{
-			return SparseElementsPresent.Contains(*ElementType);
+			return SparseElementsPresent.Contains(ElementType);
 		}
 
 		bool HasElementsInRange(int32 IndexWithinChunk, int32 Num) const
@@ -438,7 +438,7 @@ public:
 	bool HasFragmentType(const UScriptStruct* FragmentType) const
 	{
 		return UE::Mass::IsA<FMassFragment>(FragmentType) 
-			&& CompositionBitSet.Contains(*FragmentType);
+			&& CompositionBitSet.Contains(FragmentType);
 	}
 
 	/**
@@ -447,22 +447,20 @@ public:
 	bool HasTagType(const UScriptStruct* TagType) const
 	{
 		return UE::Mass::IsA<FMassTag>(TagType)
-			&& CompositionBitSet.Contains(*TagType);
+			&& CompositionBitSet.Contains(TagType);
 	}
 
 	/**
 	 * checks whether a given element type is part of this archetype's entities' composition
-	 * @param IncludeSparseElements whether the function should check if the archetype hosts any entities
-	 *	owning such a fragment, provided FragmentType represents a sparse element type
 	 */
 	bool HasElement(TNotNull<const UScriptStruct*> ElementType) const
 	{
-		return CompositionBitSet.Contains(*ElementType);
+		return CompositionBitSet.Contains(ElementType);
 	}
 
 	bool DoesContainEntitiesWithSparseElement(const UScriptStruct* ElementType) const
 	{
-		return UE::Mass::IsSparse(ElementType) && GetSparseElementsBitSet().Contains(*ElementType);
+		return UE::Mass::IsSparse(ElementType) && GetSparseElementsBitSet().Contains(ElementType);
 	}
 
 	bool IsEquivalent(const FMassElementBitSet& OtherCompositionBitSet, const UE::Mass::FArchetypeGroups& OtherGroups) const;
@@ -932,7 +930,7 @@ inline bool FMassArchetypeData::ContainsAnySparseData() const
 
 inline bool FMassArchetypeData::ContainsSparseElement(TNotNull<const UScriptStruct*> ElementType) const
 {
-	return GetSparseElementsBitSet().Contains(*ElementType);
+	return GetSparseElementsBitSet().Contains(ElementType);
 }
 
 inline const FMassElementBitSet& FMassArchetypeData::GetSparseElementsBitSet() const
diff --git a/Engine/Source/Runtime/MassEntity/Private/MassArchetypeData.cpp b/Engine/Source/Runtime/MassEntity/Private/MassArchetypeData.cpp
--- a/Engine/Source/Runtime/MassEntity/Private/MassArchetypeData.cpp
+++ b/Engine/Source/Runtime/MassEntity/Private/MassArchetypeData.cpp
@@ -36,7 +36,7 @@ namespace UE::Mass
 	{
 		check(IsEmpty() == false);
 
-		const int32 TypeIndex = FMassElementBitSet::GetTypeIndex(*ElementType);
+		const int32 TypeIndex = FMassElementBitSet::GetTypeIndex(ElementType);
 
 		FMassElementBitSet& EntitySparseElements = PerEntityElements[EntityInChunkHandle.IndexWithinChunk];
 		if (EntitySparseElements.IsBitSet(TypeIndex) == false)
@@ -57,7 +57,7 @@ namespace UE::Mass
 
 		FMassElementBitSet& EntitySparseElements = PerEntityElements[EntityInChunkHandle.IndexWithinChunk];
 
-		const int32 TypeIndex = FMassElementBitSet::GetTypeIndex(*ElementType);
+		const int32 TypeIndex = FMassElementBitSet::GetTypeIndex(ElementType);
 		if (EntitySparseElements.IsBitSet(TypeIndex))
 		{
 			EntitySparseElements.RemoveAtIndex(TypeIndex);
@@ -92,7 +92,7 @@ namespace UE::Mass
 		check(SubchunkStart >= 0 && SubchunkStart + RangeLength <= PerEntityElements.Num());
 
 		uint16 NumAdded = 0;
-		const int32 TypeIndex = FMassElementBitSet::GetTypeIndex(*ElementType);
+		const int32 TypeIndex = FMassElementBitSet::GetTypeIndex(ElementType);
 		for (int32 EntityIndex = SubchunkStart; EntityIndex < SubchunkStart + RangeLength; ++EntityIndex)
 		{
 			NumAdded += static_cast<int>(PerEntityElements[EntityIndex].IsBitSet(TypeIndex) == false);
@@ -119,7 +119,7 @@ namespace UE::Mass
 		check(SubchunkStart >= 0 && SubchunkStart + RangeLength <= PerEntityElements.Num());
 
 		uint16 NumRemoved = 0;
-		const int32 TypeIndex = FMassElementBitSet::GetTypeIndex(*ElementType);
+		const int32 TypeIndex = FMassElementBitSet::GetTypeIndex(ElementType);
 		for (int32 EntityIndex = SubchunkStart; EntityIndex < SubchunkStart + RangeLength; ++EntityIndex)
 		{
 			NumRemoved += static_cast<int>(PerEntityElements[EntityIndex].IsBitSet(TypeIndex));
@@ -144,7 +144,7 @@ namespace UE::Mass
 	bool FChunkSparseElements::Contains(const FMassRawEntityInChunkData& EntityInChunkHandle, TNotNull<const UScriptStruct*> ElementType) const
 	{
 		return PerEntityElements.IsValidIndex(EntityInChunkHandle.IndexWithinChunk)
-			&& PerEntityElements[EntityInChunkHandle.IndexWithinChunk].Contains(*ElementType);
+			&& PerEntityElements[EntityInChunkHandle.IndexWithinChunk].Contains(ElementType);
 	}
 
 	bool FChunkSparseElements::MoveElementsFrom(FChunkSparseElements& OriginalSparseElements, const int32 OriginalIndexWithinChunk, const int32 NewIndexWithinChunk, const int32 Num)
@@ -699,7 +699,7 @@ void FMassArchetypeData::SetSharedFragmentsData(const FMassEntityHandle Entity,
 			UE::Mass::FChunkSparseElements& OldSparseElements = OldChunk.GetSparseElementsUnsafe();
 			UE::Mass::FChunkSparseElements& NewSparseElements = NewChunk.GetSparseElements(NumEntitiesPerChunk);
 			// note that we don't have to update bSparseElementsBitSetDirty since we're moving elements
-			// between chunks of the same archeype so the archetype's cumulative CachedSparseElementsBitSet 
+			// between chunks of the same archetype so the archetype's cumulative CachedSparseElementsBitSet 
 			// won't be affected
 			NewSparseElements.MoveElementsFrom(OldSparseElements, EntityInChunkHandle.IndexWithinChunk, NewIndexWithinChunk, 1);
 		}
diff --git a/Engine/Source/Runtime/MassEntity/Private/MassEntityBuilder.cpp b/Engine/Source/Runtime/MassEntity/Private/MassEntityBuilder.cpp
--- a/Engine/Source/Runtime/MassEntity/Private/MassEntityBuilder.cpp
+++ b/Engine/Source/Runtime/MassEntity/Private/MassEntityBuilder.cpp
@@ -95,7 +95,7 @@ namespace UE::Mass
 			{
 				if (Mass::IsA<TElement>(Element.GetScriptStruct()))
 				{
-					if (Bitset.Contains(*Element.GetScriptStruct()) == false)
+					if (Bitset.Contains(Element.GetScriptStruct()) == false)
 					{
 						bIssuesFound = true;
 						UE_VLOG_UELOG(LogOwner, LogMass, Error, TEXT("%hs: input Composition doesn't contain %s"), __FUNCTION__, *GetNameSafe(Element.GetScriptStruct()));
@@ -575,7 +575,7 @@ bool FEntityBuilder::HandleTypeInstance(const UScriptStruct* Type, TArray<FInsta
 	if (UE::Mass::IsA<T>(Type))
 	{
 		auto& BitSet = Composition.GetElementsBitSet();
-		const int32 TypeIndex = BitSet.GetTypeIndex(*Type);
+		const int32 TypeIndex = BitSet.GetTypeIndex(Type);
 		OutTargetContainer = &GetInstancedStructContainerInternal<T>();
 		if (BitSet.IsBitSet(TypeIndex) == false)
 		{
@@ -642,7 +642,7 @@ FEntityBuilder& FEntityBuilder::Add(TNotNull<const UScriptStruct*> ElementType)
 {
 	if (UE::Mass::IsA<FMassElement>(ElementType))
 	{
-		Composition.GetElementsBitSet().Add(*ElementType);
+		Composition.GetElementsBitSet().Add(ElementType);
 		State = EState::ReadyToCommit;
 		CachedArchetypeHandle = FMassArchetypeHandle();
 	}
diff --git a/Engine/Source/Runtime/MassEntity/Private/MassEntityManager.cpp b/Engine/Source/Runtime/MassEntity/Private/MassEntityManager.cpp
--- a/Engine/Source/Runtime/MassEntity/Private/MassEntityManager.cpp
+++ b/Engine/Source/Runtime/MassEntity/Private/MassEntityManager.cpp
@@ -219,35 +219,35 @@ void FMassEntityManager::Initialize(const FMassEntityManagerStorageInitParams& I
 		{
 			if (*StructIt != FMassFragment::StaticStruct())
 			{
-				Fragments.Add(**StructIt);
+				Fragments.Add(*StructIt);
 			}
 		}
 		else if (UE::Mass::IsA<FMassTag>(*StructIt))
 		{
 			if (*StructIt != FMassTag::StaticStruct())
 			{
-				Tags.Add(**StructIt);
+				Tags.Add(*StructIt);
 			}
 		}
 		else if (UE::Mass::IsA<FMassChunkFragment>(*StructIt))
 		{
 			if (*StructIt != FMassChunkFragment::StaticStruct())
 			{
-				ChunkFragments.Add(**StructIt);
+				ChunkFragments.Add(*StructIt);
 			}
 		}
 		else if (UE::Mass::IsA<FMassSharedFragment>(*StructIt))
 		{
 			if (*StructIt != FMassSharedFragment::StaticStruct())
 			{
-				LocalSharedFragments.Add(**StructIt);
+				LocalSharedFragments.Add(*StructIt);
 			}
 		}
 		else if (UE::Mass::IsA<FMassConstSharedFragment>(*StructIt))
 		{
 			if (*StructIt != FMassConstSharedFragment::StaticStruct())
 			{
-				LocalConstSharedFragments.Add(**StructIt);
+				LocalConstSharedFragments.Add(*StructIt);
 			}
 		}
 	}
@@ -1082,26 +1082,24 @@ const UE::Mass::FArchetypeGroups& FMassEntityManager::GetGroupsForArchetype(cons
 	return ArchetypeHandle.DataPtr->GetGroups();
 }
 
-void FMassEntityManager::AddFragmentToEntity(FMassEntityHandle Entity, const UScriptStruct* FragmentType)
+void FMassEntityManager::AddFragmentToEntity(FMassEntityHandle Entity, TNotNull<const UScriptStruct*> FragmentType)
 {
 	CHECK_SYNC_API();
-	checkf(FragmentType, TEXT("Null fragment type passed in to %hs"), __FUNCTION__);
 
 	CheckIfEntityIsActive(Entity);
 
-	const FMassArchetypeCompositionDescriptor Descriptor(InternalAddFragmentListToEntityChecked(Entity, FMassFragmentBitSet(*FragmentType)));
+	const FMassArchetypeCompositionDescriptor Descriptor(InternalAddFragmentListToEntityChecked(Entity, FMassFragmentBitSet(FragmentType)));
 
 	ObserverManager.OnPostCompositionAdded(Entity, Descriptor);
 }
 
-void FMassEntityManager::AddFragmentToEntity(FMassEntityHandle Entity, const UScriptStruct* FragmentType, const FStructInitializationCallback& Initializer)
+void FMassEntityManager::AddFragmentToEntity(FMassEntityHandle Entity, TNotNull<const UScriptStruct*> FragmentType, const FStructInitializationCallback& Initializer)
 {
 	CHECK_SYNC_API();
-	checkf(FragmentType, TEXT("Null fragment type passed in to %hs"), __FUNCTION__);
 
 	CheckIfEntityIsActive(Entity);
 
-	FMassFragmentBitSet Fragments = InternalAddFragmentListToEntityChecked(Entity, FMassFragmentBitSet(*FragmentType));
+	FMassFragmentBitSet Fragments = InternalAddFragmentListToEntityChecked(Entity, FMassFragmentBitSet(FragmentType));
 	FMassArchetypeData* CurrentArchetype = GetEntityStorageInterface().GetArchetype(Entity.Index);
 	check(CurrentArchetype);
 	void* FragmentData = CurrentArchetype->GetFragmentDataForEntity(FragmentType, Entity.Index);
@@ -1342,15 +1340,15 @@ void FMassEntityManager::SwapTagsForEntity(FMassEntityHandle Entity, TNotNull<co
 	check(CurrentArchetype);
 
 	FMassElementBitSet NewCompositionBitSet = CurrentArchetype->GetCompositionBitSet();
-	NewCompositionBitSet.Remove(*OldTagType);
-	NewCompositionBitSet.Add(*NewTagType);
+	NewCompositionBitSet.Remove(OldTagType);
+	NewCompositionBitSet.Add(NewTagType);
 
 	ensureMsgf(NewCompositionBitSet.IsEquivalent(CurrentArchetype->GetCompositionBitSet()) == false, TEXT("Failing this test at this point is unexpected."));
 
 	const FMassArchetypeHandle NewArchetypeHandle = InternalCreateSimilarArchetype(CurrentArchetype, MoveTemp(NewCompositionBitSet));
 	checkSlow(NewArchetypeHandle.IsValid());
 
-	ObserverManager.OnPreCompositionRemoved(Entity, FMassElementBitSet(*OldTagType));
+	ObserverManager.OnPreCompositionRemoved(Entity, FMassElementBitSet(OldTagType));
 
 	// Move the entity over
 	CurrentArchetype->MoveEntityToAnotherArchetype(Entity, *NewArchetypeHandle.DataPtr.Get());
@@ -1358,7 +1356,7 @@ void FMassEntityManager::SwapTagsForEntity(FMassEntityHandle Entity, TNotNull<co
 	// Change the entity archetype
 	GetEntityStorageInterface().SetArchetypeFromShared(Entity.Index, NewArchetypeHandle.DataPtr);
 
-	ObserverManager.OnPostCompositionAdded(Entity, FMassElementBitSet(*NewTagType));
+	ObserverManager.OnPostCompositionAdded(Entity, FMassElementBitSet(NewTagType));
 }
 
 void FMassEntityManager::AddTagToEntity(FMassEntityHandle Entity, TNotNull<const UScriptStruct*> TagType)
@@ -1376,7 +1374,7 @@ void FMassEntityManager::AddTagToEntity(FMassEntityHandle Entity, TNotNull<const
 		if (CurrentArchetype->HasTagType(TagType) == false)
 		{
 			FMassElementBitSet NewCompositionBitSet = CurrentArchetype->GetCompositionBitSet();
-			NewCompositionBitSet.Add(*TagType);
+			NewCompositionBitSet.Add(TagType);
 			const FMassArchetypeHandle NewArchetypeHandle = InternalCreateSimilarArchetype(CurrentArchetype, MoveTemp(NewCompositionBitSet));
 			checkSlow(NewArchetypeHandle.IsValid());
 
@@ -1384,7 +1382,7 @@ void FMassEntityManager::AddTagToEntity(FMassEntityHandle Entity, TNotNull<const
 			CurrentArchetype->MoveEntityToAnotherArchetype(Entity, *NewArchetypeHandle.DataPtr.Get());
 			GetEntityStorageInterface().SetArchetypeFromShared(Entity.Index, NewArchetypeHandle.DataPtr);
 
-			ObserverManager.OnPostCompositionAdded(Entity, FMassElementBitSet(*TagType));
+			ObserverManager.OnPostCompositionAdded(Entity, FMassElementBitSet(TagType));
 		}
 	}
 	else
@@ -1407,7 +1405,7 @@ void FMassEntityManager::RemoveTagFromEntity(FMassEntityHandle Entity, TNotNull<
 
 		if (CurrentArchetype->HasTagType(TagType))
 		{
-			FMassElementBitSet TagDelta(*TagType);
+			FMassElementBitSet TagDelta(TagType);
 			FMassElementBitSet NewComposition = CurrentArchetype->GetCompositionBitSet() - TagDelta;
 
 			ObserverManager.OnPreCompositionRemoved(Entity, MoveTemp(TagDelta));
@@ -1482,7 +1480,7 @@ bool FMassEntityManager::DoesEntityHaveElement(FMassEntityHandle Entity, TNotNul
 		return CurrentArchetype->HasSparseElementForEntity(ElementType, Entity);
 	}
 	
-	return CurrentArchetype->GetCompositionBitSet().Contains(*ElementType);
+	return CurrentArchetype->GetCompositionBitSet().Contains(ElementType);
 }
 
 void FMassEntityManager::BatchAddElementToEntities(TConstArrayView<FMassEntityHandle> Entities, TNotNull<const UScriptStruct*> ElementType)
@@ -1498,11 +1496,11 @@ void FMassEntityManager::BatchAddElementToEntities(TConstArrayView<FMassEntityHa
 	}
 	else if (UE::Mass::IsA<FMassTag>(ElementType))
 	{
-		BatchChangeTagsForEntities(EntityCollections, FMassTagBitSet(*ElementType), {});
+		BatchChangeTagsForEntities(EntityCollections, FMassTagBitSet(ElementType), {});
 	}
 	else
 	{
-		BatchChangeFragmentCompositionForEntities(EntityCollections, FMassFragmentBitSet(*ElementType), {});
+		BatchChangeFragmentCompositionForEntities(EntityCollections, FMassFragmentBitSet(ElementType), {});
 	}
 }
 
@@ -1539,11 +1537,11 @@ void FMassEntityManager::BatchRemoveElementFromEntities(TConstArrayView<FMassEnt
 	}
 	else if (UE::Mass::IsA<FMassTag>(ElementType))
 	{
-		BatchChangeTagsForEntities(EntityCollections, {}, FMassTagBitSet(*ElementType));
+		BatchChangeTagsForEntities(EntityCollections, {}, FMassTagBitSet(ElementType));
 	}
 	else
 	{
-		BatchChangeFragmentCompositionForEntities(EntityCollections, {}, FMassFragmentBitSet(*ElementType));
+		BatchChangeFragmentCompositionForEntities(EntityCollections, {}, FMassFragmentBitSet(ElementType));
 	}
 }
 
@@ -1584,7 +1582,7 @@ bool FMassEntityManager::SwapConstSharedFragmentForEntity(const FMassEntityHandl
 	check(CurrentArchetype);
 
 	const UScriptStruct* FragmentType = InConstSharedFragment.GetScriptStruct();
-	if (!CurrentArchetype->GetCompositionBitSet().Contains(*FragmentType))
+	if (!CurrentArchetype->GetCompositionBitSet().Contains(FragmentType))
 	{
 		// Observers will get notify by the call to AddConstSharedFragmentToEntity
 		AddConstSharedFragmentToEntity(Entity, InConstSharedFragment);
@@ -1592,11 +1590,11 @@ bool FMassEntityManager::SwapConstSharedFragmentForEntity(const FMassEntityHandl
 	else
 	{
 		// SetSharedFragmentsData won't notify observers so we have to do it here to allow observers to handle the change
-		ObserverManager.OnPreCompositionRemoved(Entity, FMassArchetypeCompositionDescriptor(FMassElementBitSet{*FragmentType}));
+		ObserverManager.OnPreCompositionRemoved(Entity, FMassArchetypeCompositionDescriptor(FMassElementBitSet{FragmentType}));
 
 		CurrentArchetype->SetSharedFragmentsData(Entity, MakeArrayView<FConstSharedStruct>({ InConstSharedFragment }));
 
-		ObserverManager.OnPostCompositionAdded(Entity, FMassElementBitSet(*FragmentType));
+		ObserverManager.OnPostCompositionAdded(Entity, FMassElementBitSet(FragmentType));
 	}
 
 	return true;
@@ -1654,16 +1652,16 @@ bool FMassEntityManager::AddConstSharedFragmentToEntity(const FMassEntityHandle
 	// Change the entity archetype
 	GetEntityStorageInterface().SetArchetypeFromShared(Entity.Index, NewArchetypeHandle.DataPtr);
 
-	ObserverManager.OnPostCompositionAdded(Entity, FMassElementBitSet(*StructType));
+	ObserverManager.OnPostCompositionAdded(Entity, FMassElementBitSet(StructType));
 
 	return true;
 }
 
-bool FMassEntityManager::RemoveConstSharedFragmentFromEntity(const FMassEntityHandle Entity, const UScriptStruct& ConstSharedFragmentType)
+bool FMassEntityManager::RemoveConstSharedFragmentFromEntity(const FMassEntityHandle Entity, TNotNull<const UScriptStruct*> ConstSharedFragmentType)
 {
 	CHECK_SYNC_API_RETURN(return false);
 
-	if (!ensureMsgf(UE::Mass::IsA<FMassConstSharedFragment>(&ConstSharedFragmentType), TEXT("%hs parameter ConstSharedFragmentType is expected to be a FMassConstSharedFragment"), __FUNCTION__))
+	if (!ensureMsgf(UE::Mass::IsA<FMassConstSharedFragment>(ConstSharedFragmentType), TEXT("%hs parameter ConstSharedFragmentType is expected to be a FMassConstSharedFragment"), __FUNCTION__))
 	{
 		return false;
 	}
@@ -1673,7 +1671,7 @@ bool FMassEntityManager::RemoveConstSharedFragmentFromEntity(const FMassEntityHa
 	FMassArchetypeData* CurrentArchetype = GetEntityStorageInterface().GetArchetypeAsShared(Entity.Index).Get();
 	check(CurrentArchetype);
 	
-	if (!CurrentArchetype->GetCompositionDescriptor().Contains(&ConstSharedFragmentType))
+	if (!CurrentArchetype->GetCompositionDescriptor().Contains(ConstSharedFragmentType))
 	{
 		// Nothing to do. Returning false to indicate nothing has been removed, as per function's documentation 
 		return false;
@@ -1687,7 +1685,7 @@ bool FMassEntityManager::RemoveConstSharedFragmentFromEntity(const FMassEntityHa
 	check(NewArchetype);
 	
 	const FMassArchetypeSharedFragmentValues& OldSharedFragmentValues = CurrentArchetype->GetSharedFragmentValues(Entity.Index);
-	check(OldSharedFragmentValues.ContainsType(&ConstSharedFragmentType));
+	check(OldSharedFragmentValues.ContainsType(ConstSharedFragmentType));
 	FMassArchetypeSharedFragmentValues NewSharedFragmentValues(OldSharedFragmentValues);
 	
 	const FMassConstSharedFragmentBitSet ToRemove(ConstSharedFragmentType);
@@ -1725,7 +1723,7 @@ bool FMassEntityManager::AddSharedFragmentToEntity(const FMassEntityHandle Entit
 
 	const UScriptStruct* StructType = InSharedFragment.GetScriptStruct();
 	CA_ASSUME(StructType);
-	if (CurrentArchetype->GetCompositionBitSet().Contains(*StructType))
+	if (CurrentArchetype->GetCompositionBitSet().Contains(StructType))
 	{
 		const FMassArchetypeSharedFragmentValues& SharedFragmentValues = CurrentArchetype->GetSharedFragmentValues(Entity);
 		FConstSharedStruct ExistingSharedStruct = SharedFragmentValues.GetSharedFragmentStruct(StructType);
@@ -1739,7 +1737,7 @@ bool FMassEntityManager::AddSharedFragmentToEntity(const FMassEntityHandle Entit
 	}
 	
 	FMassArchetypeCompositionDescriptor NewComposition(CurrentArchetype->GetCompositionDescriptor());
-	NewComposition.GetElementsBitSet().Add(*StructType);
+	NewComposition.GetElementsBitSet().Add(StructType);
 	const FMassArchetypeHandle NewArchetypeHandle = CreateArchetype(NewComposition, FMassArchetypeCreationParams(*CurrentArchetype));
 	check(NewArchetypeHandle.IsValid());
 	FMassArchetypeData* NewArchetype = NewArchetypeHandle.DataPtr.Get();
@@ -1757,16 +1755,16 @@ bool FMassEntityManager::AddSharedFragmentToEntity(const FMassEntityHandle Entit
 	// Change the entity archetype
 	GetEntityStorageInterface().SetArchetypeFromShared(Entity.Index, NewArchetypeHandle.DataPtr);
 
-	ObserverManager.OnPostCompositionAdded(Entity, FMassElementBitSet(*StructType));
+	ObserverManager.OnPostCompositionAdded(Entity, FMassElementBitSet(StructType));
 
 	return true;
 }
 
-bool FMassEntityManager::RemoveSharedFragmentFromEntity(const FMassEntityHandle Entity, const UScriptStruct& SharedFragmentType)
+bool FMassEntityManager::RemoveSharedFragmentFromEntity(const FMassEntityHandle Entity, TNotNull<const UScriptStruct*> SharedFragmentType)
 {
 	CHECK_SYNC_API_RETURN(return false);
 
-	if (!ensureMsgf(UE::Mass::IsA<FMassSharedFragment>(&SharedFragmentType), TEXT("%hs parameter SharedFragmentType is expected to be a FMassSharedFragment"), __FUNCTION__))
+	if (!ensureMsgf(UE::Mass::IsA<FMassSharedFragment>(SharedFragmentType), TEXT("%hs parameter SharedFragmentType is expected to be a FMassSharedFragment"), __FUNCTION__))
 	{
 		return false;
 	}
@@ -1790,7 +1788,7 @@ bool FMassEntityManager::RemoveSharedFragmentFromEntity(const FMassEntityHandle
 	check(NewArchetype);
 
 	const FMassArchetypeSharedFragmentValues& OldSharedFragmentValues = CurrentArchetype->GetSharedFragmentValues(Entity.Index);
-	check(OldSharedFragmentValues.ContainsType(&SharedFragmentType));
+	check(OldSharedFragmentValues.ContainsType(SharedFragmentType));
 	FMassArchetypeSharedFragmentValues NewSharedFragmentValues(OldSharedFragmentValues);
 
 	const FMassSharedFragmentBitSet ToRemove(SharedFragmentType);
diff --git a/Engine/Source/Runtime/MassEntity/Private/MassEntityTypes.cpp b/Engine/Source/Runtime/MassEntity/Private/MassEntityTypes.cpp
--- a/Engine/Source/Runtime/MassEntity/Private/MassEntityTypes.cpp
+++ b/Engine/Source/Runtime/MassEntity/Private/MassEntityTypes.cpp
@@ -85,7 +85,7 @@ FMassArchetypeSharedFragmentValues FMassArchetypeSharedFragmentValues::CreateCom
 	auto RemovePredicate = [&NewSharedElementComposition](const FConstSharedStruct& SharedStruct)
 		{
 			return SharedStruct.GetScriptStruct() == nullptr
-				|| NewSharedElementComposition.Contains(*SharedStruct.GetScriptStruct()) == false;
+				|| NewSharedElementComposition.Contains(SharedStruct.GetScriptStruct()) == false;
 		};
 
 	NewSharedValues.SharedFragments.RemoveAllSwap(RemovePredicate);
@@ -110,7 +110,7 @@ FConstSharedStruct FMassArchetypeSharedFragmentValues::Add_GetRef(const FConstSh
 	}
 
 	check(StructType);
-	StoredElementsBitSet.Add(*StructType);
+	StoredElementsBitSet.Add(StructType);
 	FConstSharedStruct& StructInstance = ConstSharedFragments.Add_GetRef(Fragment);
 	DirtyHashCache();
 	return StructInstance;
@@ -130,7 +130,7 @@ FSharedStruct FMassArchetypeSharedFragmentValues::Add_GetRef(const FSharedStruct
 	}
 
 	check(StructType);
-	StoredElementsBitSet.Add(*StructType);
+	StoredElementsBitSet.Add(StructType);
 	FSharedStruct& StructInstance = SharedFragments.Add_GetRef(Fragment);
 	DirtyHashCache();
 	return StructInstance;
@@ -234,7 +234,7 @@ int32 FMassArchetypeSharedFragmentValues::Append(const FMassArchetypeSharedFragm
 	{
 		const UScriptStruct* StructType = SharedStruct.GetScriptStruct();
 		check(StructType);
-		if (StoredElementsBitSet.Contains(*StructType))
+		if (StoredElementsBitSet.Contains(StructType))
 		{
 			const int32 FragmentIndex = SharedFragments.IndexOfByPredicate(FStructTypeEqualOperator(StructType));
 			checkf(FragmentIndex != INDEX_NONE, TEXT("Mismatch between shared fragment bitset and stored values"));
@@ -252,7 +252,7 @@ int32 FMassArchetypeSharedFragmentValues::Append(const FMassArchetypeSharedFragm
 	{
 		const UScriptStruct* StructType = SharedStruct.GetScriptStruct();
 		check(StructType);
-		if (StoredElementsBitSet.Contains(*StructType))
+		if (StoredElementsBitSet.Contains(StructType))
 		{
 			const int32 FragmentIndex = ConstSharedFragments.IndexOfByPredicate(FStructTypeEqualOperator(StructType));
 			checkf(FragmentIndex != INDEX_NONE, TEXT("Mismatch between const shared fragment bitset and stored values"));
diff --git a/Engine/Source/Runtime/MassEntity/Private/MassObserverManager.cpp b/Engine/Source/Runtime/MassEntity/Private/MassObserverManager.cpp
--- a/Engine/Source/Runtime/MassEntity/Private/MassObserverManager.cpp
+++ b/Engine/Source/Runtime/MassEntity/Private/MassObserverManager.cpp
@@ -76,7 +76,7 @@ namespace UE::Mass::ObserverManager
 					Pipeline.SortByExecutionPriority();
 					Pipeline.Initialize(Owner, EntityManager.AsShared());
 					
-					InOutObservedBitSet.Add(*StructType);
+					InOutObservedBitSet.Add(StructType);
 				}
 			}
 		};
@@ -493,7 +493,7 @@ void FMassObserverManager::AddObserverInstance(TNotNull<const UScriptStruct*> El
 	{
 		if (!!(OperationFlags & static_cast<EMassObservedOperationFlags>(1 << OperationIndex)))
 		{
-			ObservedElements[OperationIndex].Add(*ElementType);
+			ObservedElements[OperationIndex].Add(ElementType);
 
 			FMassRuntimePipeline& Pipeline = (*ElementObservers[OperationIndex]).FindOrAdd(ElementType);
 			if (Pipeline.AppendUniqueProcessor(*ObserverProcessor) && !bInitializeCalled)
@@ -542,7 +542,7 @@ void FMassObserverManager::RemoveObserverInstance(TNotNull<const UScriptStruct*>
 				if (Pipeline->Num() == 0)
 				{
 					(*ElementObservers[OperationIndex]).Remove(ElementType);
-					ObservedElements[OperationIndex].Remove(*ElementType);
+					ObservedElements[OperationIndex].Remove(ElementType);
 				}
 			}
 		}
diff --git a/Engine/Source/Runtime/MassEntity/Private/MassProcessingQueueTypes.cpp b/Engine/Source/Runtime/MassEntity/Private/MassProcessingQueueTypes.cpp
--- a/Engine/Source/Runtime/MassEntity/Private/MassProcessingQueueTypes.cpp
+++ b/Engine/Source/Runtime/MassEntity/Private/MassProcessingQueueTypes.cpp
@@ -403,11 +403,11 @@ namespace UE::Mass
 						switch (Description.AccessMode)
 						{
 						case EMassFragmentAccess::ReadWrite:
-							DirectWritesElements.Add(*Description.StructType);
-							DirectReadsElements.Add(*Description.StructType);
+							DirectWritesElements.Add(Description.StructType);
+							DirectReadsElements.Add(Description.StructType);
 							break;
 						case EMassFragmentAccess::ReadOnly:
-							DirectReadsElements.Add(*Description.StructType);
+							DirectReadsElements.Add(Description.StructType);
 							break;
 						default:
 							break;
diff --git a/Engine/Source/Runtime/MassEntity/Private/MassProcessorDependencySolver.cpp b/Engine/Source/Runtime/MassEntity/Private/MassProcessorDependencySolver.cpp
--- a/Engine/Source/Runtime/MassEntity/Private/MassProcessorDependencySolver.cpp
+++ b/Engine/Source/Runtime/MassEntity/Private/MassProcessorDependencySolver.cpp
@@ -1087,7 +1087,7 @@ void FMassProcessorDependencySolver::GatherSubsystemInformation(const UE::Mass::
 			{
 				const UClass* SubsystemClass = SubsystemTypeIterator->GetClass();
 				check(SubsystemClass);
-				MultiThreadedSystemsBitSet.Add(*SubsystemClass);
+				MultiThreadedSystemsBitSet.Add(SubsystemClass);
 			}
 		}
 	}
diff --git a/Engine/Source/Runtime/MassEntity/Private/MassRequirements.cpp b/Engine/Source/Runtime/MassEntity/Private/MassRequirements.cpp
--- a/Engine/Source/Runtime/MassEntity/Private/MassRequirements.cpp
+++ b/Engine/Source/Runtime/MassEntity/Private/MassRequirements.cpp
@@ -48,11 +48,11 @@ namespace UE::Mass::Private
 				check(Requirement.StructType);
 				if (Requirement.AccessMode == EMassFragmentAccess::ReadOnly)
 				{
-					Out.Read.Add(*Requirement.StructType);
+					Out.Read.Add(Requirement.StructType);
 				}
 				else if (Requirement.AccessMode == EMassFragmentAccess::ReadWrite)
 				{
-					Out.Write.Add(*Requirement.StructType);
+					Out.Write.Add(Requirement.StructType);
 				}
 			}
 		}
@@ -69,7 +69,7 @@ namespace UE::Mass::Private
 				check(Requirement.StructType);
 				if (ensureMsgf(Requirement.AccessMode == EMassFragmentAccess::ReadOnly, TEXT("ReadOnly is the only supported AccessMode for ConstSharedFragments")))
 				{
-					Out.Read.Add(*Requirement.StructType);
+					Out.Read.Add(Requirement.StructType);
 				}
 			}
 		}
diff --git a/Engine/Source/Runtime/MassEntity/Private/MassSparseElementsStorage.cpp b/Engine/Source/Runtime/MassEntity/Private/MassSparseElementsStorage.cpp
--- a/Engine/Source/Runtime/MassEntity/Private/MassSparseElementsStorage.cpp
+++ b/Engine/Source/Runtime/MassEntity/Private/MassSparseElementsStorage.cpp
@@ -174,7 +174,7 @@ void FSparseElementsStorage::Deinitialize()
 
 FSparseElementsStorage::FTypePool& FSparseElementsStorage::GetOrCreateTypePool(TNotNull<const UScriptStruct*> ElementType, const int32 InTypeIndex)
 {
-	const int32 TypeIndex = InTypeIndex == INDEX_NONE ? FMassElementBitSet::GetTypeIndex(*ElementType) : InTypeIndex;
+	const int32 TypeIndex = InTypeIndex == INDEX_NONE ? FMassElementBitSet::GetTypeIndex(ElementType) : InTypeIndex;
 
 	if (TypePools.IsValidIndex(TypeIndex) == false)
 	{
@@ -186,20 +186,20 @@ FSparseElementsStorage::FTypePool& FSparseElementsStorage::GetOrCreateTypePool(T
 
 const FSparseElementsStorage::FTypePool* FSparseElementsStorage::GetTypePool(TNotNull<const UScriptStruct*> ElementType) const
 {
-	const int32 TypeIndex = FMassElementBitSet::GetTypeIndex(*ElementType);
+	const int32 TypeIndex = FMassElementBitSet::GetTypeIndex(ElementType);
 	return TypePools.IsValidIndex(TypeIndex) ? &TypePools[TypeIndex] : nullptr;
 }
 
 FSparseElementsStorage::FTypePool* FSparseElementsStorage::GetTypePool(TNotNull<const UScriptStruct*> ElementType)
 {
-	const int32 TypeIndex = FMassElementBitSet::GetTypeIndex(*ElementType);
+	const int32 TypeIndex = FMassElementBitSet::GetTypeIndex(ElementType);
 	return TypePools.IsValidIndex(TypeIndex) ? &TypePools[TypeIndex] : nullptr;
 }
 
 FStructView FSparseElementsStorage::AddElementToEntity(FMassEntityHandle EntityHandle, TNotNull<const UScriptStruct*> ElementType)
 {
 	checkSlow(!UE::Mass::IsA<FMassTag>(ElementType));
-	const int32 TypeIndex = FMassElementBitSet::GetTypeIndex(*ElementType);
+	const int32 TypeIndex = FMassElementBitSet::GetTypeIndex(ElementType);
 	FTypePool& TypePool = GetOrCreateTypePool(ElementType, TypeIndex);
 	return TypePool.Add_GetView(EntityHandle.Index);
 }
@@ -213,7 +213,7 @@ bool FSparseElementsStorage::RemoveElementFromEntity(FMassEntityHandle EntityHan
 template<typename TContainer>
 void FSparseElementsStorage::BatchAddElementToEntities(TContainer InEntityHandles, TNotNull<const UScriptStruct*> ElementType, TArray<FStructView>* OutAddedElements)
 {
-	const int32 TypeIndex = FMassElementBitSet::GetTypeIndex(*ElementType);
+	const int32 TypeIndex = FMassElementBitSet::GetTypeIndex(ElementType);
 
 	FTypePool& TypePool = GetOrCreateTypePool(ElementType, TypeIndex);
 	// @todo TypePool.Add can be optimized by checking for index continuity and whether the following index is in the same chunk.
@@ -269,7 +269,7 @@ void FSparseElementsStorage::BatchAddElementInstancesToEntities(TArrayView<const
 	checkf(InEntityHandles.Num() == FragmentPayload.Num(), TEXT("%hs: Payload size is expected to match number of entity handles"), __FUNCTION__);
 	const UScriptStruct* ElementType = FragmentPayload.GetScriptStruct();
 	CA_ASSUME(ElementType);
-	const int32 TypeIndex = FMassElementBitSet::GetTypeIndex(*ElementType);
+	const int32 TypeIndex = FMassElementBitSet::GetTypeIndex(ElementType);
 
 	FTypePool& TypePool = GetOrCreateTypePool(ElementType, TypeIndex);
 	for (int32 Index = 0; Index < InEntityHandles.Num(); ++Index)
@@ -283,7 +283,7 @@ void FSparseElementsStorage::BatchAddElementInstancesToEntities(TArrayView<const
 	check(FragmentInstance.IsValid());
 	const UScriptStruct* ElementType = FragmentInstance.GetScriptStruct();
 	CA_ASSUME(ElementType);
-	const int32 TypeIndex = FMassElementBitSet::GetTypeIndex(*ElementType);
+	const int32 TypeIndex = FMassElementBitSet::GetTypeIndex(ElementType);
 
 	FTypePool& TypePool = GetOrCreateTypePool(ElementType, TypeIndex);
 	for (int32 Index = 0; Index < InEntityHandles.Num(); ++Index)
@@ -351,7 +351,7 @@ void FSparseElementsStorage::RemoveEntity(FMassEntityHandle EntityHandle, const
 
 void FSparseElementsStorage::ConfigureType(TNotNull<const UScriptStruct*> ElementType, FTypeConfig& Config)
 {
-	const int32 TypeIndex = FMassElementBitSet::GetTypeIndex(*ElementType);
+	const int32 TypeIndex = FMassElementBitSet::GetTypeIndex(ElementType);
 
 	if (ensureMsgf(TypePools.IsValidIndex(TypeIndex) == false || TypePools[TypeIndex].GetNumElements() == 0
 		, TEXT("%hs: Configuring types after elements have been instantiated is not supported."), __FUNCTION__))
diff --git a/Engine/Source/Runtime/MassEntity/Public/MassBitSetRegistry.h b/Engine/Source/Runtime/MassEntity/Public/MassBitSetRegistry.h
--- a/Engine/Source/Runtime/MassEntity/Public/MassBitSetRegistry.h
+++ b/Engine/Source/Runtime/MassEntity/Public/MassBitSetRegistry.h
@@ -108,9 +108,9 @@ namespace UE::Mass
 			 * @param Type the UScriptStruct representing the type to register.
 			 * @return The index assigned to the registered type.
 			 */
-			inline int32 RegisterType(const UScriptStruct* Type)
+			inline int32 RegisterType(TNotNull<const UScriptStruct*> Type)
 			{
-				return StructTracker.Register(*Type);
+				return StructTracker.Register(Type);
 			}
 
 			/**
diff --git a/Engine/Source/Runtime/MassEntity/Public/MassEntityManager.h b/Engine/Source/Runtime/MassEntity/Public/MassEntityManager.h
--- a/Engine/Source/Runtime/MassEntity/Public/MassEntityManager.h
+++ b/Engine/Source/Runtime/MassEntity/Public/MassEntityManager.h
@@ -405,8 +405,8 @@ public:
 
 	UE_API const UE::Mass::FArchetypeGroups& GetGroupsForArchetype(const FMassArchetypeHandle& ArchetypeHandle) const;
 
-	UE_API void AddFragmentToEntity(FMassEntityHandle EntityHandle, const UScriptStruct* FragmentType);
-	UE_API void AddFragmentToEntity(FMassEntityHandle EntityHandle, const UScriptStruct* FragmentType, const FStructInitializationCallback& Initializer);
+	UE_API void AddFragmentToEntity(FMassEntityHandle EntityHandle, TNotNull<const UScriptStruct*> FragmentType);
+	UE_API void AddFragmentToEntity(FMassEntityHandle EntityHandle, TNotNull<const UScriptStruct*> FragmentType, const FStructInitializationCallback& Initializer);
 
 	/** 
 	 *  Ensures that only unique fragments are added. 
@@ -519,7 +519,13 @@ public:
 	 * Will do nothing if entity did not have the shared fragment.
 	 * @return True if fragment removed from entity, false otherwise.
 	 */
-	UE_API bool RemoveConstSharedFragmentFromEntity(const FMassEntityHandle EntityHandle, const UScriptStruct& ConstSharedFragmentType);
+	UE_API bool RemoveConstSharedFragmentFromEntity(const FMassEntityHandle EntityHandle, TNotNull<const UScriptStruct*> ConstSharedFragmentType);
+
+	UE_DEPRECATED(5.8, "Use overload taking TNotNull")
+	bool RemoveConstSharedFragmentFromEntity(const FMassEntityHandle EntityHandle, const UScriptStruct& ConstSharedFragmentType)
+	{
+		return RemoveConstSharedFragmentFromEntity(EntityHandle, &ConstSharedFragmentType);
+	}
 
 	/**
 	 * Adds a new shared fragment to the given entity. Note that it only works if the given entity doesn't have
@@ -536,7 +542,13 @@ public:
 	 * Will do nothing if entity did not have the shared fragment.
 	 * @return True if fragment removed from entity, false otherwise.
 	 */
-	UE_API bool RemoveSharedFragmentFromEntity(const FMassEntityHandle EntityHandle, const UScriptStruct& SharedFragmentType);
+	UE_API bool RemoveSharedFragmentFromEntity(const FMassEntityHandle EntityHandle, TNotNull<const UScriptStruct*> SharedFragmentType);
+
+	UE_DEPRECATED(5.8, "Use overload taking TNotNull")
+	bool RemoveSharedFragmentFromEntity(const FMassEntityHandle EntityHandle, const UScriptStruct& ConstSharedFragmentType)
+	{
+		return RemoveSharedFragmentFromEntity(EntityHandle, &ConstSharedFragmentType);
+	}
 
 	/**
 	 * Removes EntityHandle from any-and-all groups of given type - i.e. the entity will be moved to an archetype
diff --git a/Engine/Source/Runtime/MassEntity/Public/MassEntityTypes.h b/Engine/Source/Runtime/MassEntity/Public/MassEntityTypes.h
--- a/Engine/Source/Runtime/MassEntity/Public/MassEntityTypes.h
+++ b/Engine/Source/Runtime/MassEntity/Public/MassEntityTypes.h
@@ -419,7 +419,7 @@ struct FMassArchetypeSharedFragmentValues
 	inline bool ContainsType(const UScriptStruct* FragmentType) const
 	{
 		return (UE::Mass::IsA<FMassSharedFragment>(FragmentType) || UE::Mass::IsA<FMassConstSharedFragment>(FragmentType))
-			&& StoredElementsBitSet.Contains(*FragmentType);
+			&& StoredElementsBitSet.Contains(FragmentType);
 	}
 
 	template<typename T>
@@ -427,7 +427,7 @@ struct FMassArchetypeSharedFragmentValues
 	{
 		if constexpr (UE::Mass::CConstSharedFragment<T> || UE::Mass::CSharedFragment<T>)
 		{
-			return StoredElementsBitSet.Contains(*T::StaticStruct());
+			return StoredElementsBitSet.Contains(T::StaticStruct());
 		}		
 		else
 		{
@@ -991,34 +991,34 @@ auto FMassArchetypeCompositionDescriptor::GetContainer() const
 template<typename T>
 bool FMassArchetypeCompositionDescriptor::Contains() const
 {
-	return ElementsBitSet.Contains(*T::StaticStruct());
+	return ElementsBitSet.Contains(T::StaticStruct());
 }
 
 inline bool FMassArchetypeCompositionDescriptor::Contains(TNotNull<const UScriptStruct*> ElementType) const
 {
-	return ElementsBitSet.Contains(*ElementType);
+	return ElementsBitSet.Contains(ElementType);
 }
 
 template<typename T>
 void FMassArchetypeCompositionDescriptor::Add()
 {
-	return ElementsBitSet.Add(*T::StaticStruct());
+	return ElementsBitSet.Add(T::StaticStruct());
 }
 
 inline void FMassArchetypeCompositionDescriptor::Add(TNotNull<const UScriptStruct*> ElementType)
 {
-	return ElementsBitSet.Add(*ElementType);
+	return ElementsBitSet.Add(ElementType);
 }
 
 template<typename T>
 void FMassArchetypeCompositionDescriptor::Remove()
 {
-	return ElementsBitSet.Remove(*T::StaticStruct());
+	return ElementsBitSet.Remove(T::StaticStruct());
 }
 
 inline void FMassArchetypeCompositionDescriptor::Remove(TNotNull<const UScriptStruct*> ElementType)
 {
-	return ElementsBitSet.Remove(*ElementType);
+	return ElementsBitSet.Remove(ElementType);
 }
 
 inline FMassFragmentBitSet FMassArchetypeCompositionDescriptor::DebugGetFragments() const 
diff --git a/Engine/Source/Runtime/MassEntity/Public/MassExecutionContext.h b/Engine/Source/Runtime/MassEntity/Public/MassExecutionContext.h
--- a/Engine/Source/Runtime/MassEntity/Public/MassExecutionContext.h
+++ b/Engine/Source/Runtime/MassEntity/Public/MassExecutionContext.h
@@ -26,7 +26,7 @@
 		, View ? *GetNameSafe(View->Requirement.StructType) : TEXT("[Not found]"))
 
 #define MASS_CHECK_IF_SPARSE_REQUIREMENT(ElementType) \
-	checkf(RequestedSparseElementsAccess.Contains(*ElementType) \
+	checkf(RequestedSparseElementsAccess.Contains(ElementType) \
 			, TEXT("%hs: %s not requested as requirement by the query"), __FUNCTION__, *ElementType->GetName())
 
 struct FMassEntityQuery;
@@ -406,33 +406,51 @@ public:
 		}
 	}
 
-	bool DoesArchetypeHaveElement(const UScriptStruct& ElementType) const
+	bool DoesArchetypeHaveElement(TNotNull<const UScriptStruct*> ElementType) const
 	{
 		return CurrentArchetypeCompositionBitSet.Contains(ElementType);
 	}
 
+	UE_DEPRECATED(5.8, "Use overload taking TNotNull")
+	bool DoesArchetypeHaveElement(const UScriptStruct& ElementType) const
+	{
+		return DoesArchetypeHaveElement(&ElementType);
+	}
+
 	template<UE::Mass::CElement T>
 	bool DoesArchetypeHaveElement() const
 	{
 		return CurrentArchetypeCompositionBitSet.Contains<T>();
 	}
 
-	bool DoesArchetypeHaveFragment(const UScriptStruct& FragmentType) const
+	bool DoesArchetypeHaveFragment(TNotNull<const UScriptStruct*> FragmentType) const
 	{
 		return CurrentArchetypeCompositionBitSet.Contains(FragmentType);
 	}
 
+	UE_DEPRECATED(5.8, "Use overload taking TNotNull")
+	bool DoesArchetypeHaveFragment(const UScriptStruct& ElementType) const
+	{
+		return DoesArchetypeHaveFragment(&ElementType);
+	}
+
 	template<UE::Mass::CFragment T>
 	bool DoesArchetypeHaveFragment() const
 	{
 		return CurrentArchetypeCompositionBitSet.Contains<T>();
 	}
 
-	bool DoesArchetypeHaveTag(const UScriptStruct& TagType) const
+	bool DoesArchetypeHaveTag(TNotNull<const UScriptStruct*> TagType) const
 	{
 		return CurrentArchetypeCompositionBitSet.Contains(TagType);
 	}
 
+	UE_DEPRECATED(5.8, "Use overload taking TNotNull")
+	bool DoesArchetypeHaveTag(const UScriptStruct& ElementType) const
+	{
+		return DoesArchetypeHaveTag(&ElementType);
+	}
+
 	template<UE::Mass::CTag T>
 	bool DoesArchetypeHaveTag() const
 	{
diff --git a/Engine/Source/Runtime/MassEntity/Public/MassRequirements.h b/Engine/Source/Runtime/MassEntity/Public/MassRequirements.h
--- a/Engine/Source/Runtime/MassEntity/Public/MassRequirements.h
+++ b/Engine/Source/Runtime/MassEntity/Public/MassRequirements.h
@@ -130,11 +130,11 @@ struct FMassSubsystemRequirements
 		switch (AccessMode)
 		{
 		case EMassFragmentAccess::ReadOnly:
-			RequiredConstSubsystems.Add(**SubsystemClass);
+			RequiredConstSubsystems.Add(*SubsystemClass);
 			bRequiresGameThreadExecution |= bGameThreadOnly;
 			break;
 		case EMassFragmentAccess::ReadWrite:
-			RequiredMutableSubsystems.Add(**SubsystemClass);
+			RequiredMutableSubsystems.Add(*SubsystemClass);
 			bRequiresGameThreadExecution |= bGameThreadOnly;
 			break;
 		default:
@@ -213,16 +213,16 @@ struct FMassFragmentRequirements
 		switch (Presence)
 		{
 		case EMassFragmentPresence::All:
-			RequiredAllFragments.Add(*FragmentType);
+			RequiredAllFragments.Add(FragmentType);
 			break;
 		case EMassFragmentPresence::Any:
-			RequiredAnyFragments.Add(*FragmentType);
+			RequiredAnyFragments.Add(FragmentType);
 			break;
 		case EMassFragmentPresence::Optional:
-			RequiredOptionalFragments.Add(*FragmentType);
+			RequiredOptionalFragments.Add(FragmentType);
 			break;
 		case EMassFragmentPresence::None:
-			RequiredNoneFragments.Add(*FragmentType);
+			RequiredNoneFragments.Add(FragmentType);
 			break;
 		}
 		// force recaching the next time this query is used or the following CacheArchetypes call.
@@ -272,16 +272,16 @@ struct FMassFragmentRequirements
 		switch (Presence)
 		{
 		case EMassFragmentPresence::All:
-			RequiredAllTags.Add(*TagType);
+			RequiredAllTags.Add(TagType);
 			break;
 		case EMassFragmentPresence::Any:
-			RequiredAnyTags.Add(*TagType);
+			RequiredAnyTags.Add(TagType);
 			break;
 		case EMassFragmentPresence::None:
-			RequiredNoneTags.Add(*TagType);
+			RequiredNoneTags.Add(TagType);
 			break;
 		case EMassFragmentPresence::Optional:
-			RequiredOptionalTags.Add(*TagType);
+			RequiredOptionalTags.Add(TagType);
 			break;
 		}
 		IncrementChangeCounter();
@@ -371,16 +371,16 @@ struct FMassFragmentRequirements
 		switch (Presence)
 		{
 		case EMassFragmentPresence::All:
-			RequiredAllSparseElements.Add(*ElementType);
+			RequiredAllSparseElements.Add(ElementType);
 			break;
 		case EMassFragmentPresence::Any:
-			RequiredAnySparseElements.Add(*ElementType);
+			RequiredAnySparseElements.Add(ElementType);
 			break;
 		case EMassFragmentPresence::None:
-			RequiredNoneSparseElements.Add(*ElementType);
+			RequiredNoneSparseElements.Add(ElementType);
 			break;
 		case EMassFragmentPresence::Optional:
-			RequiredOptionalSparseElements.Add(*ElementType);
+			RequiredOptionalSparseElements.Add(ElementType);
 			break;
 		}
 
@@ -425,15 +425,15 @@ struct FMassFragmentRequirements
 		switch (Presence)
 		{
 		case EMassFragmentPresence::All:
-			RequiredAllChunkFragments.Add(*ChunkFragmentType);
+			RequiredAllChunkFragments.Add(ChunkFragmentType);
 			ChunkFragmentRequirements.Emplace(ChunkFragmentType, AccessMode, Presence);
 			break;
 		case EMassFragmentPresence::Optional:
-			RequiredOptionalChunkFragments.Add(*ChunkFragmentType);
+			RequiredOptionalChunkFragments.Add(ChunkFragmentType);
 			ChunkFragmentRequirements.Emplace(ChunkFragmentType, AccessMode, Presence);
 			break;
 		case EMassFragmentPresence::None:
-			RequiredNoneChunkFragments.Add(*ChunkFragmentType);
+			RequiredNoneChunkFragments.Add(ChunkFragmentType);
 			break;
 		}
 		IncrementChangeCounter();
@@ -483,15 +483,15 @@ struct FMassFragmentRequirements
 		switch (Presence)
 		{
 		case EMassFragmentPresence::All:
-			RequiredAllConstSharedFragments.Add(*FragmentType);
+			RequiredAllConstSharedFragments.Add(FragmentType);
 			ConstSharedFragmentRequirements.Emplace(FragmentType, EMassFragmentAccess::ReadOnly, Presence);
 			break;
 		case EMassFragmentPresence::Optional:
-			RequiredOptionalConstSharedFragments.Add(*FragmentType);
+			RequiredOptionalConstSharedFragments.Add(FragmentType);
 			ConstSharedFragmentRequirements.Emplace(FragmentType, EMassFragmentAccess::ReadOnly, Presence);
 			break;
 		case EMassFragmentPresence::None:
-			RequiredNoneConstSharedFragments.Add(*FragmentType);
+			RequiredNoneConstSharedFragments.Add(FragmentType);
 			break;
 		}
 		IncrementChangeCounter();
@@ -544,7 +544,7 @@ struct FMassFragmentRequirements
 		switch (Presence)
 		{
 		case EMassFragmentPresence::All:
-			RequiredAllSharedFragments.Add(*SharedFragmentType);
+			RequiredAllSharedFragments.Add(SharedFragmentType);
 			SharedFragmentRequirements.Emplace(SharedFragmentType, AccessMode, Presence);
 			if (AccessMode == EMassFragmentAccess::ReadWrite)
 			{
@@ -552,7 +552,7 @@ struct FMassFragmentRequirements
 			}
 			break;
 		case EMassFragmentPresence::Optional:
-			RequiredOptionalSharedFragments.Add(*SharedFragmentType);
+			RequiredOptionalSharedFragments.Add(SharedFragmentType);
 			SharedFragmentRequirements.Emplace(SharedFragmentType, AccessMode, Presence);
 			if (AccessMode == EMassFragmentAccess::ReadWrite)
 			{
@@ -560,7 +560,7 @@ struct FMassFragmentRequirements
 			}
 			break;
 		case EMassFragmentPresence::None:
-			RequiredNoneSharedFragments.Add(*SharedFragmentType);
+			RequiredNoneSharedFragments.Add(SharedFragmentType);
 			break;
 		}
 		IncrementChangeCounter();
@@ -604,10 +604,10 @@ struct FMassFragmentRequirements
 		switch (AccessMode)
 		{
 		case EMassFragmentAccess::ReadOnly:
-			IndirectReadOnlyFragments.Add(*FragmentType);
+			IndirectReadOnlyFragments.Add(FragmentType);
 			break;
 		case EMassFragmentAccess::ReadWrite:
-			IndirectReadWriteFragments.Add(*FragmentType);
+			IndirectReadWriteFragments.Add(FragmentType);
 			break;
 		}
 
@@ -656,10 +656,10 @@ struct FMassFragmentRequirements
 		switch (AccessMode)
 		{
 		case EMassFragmentAccess::ReadOnly:
-			IndirectReadOnlyFragments.Add(*FragmentType);
+			IndirectReadOnlyFragments.Add(FragmentType);
 			break;
 		case EMassFragmentAccess::ReadWrite:
-			IndirectReadWriteFragments.Add(*FragmentType);
+			IndirectReadWriteFragments.Add(FragmentType);
 			break;
 		}
 
diff --git a/Engine/Source/Runtime/MassEntity/Public/MassSubsystemAccess.h b/Engine/Source/Runtime/MassEntity/Public/MassSubsystemAccess.h
--- a/Engine/Source/Runtime/MassEntity/Public/MassSubsystemAccess.h
+++ b/Engine/Source/Runtime/MassEntity/Public/MassSubsystemAccess.h
@@ -74,7 +74,7 @@ struct FMassSubsystemAccess
 	T* GetMutableSubsystem(const TSubclassOf<USubsystem> SubsystemClass)
 	{
 		check(SubsystemClass);
-		const int32 SystemIndex = FMassExternalSubsystemBitSet::GetTypeIndex(**SubsystemClass);
+		const int32 SystemIndex = FMassExternalSubsystemBitSet::GetTypeIndex(*SubsystemClass);
 		if (ensureMsgf(MutableSubsystemsBitSet.IsBitSet(SystemIndex)
 			, TEXT("Undeclared read/write access to subsystem %s"), *GetNameSafe(SubsystemClass)))
 		{
@@ -97,7 +97,7 @@ struct FMassSubsystemAccess
 	const T* GetSubsystem(const TSubclassOf<USubsystem> SubsystemClass)
 	{
 		check(SubsystemClass);
-		const int32 SystemIndex = FMassExternalSubsystemBitSet::GetTypeIndex(**SubsystemClass);
+		const int32 SystemIndex = FMassExternalSubsystemBitSet::GetTypeIndex(*SubsystemClass);
 		if (ensureMsgf(ConstSubsystemsBitSet.IsBitSet(SystemIndex) || MutableSubsystemsBitSet.IsBitSet(SystemIndex)
 			, TEXT("Undeclared read-only access to subsystem %s"), *GetNameSafe(SubsystemClass)))
 		{
diff --git a/Engine/Source/Runtime/MassEntity/Public/TypeBitSetBuilder.h b/Engine/Source/Runtime/MassEntity/Public/TypeBitSetBuilder.h
--- a/Engine/Source/Runtime/MassEntity/Public/TypeBitSetBuilder.h
+++ b/Engine/Source/Runtime/MassEntity/Public/TypeBitSetBuilder.h
@@ -34,10 +34,8 @@ struct TTypeBitSetBuilder : TTypeBitSetBase<TTypeBitSetBuilder<TBaseStruct, TUSt
 {
 	/** Define the base class for easier reference */
 	using Super = TTypeBitSetBase<TTypeBitSetBuilder, TBaseStruct, TUStructType, TContainer&, bTestInheritanceAtRuntime>;
-	/** Alias for Unreal's struct type */
-	using FUStructType = TUStructType; 
 	/** Alias for the base struct type */
-	using FBaseStruct = TBaseStruct;   
+	using FBaseStruct = TBaseStruct;
 	/** Traits for compile-time checks */
 	using FTraits = TTypeBitSetTraits<TBaseStruct>;
 
@@ -122,14 +120,14 @@ public:
 	 * @param InStructType the struct type to get the index for.
 	 * @return The index of the struct type.
 	 */
-	int32 GetTypeIndex(const TUStructType& InStructType) const
+	int32 GetTypeIndex(TNotNull<const TUStructType*> InStructType) const
 	{
 		if constexpr (bTestInheritanceAtRuntime)
 		{
 			// Ensure the struct type derives from the base struct
-			if (UNLIKELY(!ensureMsgf(InStructType.IsChildOf(GetBaseUStruct())
+			if (UNLIKELY(!ensureMsgf(InStructType->IsChildOf(GetBaseUStruct())
 				, TEXT("Creating index for '%s' while it doesn't derive from the expected struct type %s")
-				, *InStructType.GetPathName(), *GetBaseUStruct()->GetName())))
+				, *InStructType->GetPathName(), *GetBaseUStruct()->GetName())))
 			{
 				return INDEX_NONE;
 			}
@@ -139,20 +137,26 @@ public:
 		return GetStructTracker().FindOrAddStructTypeIndex(InStructType);
 	}
 
+	UE_DEPRECATED(5.8, "Use overload taking TNotNull")
+	int32 GetTypeIndex(const TUStructType& InStructType) const
+	{
+		return GetTypeIndex(&InStructType);
+	}
+
 	/**
 	 * Static method to get the type index for a given struct type.
 	 * @param InStructTracker the struct tracker to use.
 	 * @param InStructType the struct type to get the index for.
 	 * @return The index of the struct type.
 	 */
-	static int32 GetTypeIndex(const FStructTracker& InStructTracker, const TUStructType& InStructType)
+	static int32 GetTypeIndex(const FStructTracker& InStructTracker, TNotNull<const TUStructType*> InStructType)
 	{
 		if constexpr (bTestInheritanceAtRuntime)
 		{
 			// Ensure the struct type derives from the base struct
-			if (UNLIKELY(!ensureMsgf(InStructType.IsChildOf(GetBaseUStruct())
+			if (UNLIKELY(!ensureMsgf(InStructType->IsChildOf(GetBaseUStruct())
 				, TEXT("Creating index for '%s' while it doesn't derive from the expected struct type %s")
-				, *InStructType.GetPathName(), *GetBaseUStruct()->GetName())))
+				, *InStructType->GetPathName(), *GetBaseUStruct()->GetName())))
 			{
 				return INDEX_NONE;
 			}
@@ -162,6 +166,12 @@ public:
 		return InStructTracker.FindOrAddStructTypeIndex(InStructType);
 	}
 
+	UE_DEPRECATED(5.8, "Use overload taking TNotNull")
+	int32 GetTypeIndex(const FStructTracker& InStructTracker, const TUStructType& InStructType) const
+	{
+		return GetTypeIndex(InStructTracker, &InStructType);
+	}
+
 	/**
 	 * Template method to get the type index for a specific C++ type.
 	 * Ensures at compile-time that the type is valid.
@@ -172,7 +182,7 @@ public:
 	static int32 GetTypeIndex(FStructTracker& InStructTracker)
 	{
 		static_assert(FTraits::template IsValidType<T>, "Given struct is not a valid type for this TypeBitSetBuilder.");
-		static const int32 TypeIndex = GetTypeIndex(InStructTracker, *UE::StructUtils::GetAsUStruct<T>());
+		static const int32 TypeIndex = GetTypeIndex(InStructTracker, UE::StructUtils::GetAsUStruct<T>());
 		return TypeIndex;
 	}
 
@@ -186,7 +196,7 @@ public:
 	int32 GetTypeIndex() const
 	{
 		static_assert(FTraits::template IsValidType<T>, "Given struct is not a valid type for this TypeBitSetBuilder.");
-		static const int32 TypeIndex = GetTypeIndex(*UE::StructUtils::GetAsUStruct<T>());
+		static const int32 TypeIndex = GetTypeIndex(UE::StructUtils::GetAsUStruct<T>());
 		return TypeIndex;
 	}
 	
