[Mass] Added FElementBitSet::GetSparseElements for easy extraction of sparse elements from a given element bitset

#rnx

#Synced-CL 51337022

diff --git a/Engine/Source/Runtime/MassEntity/Public/MassElement.h b/Engine/Source/Runtime/MassEntity/Public/MassElement.h
--- a/Engine/Source/Runtime/MassEntity/Public/MassElement.h
+++ b/Engine/Source/Runtime/MassEntity/Public/MassElement.h
@@ -59,6 +59,7 @@ namespace UE::Mass
 		
 		bool ContainsOnlyTagsAndFragments() const;
 		FElementBitSet GetFragmentsAndTags() const;
+		FElementBitSet GetSparseElements() const;
 		FElementBitSet GetSharedFragments() const;
 	
 		static EElementType DetermineElementType(const TNotNull<const UStruct*> ElementType)
@@ -165,6 +166,11 @@ namespace UE::Mass
 		return (*this) & GetAllFragmentsAndTags();
 	}
 
+	inline FElementBitSet FElementBitSet::GetSparseElements() const
+	{
+		return (*this) & GetAllSparseElements();
+	}
+
 	inline FElementBitSet FElementBitSet::GetSharedFragments() const
 	{
 		return (*this) & GetAllSharedFragments();
