[Mass] Added FMassElementBitSet-flavored functions to FMassObserverManager to support future removal of FMassArchetypeCompositionDescriptor

#rnx

#Synced-CL 51337022

diff --git a/Engine/Source/Runtime/MassEntity/Private/MassObserverManager.cpp b/Engine/Source/Runtime/MassEntity/Private/MassObserverManager.cpp
--- a/Engine/Source/Runtime/MassEntity/Private/MassObserverManager.cpp
+++ b/Engine/Source/Runtime/MassEntity/Private/MassObserverManager.cpp
@@ -278,7 +278,7 @@ bool FMassObserverManager::OnPostEntitiesCreated(const FMassArchetypeEntityColle
 	return OnCompositionChanged(EntityCollection, ArchetypeComposition, EMassObservedOperation::CreateEntity);
 }
 
-bool FMassObserverManager::OnPostEntityCreated(FMassEntityHandle EntityHandle, const FMassArchetypeCompositionDescriptor& Composition)
+bool FMassObserverManager::OnPostEntityCreated(FMassEntityHandle EntityHandle, const FMassElementBitSet& InQueriedBitSet)
 {
 	TRACE_CPUPROFILER_EVENT_SCOPE(MassObserver_OnPostEntitiesCreated_Collection);
 
@@ -289,14 +289,14 @@ bool FMassObserverManager::OnPostEntityCreated(FMassEntityHandle EntityHandle, c
 		return false;
 	}
 
-	if (Composition.IsEmpty())
+	if (InQueriedBitSet.IsEmpty())
 	{
 		const FMassArchetypeHandle ArchetypeHandle = EntityManager.GetArchetypeForEntity(EntityHandle);
 		const FMassArchetypeCompositionDescriptor& ArchetypeComposition = EntityManager.GetArchetypeComposition(ArchetypeHandle);
 		return OnCompositionChanged(EntityHandle, ArchetypeComposition, EMassObservedOperation::CreateEntity);
 	}
 
-	return OnCompositionChanged(EntityHandle, Composition, EMassObservedOperation::CreateEntity);
+	return OnCompositionChanged(EntityHandle, InQueriedBitSet, EMassObservedOperation::CreateEntity);
 }
 
 bool FMassObserverManager::OnPreEntitiesDestroyed(const FMassArchetypeEntityCollection& EntityCollection)
@@ -313,10 +313,10 @@ bool FMassObserverManager::OnPreEntitiesDestroyed(UE::Mass::FProcessingContext&
 	return OnCompositionChanged(EntityCollection, ArchetypeComposition, EMassObservedOperation::DestroyEntity, &ProcessingContext);
 }
 
-bool FMassObserverManager::OnPreEntityDestroyed(const FMassArchetypeCompositionDescriptor& ArchetypeComposition, const FMassEntityHandle Entity)
+bool FMassObserverManager::OnPreEntityDestroyed(const FMassEntityHandle Entity, const FMassElementBitSet& Composition)
 {
 	TRACE_CPUPROFILER_EVENT_SCOPE_STR("OnPreEntityDestroyed")
-	return OnCompositionChanged(Entity, ArchetypeComposition, EMassObservedOperation::DestroyEntity);
+	return OnCompositionChanged(Entity, Composition, EMassObservedOperation::DestroyEntity);
 }
 
 bool FMassObserverManager::OnCompositionChanged(FCollectionRefOrHandle&& EntityCollection, const FMassElementBitSet& CompositionDelta
diff --git a/Engine/Source/Runtime/MassEntity/Public/MassObserverManager.h b/Engine/Source/Runtime/MassEntity/Public/MassObserverManager.h
--- a/Engine/Source/Runtime/MassEntity/Public/MassObserverManager.h
+++ b/Engine/Source/Runtime/MassEntity/Public/MassObserverManager.h
@@ -82,6 +82,11 @@ struct FMassObserverManager
 		return (InQueriedBitSet.IsEmpty() == false) && ObservedElements[static_cast<uint8>(Operation)].HasAny(InQueriedBitSet);
 	}
 
+	bool HasObserversForComposition(const FMassElementBitSet& InQueriedBitSet, const EMassObservedOperation Operation) const
+	{
+		return HasObserversForBitSet(InQueriedBitSet, Operation);
+	}
+
 	bool HasObserversForComposition(const FMassArchetypeCompositionDescriptor& Composition, const EMassObservedOperation Operation) const
 	{
 		return HasObserversForBitSet(Composition.GetElementsBitSet(), Operation);
@@ -91,7 +96,12 @@ struct FMassObserverManager
 	MASSENTITY_API bool OnPostEntitiesCreated(const FMassArchetypeEntityCollection& EntityCollection);
 
 	/** @return whether there are observers watching affected elements */
-	MASSENTITY_API bool OnPostEntityCreated(const FMassEntityHandle Entity, const FMassArchetypeCompositionDescriptor& Composition);
+	MASSENTITY_API bool OnPostEntityCreated(const FMassEntityHandle Entity, const FMassElementBitSet& ElementsBitSet);
+
+	bool OnPostEntityCreated(const FMassEntityHandle Entity, const FMassArchetypeCompositionDescriptor& Composition)
+	{
+		return OnPostEntityCreated(Entity, Composition.GetElementsBitSet());
+	}
 
 	/** @return whether there are observers watching affected elements */
 	MASSENTITY_API bool OnPreEntitiesDestroyed(const FMassArchetypeEntityCollection& EntityCollection);
@@ -100,17 +110,35 @@ struct FMassObserverManager
 	MASSENTITY_API bool OnPreEntitiesDestroyed(UE::Mass::FProcessingContext& ProcessingContext, const FMassArchetypeEntityCollection& EntityCollection);
 
 	/** @return whether there are observers watching affected elements */
-	MASSENTITY_API bool OnPreEntityDestroyed(const FMassArchetypeCompositionDescriptor& ArchetypeComposition, const FMassEntityHandle Entity);
+	MASSENTITY_API bool OnPreEntityDestroyed(const FMassEntityHandle Entity, const FMassElementBitSet& ElementsBitSet);
+
+	bool OnPreEntityDestroyed(const FMassArchetypeCompositionDescriptor& ArchetypeComposition, const FMassEntityHandle Entity)
+	{
+		return OnPreEntityDestroyed(Entity, ArchetypeComposition.GetElementsBitSet());
+	}
+
+	/** @return whether there are observers watching affected elements */
+	bool OnPostCompositionAdded(const FMassEntityHandle Entity, const FMassElementBitSet& CompositionBitSet)
+	{
+		return OnCompositionChanged(Entity, CompositionBitSet, EMassObservedOperation::AddElement);
+	}
+
+	/** @return whether there are observers watching affected elements */
+	bool OnPreCompositionRemoved(const FMassEntityHandle Entity, const FMassElementBitSet& CompositionBitSet)
+	{
+		return OnCompositionChanged(Entity, CompositionBitSet, EMassObservedOperation::RemoveElement);
+	}
 
 	/** @return whether there are observers watching affected elements */
 	bool OnPostCompositionAdded(const FMassEntityHandle Entity, const FMassArchetypeCompositionDescriptor& Composition)
 	{
-		return OnCompositionChanged(Entity, Composition, EMassObservedOperation::AddElement);
+		return OnPostCompositionAdded(Entity, Composition.GetElementsBitSet());
 	}
+
 	/** @return whether there are observers watching affected elements */
 	bool OnPreCompositionRemoved(const FMassEntityHandle Entity, const FMassArchetypeCompositionDescriptor& Composition)
 	{
-		return OnCompositionChanged(Entity, Composition, EMassObservedOperation::RemoveElement);
+		return OnPreCompositionRemoved(Entity, Composition.GetElementsBitSet()); 
 	}
 
 protected:
