[Mass] a set of changes to seamlessly support sparse elements in tag and fragment API

diff --git a/Engine/Source/Runtime/MassEntity/Internal/MassArchetypeData.h b/Engine/Source/Runtime/MassEntity/Internal/MassArchetypeData.h
--- a/Engine/Source/Runtime/MassEntity/Internal/MassArchetypeData.h
+++ b/Engine/Source/Runtime/MassEntity/Internal/MassArchetypeData.h
@@ -486,6 +486,7 @@ public:
 	FStructView AddSparseElementToEntity(FMassEntityHandle EntityHandle, TNotNull<const UScriptStruct*> ElementType);
 	/** Works with all sparse element types (fragments and tags). */
 	void RemoveSparseElementFromEntity(FMassEntityHandle EntityHandle, TNotNull<const UScriptStruct*> ElementType);
+	void RemoveSparseElementFromEntity(FMassEntityHandle EntityHandle, const FMassElementBitSet& SparseElementsBitSet);
 	/** Expects non-tag sparse element type. Will assert if the assumption is broken. */
 	FStructView GetMutableSparseElementDataForEntity(TNotNull<const UScriptStruct*> ElementType, const FMassEntityHandle EntityHandle) const;
 	/** Expects non-tag sparse element type. Will assert if the assumption is broken. */
@@ -493,8 +494,10 @@ public:
 	/** @return a bitset indicating all sparse elements (both tags and fragments) owned by the given entity */
 	FMassElementBitSet GetSparseElementsBitSetForEntity(const FMassEntityHandle EntityHandle) const;
 
-	void* GetFragmentDataForEntityChecked(TNotNull<const UScriptStruct*> FragmentType, int32 EntityIndex) const;
-	void* GetFragmentDataForEntity(TNotNull<const UScriptStruct*> FragmentType, int32 EntityIndex) const;
+	void* GetMutableFragmentDataForEntityChecked(TNotNull<const UScriptStruct*> FragmentType, const FMassEntityHandle EntityHandle);
+	void* GetMutableFragmentDataForEntity(TNotNull<const UScriptStruct*> FragmentType, const FMassEntityHandle EntityHandle);
+	const void* GetFragmentDataForEntityChecked(TNotNull<const UScriptStruct*> FragmentType, const FMassEntityHandle EntityHandle) const;
+	const void* GetFragmentDataForEntity(TNotNull<const UScriptStruct*> FragmentType, const FMassEntityHandle EntityHandle) const;
 
 	bool HasSparseElementForEntity(TNotNull<const UScriptStruct*> ElementType, const FMassEntityHandle EntityHandle) const;
 	bool HasSparseElementForEntity(TNotNull<const UScriptStruct*> ElementType, const FMassEntityInChunkDataHandle& EntityInChunkHandle) const;
diff --git a/Engine/Source/Runtime/MassEntity/Private/MassArchetypeData.cpp b/Engine/Source/Runtime/MassEntity/Private/MassArchetypeData.cpp
--- a/Engine/Source/Runtime/MassEntity/Private/MassArchetypeData.cpp
+++ b/Engine/Source/Runtime/MassEntity/Private/MassArchetypeData.cpp
@@ -412,6 +412,26 @@ void FMassArchetypeData::RemoveSparseElementFromEntity(FMassEntityHandle EntityH
 	}
 }
 
+void FMassArchetypeData::RemoveSparseElementFromEntity(FMassEntityHandle EntityHandle, const FMassElementBitSet& SparseElementsBitSet)
+{
+	FMassEntityInChunkDataHandle EntityInChunkHandle = MakeEntityHandle(EntityHandle.Index);
+	UE::Mass::FChunkSparseElements& SparseElements = Chunks[EntityInChunkHandle.ChunkIndex].GetSparseElementsUnsafe();
+
+	if (SparseElements.IsEmpty() == false)
+	{
+		for (FMassElementBitSet::FIndexIterator It = SparseElementsBitSet.GetIndexIterator(); It; ++It)
+		{
+			const UScriptStruct* ElementType = FMassElementBitSet::GetTypeAtIndex(*It);
+			if (UE::Mass::IsA<FMassTag>(ElementType) == false)
+			{
+				CA_ASSUME(SparseElementsStorage);
+				SparseElementsStorage->RemoveElementFromEntity(EntityHandle, ElementType);
+			}
+			bSparseElementsBitSetDirty |= SparseElements.Remove(EntityInChunkHandle, ElementType);
+		}
+	}
+}
+
 FStructView FMassArchetypeData::GetMutableSparseElementDataForEntity(TNotNull<const UScriptStruct*> ElementType, const FMassEntityHandle EntityHandle) const
 {
 	check(UE::Mass::IsA<FMassTag>(ElementType) == false);
@@ -634,27 +654,43 @@ bool FMassArchetypeData::HasSparseElementForEntity(TNotNull<const UScriptStruct*
 		&& Chunks[EntityInChunkHandle.ChunkIndex].GetSparseElementsUnsafe().Contains(EntityInChunkHandle, ElementType);
 }
 
-void* FMassArchetypeData::GetFragmentDataForEntityChecked(TNotNull<const UScriptStruct*> FragmentType, int32 EntityIndex) const
+void* FMassArchetypeData::GetMutableFragmentDataForEntityChecked(TNotNull<const UScriptStruct*> FragmentType, const FMassEntityHandle EntityHandle)
 {
-	void* ReturnPtr = GetFragmentDataForEntity(FragmentType, EntityIndex);
-	checkf(ReturnPtr, TEXT("Entity %d has neither a regular nor a transient fragment of type %s"), EntityIndex, *FragmentType->GetName())
+	void* ReturnPtr = GetMutableFragmentDataForEntity(FragmentType, EntityHandle);
+	checkf(ReturnPtr, TEXT("Entity %s has neither a regular nor a sparse fragment of type %s"), *EntityHandle.DebugGetDescription(), *FragmentType->GetName())
 	return ReturnPtr;
 }
 
-void* FMassArchetypeData::GetFragmentDataForEntity(TNotNull<const UScriptStruct*> FragmentType, int32 EntityIndex) const
+const void* FMassArchetypeData::GetFragmentDataForEntityChecked(TNotNull<const UScriptStruct*> FragmentType, const FMassEntityHandle EntityHandle) const
+{
+	// ugly but internal, so :shrug:
+	return const_cast<FMassArchetypeData*>(this)->GetMutableFragmentDataForEntityChecked(FragmentType, EntityHandle);
+}
+
+void* FMassArchetypeData::GetMutableFragmentDataForEntity(TNotNull<const UScriptStruct*> FragmentType, const FMassEntityHandle EntityHandle)
 {
 	checkf(UE::Mass::IsA<FMassFragment>(FragmentType), TEXT("%s is not a fragment type"), *FragmentType->GetName());
 
-	if (const int32* FragmentIndex = FragmentIndexMap.Find(FragmentType))
+	if (UE::Mass::IsSparse(FragmentType) == false)
 	{
-		FMassRawEntityInChunkData InternalIndex = MakeEntityHandle(EntityIndex);
-		// failing the below Find means given entity's archetype is missing given FragmentType
-		return GetFragmentData(*FragmentIndex, InternalIndex);
+		if (const int32* FragmentIndex = FragmentIndexMap.Find(FragmentType))
+		{
+			FMassRawEntityInChunkData InternalIndex = MakeEntityHandle(EntityHandle.Index);
+			// failing the below Find means given entity's archetype is missing given FragmentType
+			return GetFragmentData(*FragmentIndex, InternalIndex);
+		}
 	}
+	
+	// Might be sparse element
+	// Note that it's fine to return the pointer from the resulting view, since worst-case-scenario
+	// it returns an invalid FStructView, which will result in this function returning nullptr.
+	return GetMutableSparseElementDataForEntity(FragmentType, EntityHandle).GetMemory();
+}
 
-	ensureMsgf(UE::Mass::IsSparse(FragmentType) == false
-		, TEXT("Trying to fetch sparse fragment data (%s) from the archetype. Use FSparseElementsStorage instead."), *FragmentType->GetName());
-	return nullptr;
+const void* FMassArchetypeData::GetFragmentDataForEntity(TNotNull<const UScriptStruct*> FragmentType, const FMassEntityHandle EntityHandle) const
+{
+	// ugly but internal, so :shrug:
+	return const_cast<FMassArchetypeData*>(this)->GetMutableFragmentDataForEntity(FragmentType, EntityHandle);
 }
 
 bool FMassArchetypeData::DoesEntityMatchSparseComposition(const FMassEntityInChunkDataHandle& EntityInChunkHandle
@@ -1350,7 +1386,7 @@ void FMassArchetypeData::DebugPrintEntity(FMassEntityHandle Entity, FOutputDevic
 {
 	for (const FMassArchetypeFragmentConfig& FragmentConfig : FragmentConfigs)
 	{
-		void* Data = GetFragmentDataForEntityChecked(FragmentConfig.FragmentType, Entity.Index);
+		const void* Data = GetFragmentDataForEntityChecked(FragmentConfig.FragmentType, Entity);
 		
 		FString FragmentName = FragmentConfig.FragmentType->GetName();
 		FragmentName.RemoveFromStart(InPrefix);
diff --git a/Engine/Source/Runtime/MassEntity/Private/MassDebugger.cpp b/Engine/Source/Runtime/MassEntity/Private/MassDebugger.cpp
--- a/Engine/Source/Runtime/MassEntity/Private/MassDebugger.cpp
+++ b/Engine/Source/Runtime/MassEntity/Private/MassDebugger.cpp
@@ -1527,7 +1527,7 @@ bool FMassDebugger::GetFragmentData(const FMassEntityManager& EntityManager, con
 	TSharedPtr<FMassArchetypeData> Archetype = EntityManager.DebugGetEntityStorageInterface().GetArchetypeAsShared(Entity.Index);
 	if (Archetype.IsValid())
 	{
-		void* FragmentData = Archetype->GetFragmentDataForEntity(FragmentType, Entity.Index);
+		const void* FragmentData = Archetype->GetFragmentDataForEntity(FragmentType, Entity);
 		if (FragmentData)
 		{
 			const UStruct* AsUStruct = FragmentType;
diff --git a/Engine/Source/Runtime/MassEntity/Private/MassEntityManager.cpp b/Engine/Source/Runtime/MassEntity/Private/MassEntityManager.cpp
--- a/Engine/Source/Runtime/MassEntity/Private/MassEntityManager.cpp
+++ b/Engine/Source/Runtime/MassEntity/Private/MassEntityManager.cpp
@@ -1088,9 +1088,9 @@ void FMassEntityManager::AddFragmentToEntity(FMassEntityHandle Entity, TNotNull<
 
 	CheckIfEntityIsActive(Entity);
 
-	const FMassArchetypeCompositionDescriptor Descriptor(InternalAddFragmentListToEntityChecked(Entity, FMassFragmentBitSet(FragmentType)));
+	FMassFragmentBitSet AddedBits = InternalAddFragmentListToEntityChecked(Entity, FMassFragmentBitSet(FragmentType));
 
-	ObserverManager.OnPostCompositionAdded(Entity, Descriptor);
+	ObserverManager.OnPostCompositionAdded(Entity, AddedBits);
 }
 
 void FMassEntityManager::AddFragmentToEntity(FMassEntityHandle Entity, TNotNull<const UScriptStruct*> FragmentType, const FStructInitializationCallback& Initializer)
@@ -1102,10 +1102,10 @@ void FMassEntityManager::AddFragmentToEntity(FMassEntityHandle Entity, TNotNull<
 	FMassFragmentBitSet Fragments = InternalAddFragmentListToEntityChecked(Entity, FMassFragmentBitSet(FragmentType));
 	FMassArchetypeData* CurrentArchetype = GetEntityStorageInterface().GetArchetype(Entity.Index);
 	check(CurrentArchetype);
-	void* FragmentData = CurrentArchetype->GetFragmentDataForEntity(FragmentType, Entity.Index);
+	void* FragmentData = CurrentArchetype->GetMutableFragmentDataForEntity(FragmentType, Entity);
 	Initializer(FragmentData, *FragmentType);
 
-	ObserverManager.OnPostCompositionAdded(Entity, FMassElementBitSet(MoveTemp(Fragments)));
+	ObserverManager.OnPostCompositionAdded(Entity, Fragments);
 }
 
 void FMassEntityManager::AddFragmentListToEntity(FMassEntityHandle Entity, TConstArrayView<const UScriptStruct*> FragmentList)
@@ -1114,9 +1114,9 @@ void FMassEntityManager::AddFragmentListToEntity(FMassEntityHandle Entity, TCons
 
 	CheckIfEntityIsActive(Entity);
 
-	const FMassArchetypeCompositionDescriptor Descriptor(InternalAddFragmentListToEntityChecked(Entity, FMassFragmentBitSet(FragmentList)));
+	const FMassElementBitSet AddedElements(InternalAddFragmentListToEntityChecked(Entity, FMassElementBitSet(FragmentList)));
 	
-	ObserverManager.OnPostCompositionAdded(Entity, Descriptor);
+	ObserverManager.OnPostCompositionAdded(Entity, AddedElements);
 }
 
 void FMassEntityManager::AddCompositionToEntity_GetDelta(FMassEntityHandle Entity, FMassArchetypeCompositionDescriptor& InOutDescriptor, const FMassArchetypeSharedFragmentValues* AddedSharedFragmentValues)
@@ -1234,32 +1234,36 @@ void FMassEntityManager::InternalReleaseEntity(FMassEntityHandle Entity)
 	GetEntityStorageInterface().ForceReleaseOne(Entity);
 }
 
-FMassFragmentBitSet FMassEntityManager::InternalAddFragmentListToEntityChecked(FMassEntityHandle Entity, const FMassFragmentBitSet& InFragments)
+FMassElementBitSet FMassEntityManager::InternalAddFragmentListToEntityChecked(FMassEntityHandle Entity, const FMassElementBitSet& FragmentList)
 {
 	TSharedPtr<FMassArchetypeData>& OldArchetype = GetEntityStorageInterface().GetArchetypeAsShared(Entity.Index);
 	check(OldArchetype);
 
-	UE_CLOG(OldArchetype->GetCompositionBitSet().HasAny(InFragments), LogMass, Log
-		, TEXT("Trying to add a new fragment type to an entity, but it already has some of them. (%s)")
-		, *InFragments.GetOverlap(OldArchetype->GetCompositionBitSet()).DebugGetStringDesc());
+	FMassElementBitSet SparseElements = FMassElementBitSet::GetAllSparseElements() & FragmentList;
+	FMassElementBitSet RegularElements = FragmentList - SparseElements;
 
-	FMassFragmentBitSet NewFragments = InFragments - OldArchetype->GetCompositionBitSet();
+	FMassFragmentBitSet NewFragments = RegularElements - OldArchetype->GetCompositionBitSet();
 	if (NewFragments.IsEmpty() == false)
 	{
 		InternalAddFragmentListToEntity(Entity, NewFragments);
 	}
-	return MoveTemp(NewFragments);
+	if (SparseElements.IsEmpty() == false)
+	{
+		AddSparseElementsToEntity(Entity, SparseElements);
+	}
+
+	return NewFragments + SparseElements;
 }
 
-void FMassEntityManager::InternalAddFragmentListToEntity(FMassEntityHandle Entity, const FMassFragmentBitSet& InFragments)
+void FMassEntityManager::InternalAddFragmentListToEntity(FMassEntityHandle Entity, const FMassElementBitSet& FragmentList)
 {
-	checkf(InFragments.IsEmpty() == false, TEXT("%hs is intended for internal calls with non empty NewFragments parameter"), __FUNCTION__);
+	checkf(FragmentList.IsEmpty() == false, TEXT("%hs is intended for internal calls with non empty NewFragments parameter"), __FUNCTION__);
 	check(GetEntityStorageInterface().IsValidIndex(Entity.Index));
 	TSharedPtr<FMassArchetypeData>& OldArchetype = GetEntityStorageInterface().GetArchetypeAsShared(Entity.Index);
 	check(OldArchetype.IsValid());
 
 	// fetch or create the new archetype
-	const FMassArchetypeHandle NewArchetypeHandle = CreateArchetype(OldArchetype, InFragments);
+	const FMassArchetypeHandle NewArchetypeHandle = CreateArchetype(OldArchetype, FragmentList);
 	checkf(NewArchetypeHandle.DataPtr != OldArchetype, TEXT("%hs is intended for internal calls with non overlapping fragment list."), __FUNCTION__);
 
 	// Move the entity over
@@ -1291,32 +1295,45 @@ void FMassEntityManager::RemoveFragmentFromEntity(FMassEntityHandle Entity, cons
 	RemoveFragmentListFromEntity(Entity, MakeArrayView(&FragmentType, 1));
 }
 
-void FMassEntityManager::RemoveFragmentListFromEntity(FMassEntityHandle Entity, TConstArrayView<const UScriptStruct*> FragmentList)
+void FMassEntityManager::RemoveFragmentListFromEntity(FMassEntityHandle EntityHandle, TConstArrayView<const UScriptStruct*> FragmentList)
 {
 	CHECK_SYNC_API();
 
-	CheckIfEntityIsActive(Entity);
+	CheckIfEntityIsActive(EntityHandle);
 
-	FMassArchetypeData* OldArchetype = GetEntityStorageInterface().GetArchetype(Entity.Index);
+	FMassArchetypeData* OldArchetype = GetEntityStorageInterface().GetArchetype(EntityHandle.Index);
 	check(OldArchetype);
 
-	const FMassFragmentBitSet FragmentsToRemove(FragmentList);
+	const FMassElementBitSet FragmentsToRemove(FragmentList);
+	const FMassElementBitSet SparseElements = FragmentsToRemove.GetSparseElements();
+	const bool bRegularFragmentsChange = OldArchetype->GetCompositionBitSet().HasAny(FragmentsToRemove);
 
-	if (OldArchetype->GetCompositionBitSet().HasAny(FragmentsToRemove))
+	if (bRegularFragmentsChange)
 	{
 		// If all the fragments got removed this will result in fetching of the empty archetype
 		const FMassArchetypeCompositionDescriptor NewComposition(OldArchetype->GetCompositionBitSet() - FragmentsToRemove);
 		const FMassArchetypeHandle NewArchetypeHandle = CreateArchetype(NewComposition, FMassArchetypeCreationParams(*OldArchetype));
 
 		// Find overlap.  It isn't guaranteed that the old archetype has all the fragments being removed.
-		ObserverManager.OnPreCompositionRemoved(Entity, FMassElementBitSet(OldArchetype->GetCompositionBitSet().GetOverlap(FragmentsToRemove)));
+		ObserverManager.OnPreCompositionRemoved(EntityHandle, OldArchetype->GetCompositionBitSet().GetOverlap(FragmentsToRemove));
 
 		// Move the entity over
 		FMassArchetypeData& NewArchetype = FMassArchetypeHelper::ArchetypeDataFromHandleChecked(NewArchetypeHandle);
 		NewArchetype.CopyDebugNamesFrom(*OldArchetype);
-		OldArchetype->MoveEntityToAnotherArchetype(Entity, NewArchetype);
+		OldArchetype->MoveEntityToAnotherArchetype(EntityHandle, NewArchetype);
 
-		GetEntityStorageInterface().SetArchetypeFromShared(Entity.Index, NewArchetypeHandle.DataPtr);
+		GetEntityStorageInterface().SetArchetypeFromShared(EntityHandle.Index, NewArchetypeHandle.DataPtr);
+	}
+	if (SparseElements.IsEmpty() == false)
+	{
+		// if there were regular fragments removed then the entity's archetype has changed.
+		FMassArchetypeData* CurrentArchetype = bRegularFragmentsChange
+			? GetEntityStorageInterface().GetArchetype(EntityHandle.Index)
+			: OldArchetype;
+
+		check(CurrentArchetype);
+
+		CurrentArchetype->RemoveSparseElementFromEntity(EntityHandle, SparseElements);
 	}
 }
 
@@ -2095,18 +2112,18 @@ void* FMassEntityManager::InternalGetFragmentDataChecked(FMassEntityHandle Entit
 	// note that FragmentType is guaranteed to be of valid type - it's either statically checked by the template versions
 	// or `checkf`ed by the non-template one
 	CheckIfEntityIsActive(Entity);
-	const FMassArchetypeData* CurrentArchetype = GetEntityStorageInterface().GetArchetype(Entity.Index);
+	FMassArchetypeData* CurrentArchetype = GetEntityStorageInterface().GetArchetype(Entity.Index);
 	check(CurrentArchetype);
-	return CurrentArchetype->GetFragmentDataForEntityChecked(FragmentType, Entity.Index);
+	return CurrentArchetype->GetMutableFragmentDataForEntityChecked(FragmentType, Entity);
 }
 
 void* FMassEntityManager::InternalGetFragmentDataPtr(FMassEntityHandle Entity, const UScriptStruct* FragmentType) const
 {
 	// note that FragmentType is guaranteed to be of valid type - it's either statically checked by the template versions
 	// or `checkf`ed by the non-template one
-	const FMassArchetypeData* CurrentArchetype = GetEntityStorageInterface().GetArchetype(Entity.Index);
+	FMassArchetypeData* CurrentArchetype = GetEntityStorageInterface().GetArchetype(Entity.Index);
 	check(CurrentArchetype);
-	return CurrentArchetype->GetFragmentDataForEntity(FragmentType, Entity.Index);
+	return CurrentArchetype->GetMutableFragmentDataForEntity(FragmentType, Entity);
 }
 
 const FConstSharedStruct* FMassEntityManager::InternalGetConstSharedFragmentPtr(FMassEntityHandle Entity, const UScriptStruct* ConstSharedFragmentType) const
diff --git a/Engine/Source/Runtime/MassEntity/Private/MassEntityView.cpp b/Engine/Source/Runtime/MassEntity/Private/MassEntityView.cpp
--- a/Engine/Source/Runtime/MassEntity/Private/MassEntityView.cpp
+++ b/Engine/Source/Runtime/MassEntity/Private/MassEntityView.cpp
@@ -114,13 +114,13 @@ void* FMassEntityView::GetSharedFragmentPtrChecked(const UScriptStruct& Fragment
 bool FMassEntityView::HasTag(TNotNull<const UScriptStruct*> TagType) const
 {
 	checkSlow(EntityDataHandle.IsValid(Archetype));
-	return Archetype && Archetype->HasTagType(TagType);
+	return Archetype && (UE::Mass::IsSparse(TagType) ? Archetype->HasSparseElementForEntity(TagType, EntityDataHandle) : Archetype->HasTagType(TagType));
 }
 
 bool FMassEntityView::HasFragment(TNotNull<const UScriptStruct*> FragmentType) const
 {
 	checkSlow(EntityDataHandle.IsValid(Archetype));
-	return Archetype && Archetype->HasFragmentType(FragmentType);
+	return Archetype && (UE::Mass::IsSparse(FragmentType) ? Archetype->HasSparseElementForEntity(FragmentType, EntityDataHandle) : Archetype->HasFragmentType(FragmentType));
 }
 
 bool FMassEntityView::HasElement(TNotNull<const UScriptStruct*> ElementType, const UE::Mass::EIncludeSparseElements IncludeSparseElements) const
diff --git a/Engine/Source/Runtime/MassEntity/Public/MassEntityManager.h b/Engine/Source/Runtime/MassEntity/Public/MassEntityManager.h
--- a/Engine/Source/Runtime/MassEntity/Public/MassEntityManager.h
+++ b/Engine/Source/Runtime/MassEntity/Public/MassEntityManager.h
@@ -472,7 +472,7 @@ public:
 	 * This function handles sparse elements.
 	 * The function will assert on unhandled types.
 	 */
-	void AddElementToEntity(FMassEntityHandle Entity, TNotNull<const UScriptStruct*> ElementType);
+	UE_API void AddElementToEntity(FMassEntityHandle Entity, TNotNull<const UScriptStruct*> ElementType);
 
 	/*
 	 * Removes ElementType from the entities, treating it accordingly based on which element type it represents (i.e. Fragment or Tag).
@@ -484,7 +484,7 @@ public:
 	 * Removes ElementType from the entity, treating it accordingly based on which element type it represents (i.e. Fragment or Tag).
 	 * The function will assert on unhandled types.
 	 */
-	void RemoveElementFromEntity(FMassEntityHandle Entity, TNotNull<const UScriptStruct*> ElementType);
+	UE_API void RemoveElementFromEntity(FMassEntityHandle Entity, TNotNull<const UScriptStruct*> ElementType);
 
 	/** @return whether Entity has an element of ElementType */
 	UE_API bool DoesEntityHaveElement(FMassEntityHandle Entity, TNotNull<const UScriptStruct*> ElementType) const;
@@ -1314,13 +1314,13 @@ private:
 	 *  Adds fragments in FragmentList to the entity indicated by EntityHandle. Only the unique fragments will be added.
 	 *  @return Bitset for the added fragments (might be empty or a subset of `InFragments` depending on the current archetype fragments)
 	 */
-	FMassFragmentBitSet InternalAddFragmentListToEntityChecked(FMassEntityHandle EntityHandle, const FMassFragmentBitSet& InFragments);
+	FMassElementBitSet InternalAddFragmentListToEntityChecked(FMassEntityHandle EntityHandle, const FMassElementBitSet& FragmentList);
 
 	/** 
 	 *  Similar to InternalAddFragmentListToEntity but expects NewFragmentList not overlapping with current entity's
 	 *  fragment list. It's callers responsibility to ensure that's true. Failing this will cause a `check` fail.
 	 */
-	void InternalAddFragmentListToEntity(FMassEntityHandle EntityHandle, const FMassFragmentBitSet& InFragments);
+	void InternalAddFragmentListToEntity(FMassEntityHandle EntityHandle, const FMassElementBitSet& FragmentList);
 	/** Note that it's the caller's responsibility to ensure `FragmentType` is a kind of FMassFragment */
 	UE_API void* InternalGetFragmentDataChecked(FMassEntityHandle EntityHandle, const UScriptStruct* FragmentType) const;
 	/** Note that it's the caller's responsibility to ensure `FragmentType` is a kind of FMassFragment */
