[Mass] Spares elements comments in MassArchetypeData improvements

also one minor optimization in FMassArchetypeData::GetMutableFragmentDataForEntity

#rnx
#Synced-CL 51590771

diff --git a/Engine/Source/Runtime/MassEntity/Internal/MassArchetypeData.h b/Engine/Source/Runtime/MassEntity/Internal/MassArchetypeData.h
--- a/Engine/Source/Runtime/MassEntity/Internal/MassArchetypeData.h
+++ b/Engine/Source/Runtime/MassEntity/Internal/MassArchetypeData.h
@@ -486,8 +486,11 @@ public:
 	FStructView AddSparseElementToEntity(FMassEntityHandle EntityHandle, TNotNull<const UScriptStruct*> ElementType);
 	/** Works with all sparse element types (fragments and tags). */
 	void RemoveSparseElementFromEntity(FMassEntityHandle EntityHandle, TNotNull<const UScriptStruct*> ElementType);
+	/** Removes all sparse elements indicated by SparseElementsBitSet from the given entity. Works with all sparse element types (fragments and tags). */
 	void RemoveSparseElementFromEntity(FMassEntityHandle EntityHandle, const FMassElementBitSet& SparseElementsBitSet);
-	/** Expects non-tag sparse element type. Will assert if the assumption is broken. */
+	/** Expects non-tag sparse element type. Will assert if the assumption is broken.
+	 * @todo this method is const but returns mutable data (const doesn't propagate through the SparseElementsStorage
+	 * raw pointer indirection). Consider making non-const to match GetMutableFragmentDataForEntity's convention. */
 	FStructView GetMutableSparseElementDataForEntity(TNotNull<const UScriptStruct*> ElementType, const FMassEntityHandle EntityHandle) const;
 	/** Expects non-tag sparse element type. Will assert if the assumption is broken. */
 	FConstStructView GetSparseElementDataForEntity(TNotNull<const UScriptStruct*> ElementType, const FMassEntityHandle EntityHandle) const;
@@ -553,7 +556,8 @@ public:
 	int32 CompactEntities(const double TimeAllowed);
 
 	/**
-	 * Moves the entity from this archetype to another, will only copy all matching fragment types
+	 * Moves the entity from this archetype to another, will only copy all matching fragment types.
+	 * Uses an internally-cached fragment index mapping for the source->target archetype pair.
 	 * @param Entity is the entity to move
 	 * @param NewArchetype the archetype to move to
 	 * @param SharedFragmentValuesOverride if provided will override all given Entity's shared fragment values
diff --git a/Engine/Source/Runtime/MassEntity/Private/MassArchetypeData.cpp b/Engine/Source/Runtime/MassEntity/Private/MassArchetypeData.cpp
--- a/Engine/Source/Runtime/MassEntity/Private/MassArchetypeData.cpp
+++ b/Engine/Source/Runtime/MassEntity/Private/MassArchetypeData.cpp
@@ -657,13 +657,14 @@ bool FMassArchetypeData::HasSparseElementForEntity(TNotNull<const UScriptStruct*
 void* FMassArchetypeData::GetMutableFragmentDataForEntityChecked(TNotNull<const UScriptStruct*> FragmentType, const FMassEntityHandle EntityHandle)
 {
 	void* ReturnPtr = GetMutableFragmentDataForEntity(FragmentType, EntityHandle);
-	checkf(ReturnPtr, TEXT("Entity %s has neither a regular nor a sparse fragment of type %s"), *EntityHandle.DebugGetDescription(), *FragmentType->GetName())
+	checkf(ReturnPtr, TEXT("Entity %s has neither a regular nor a sparse fragment of type %s"), *EntityHandle.DebugGetDescription(), *FragmentType->GetName());
 	return ReturnPtr;
 }
 
 const void* FMassArchetypeData::GetFragmentDataForEntityChecked(TNotNull<const UScriptStruct*> FragmentType, const FMassEntityHandle EntityHandle) const
 {
-	// ugly but internal, so :shrug:
+	// const delegates to mutable variant via const_cast to avoid code duplication. Safe because
+	// the mutable accessor's return value is not written through in the const call path.
 	return const_cast<FMassArchetypeData*>(this)->GetMutableFragmentDataForEntityChecked(FragmentType, EntityHandle);
 }
 
@@ -671,25 +672,23 @@ void* FMassArchetypeData::GetMutableFragmentDataForEntity(TNotNull<const UScript
 {
 	checkf(UE::Mass::IsA<FMassFragment>(FragmentType), TEXT("%s is not a fragment type"), *FragmentType->GetName());
 
-	if (UE::Mass::IsSparse(FragmentType) == false)
+	if (const int32* FragmentIndex = FragmentIndexMap.Find(FragmentType))
 	{
-		if (const int32* FragmentIndex = FragmentIndexMap.Find(FragmentType))
-		{
-			FMassRawEntityInChunkData InternalIndex = MakeEntityHandle(EntityHandle.Index);
-			// failing the below Find means given entity's archetype is missing given FragmentType
-			return GetFragmentData(*FragmentIndex, InternalIndex);
-		}
+		FMassRawEntityInChunkData InternalIndex = MakeEntityHandle(EntityHandle.Index);
+		// failing the below Find means given entity's archetype is missing given FragmentType
+		return GetFragmentData(*FragmentIndex, InternalIndex);
 	}
-	
-	// Might be sparse element
+
+	// Not in the regular fragment map - look it up as a sparse fragment.
 	// Note that it's fine to return the pointer from the resulting view, since worst-case-scenario
-	// it returns an invalid FStructView, which will result in this function returning nullptr.
-	return GetMutableSparseElementDataForEntity(FragmentType, EntityHandle).GetMemory();
+	// it returns a default-constructed FStructView whose GetMemory() returns nullptr.
+	return UE::Mass::IsSparse(FragmentType) ? GetMutableSparseElementDataForEntity(FragmentType, EntityHandle).GetMemory() : nullptr;
 }
 
 const void* FMassArchetypeData::GetFragmentDataForEntity(TNotNull<const UScriptStruct*> FragmentType, const FMassEntityHandle EntityHandle) const
 {
-	// ugly but internal, so :shrug:
+	// const delegates to mutable variant via const_cast to avoid code duplication. Safe because
+	// the mutable accessor's return value is not written through in the const call path.
 	return const_cast<FMassArchetypeData*>(this)->GetMutableFragmentDataForEntity(FragmentType, EntityHandle);
 }
 
