[Mass] Unified sparse/regular mass elements API in MassRequirements

#rnx
#Synced-CL 51590771

diff --git a/Engine/Source/Runtime/MassEntity/Public/MassRequirements.h b/Engine/Source/Runtime/MassEntity/Public/MassRequirements.h
--- a/Engine/Source/Runtime/MassEntity/Public/MassRequirements.h
+++ b/Engine/Source/Runtime/MassEntity/Public/MassRequirements.h
@@ -202,31 +202,40 @@ struct FMassFragmentRequirements
 	FMassFragmentRequirements& AddRequirement(const UScriptStruct* FragmentType, const EMassFragmentAccess AccessMode, const EMassFragmentPresence Presence = EMassFragmentPresence::All)
 	{
 		checkf(bInitialized, TEXT("Modifying requirements before initialization is not supported."));
-		checkf(FragmentRequirements.FindByPredicate([FragmentType](const FMassFragmentRequirementDescription& Item) { return Item.StructType == FragmentType; }) == nullptr
-			, TEXT("Duplicated requirements are not supported. %s already present"), *GetNameSafe(FragmentType));
+		check(FragmentType);
 
-		if (Presence != EMassFragmentPresence::None)
+		if (UE::Mass::IsSparse(FragmentType))
 		{
-			FragmentRequirements.Emplace(FragmentType, AccessMode, Presence);
+			AddSparseRequirement(FragmentType, Presence);
 		}
-
-		switch (Presence)
+		else
 		{
-		case EMassFragmentPresence::All:
-			RequiredAllFragments.Add(FragmentType);
-			break;
-		case EMassFragmentPresence::Any:
-			RequiredAnyFragments.Add(FragmentType);
-			break;
-		case EMassFragmentPresence::Optional:
-			RequiredOptionalFragments.Add(FragmentType);
-			break;
-		case EMassFragmentPresence::None:
-			RequiredNoneFragments.Add(FragmentType);
-			break;
+			checkf(FragmentRequirements.FindByPredicate([FragmentType](const FMassFragmentRequirementDescription& Item) { return Item.StructType == FragmentType; }) == nullptr
+				, TEXT("Duplicated requirements are not supported. %s already present"), *GetNameSafe(FragmentType));
+
+			if (Presence != EMassFragmentPresence::None)
+			{
+				FragmentRequirements.Emplace(FragmentType, AccessMode, Presence);
+			}
+
+			switch (Presence)
+			{
+			case EMassFragmentPresence::All:
+				RequiredAllFragments.Add(FragmentType);
+				break;
+			case EMassFragmentPresence::Any:
+				RequiredAnyFragments.Add(FragmentType);
+				break;
+			case EMassFragmentPresence::Optional:
+				RequiredOptionalFragments.Add(FragmentType);
+				break;
+			case EMassFragmentPresence::None:
+				RequiredNoneFragments.Add(FragmentType);
+				break;
+			}
+			// force recaching the next time this query is used or the following CacheArchetypes call.
+			IncrementChangeCounter();
 		}
-		// force recaching the next time this query is used or the following CacheArchetypes call.
-		IncrementChangeCounter();
 		return *this;
 	}
 
@@ -240,28 +249,35 @@ struct FMassFragmentRequirements
 
 		static_assert(UE::Mass::CFragment<T>, MASS_INVALID_FRAGMENT_MSG);
 
-		if (Presence != EMassFragmentPresence::None)
+		if constexpr (UE::Mass::CSparse<T>)
 		{
-			FragmentRequirements.Emplace(T::StaticStruct(), AccessMode, Presence);
+			AddSparseRequirement(T::StaticStruct(), Presence);
 		}
-
-		switch (Presence)
+		else
 		{
-		case EMassFragmentPresence::All:
-			RequiredAllFragments.Add<T>();
-			break;
-		case EMassFragmentPresence::Any:
-			RequiredAnyFragments.Add<T>();
-			break;
-		case EMassFragmentPresence::Optional:
-			RequiredOptionalFragments.Add<T>();
-			break;
-		case EMassFragmentPresence::None:
-			RequiredNoneFragments.Add<T>();
-			break;
+			if (Presence != EMassFragmentPresence::None)
+			{
+				FragmentRequirements.Emplace(T::StaticStruct(), AccessMode, Presence);
+			}
+
+			switch (Presence)
+			{
+			case EMassFragmentPresence::All:
+				RequiredAllFragments.Add<T>();
+				break;
+			case EMassFragmentPresence::Any:
+				RequiredAnyFragments.Add<T>();
+				break;
+			case EMassFragmentPresence::Optional:
+				RequiredOptionalFragments.Add<T>();
+				break;
+			case EMassFragmentPresence::None:
+				RequiredNoneFragments.Add<T>();
+				break;
+			}
+			// force recaching the next time this query is used or the following CacheArchetypes call.
+			IncrementChangeCounter();
 		}
-		// force recaching the next time this query is used or the following CacheArchetypes call.
-		IncrementChangeCounter();
 		return *this;
 	}
 
@@ -269,22 +285,30 @@ struct FMassFragmentRequirements
 	{
 		checkf(bInitialized, TEXT("Modifying requirements before initialization is not supported."));
 		checkf(int(Presence) != int(EMassFragmentPresence::MAX), TEXT("MAX presence is not a valid value for AddTagRequirement"));
-		switch (Presence)
+
+		if (UE::Mass::IsSparse(TagType))
 		{
-		case EMassFragmentPresence::All:
-			RequiredAllTags.Add(TagType);
-			break;
-		case EMassFragmentPresence::Any:
-			RequiredAnyTags.Add(TagType);
-			break;
-		case EMassFragmentPresence::None:
-			RequiredNoneTags.Add(TagType);
-			break;
-		case EMassFragmentPresence::Optional:
-			RequiredOptionalTags.Add(TagType);
-			break;
+			AddSparseRequirement(TagType, Presence);
+		}
+		else
+		{
+			switch (Presence)
+			{
+			case EMassFragmentPresence::All:
+				RequiredAllTags.Add(TagType);
+				break;
+			case EMassFragmentPresence::Any:
+				RequiredAnyTags.Add(TagType);
+				break;
+			case EMassFragmentPresence::None:
+				RequiredNoneTags.Add(TagType);
+				break;
+			case EMassFragmentPresence::Optional:
+				RequiredOptionalTags.Add(TagType);
+				break;
+			}
+			IncrementChangeCounter();
 		}
-		IncrementChangeCounter();
 		return *this;
 	}
 
@@ -298,23 +322,32 @@ struct FMassFragmentRequirements
 	{
 		checkf(bInitialized, TEXT("Modifying requirements before initialization is not supported."));
 		checkf(int(Presence) != int(EMassFragmentPresence::MAX), TEXT("MAX presence is not a valid value for AddTagRequirement"));
-		static_assert(UE::Mass::CTag<T>, "Given struct doesn't represent a valid tag type. Make sure to inherit from FMassFragment or one of its child-types.");
-		switch (Presence)
+
+		static_assert(UE::Mass::CTag<T>, "Given struct doesn't represent a valid tag type. Make sure to inherit from FMassTag or one of its child-types.");
+
+		if constexpr (UE::Mass::CSparse<T>)
 		{
-		case EMassFragmentPresence::All:
-			RequiredAllTags.Add<T>();
-			break;
-		case EMassFragmentPresence::Any:
-			RequiredAnyTags.Add<T>();
-			break;
-		case EMassFragmentPresence::None:
-			RequiredNoneTags.Add<T>();
-			break;
-		case EMassFragmentPresence::Optional:
-			RequiredOptionalTags.Add<T>();
-			break;
+			AddSparseRequirement(T::StaticStruct(), Presence);
+		}
+		else 
+		{
+			switch (Presence)
+			{
+			case EMassFragmentPresence::All:
+				RequiredAllTags.Add<T>();
+				break;
+			case EMassFragmentPresence::Any:
+				RequiredAnyTags.Add<T>();
+				break;
+			case EMassFragmentPresence::None:
+				RequiredNoneTags.Add<T>();
+				break;
+			case EMassFragmentPresence::Optional:
+				RequiredOptionalTags.Add<T>();
+				break;
+			}
+			IncrementChangeCounter();
 		}
-		IncrementChangeCounter();
 		return *this;
 	}
 
