[Mass] Remove stale reviewer note in MassCommands.h

#rnx
#Synced-CL 51621734

diff --git a/Engine/Source/Runtime/MassEntity/Public/MassCommands.h b/Engine/Source/Runtime/MassEntity/Public/MassCommands.h
--- a/Engine/Source/Runtime/MassEntity/Public/MassCommands.h
+++ b/Engine/Source/Runtime/MassEntity/Public/MassCommands.h
@@ -142,8 +142,6 @@ PRAGMA_ENABLE_DEPRECATION_WARNINGS
 #endif // CSV_PROFILER_STATS || WITH_MASSENTITY_DEBUG
 
 protected:
-	// @todo note for reviewers - I could use an opinion if having a virtual function per-command would be a more 
-	// preferable way of asking commands if there's anything to do.
 	bool bHasWork = false;
 	EMassCommandOperationType OperationType = EMassCommandOperationType::None;
 
