[Mass] Reword comment in MassProcessorDependencySolver.h per [REDACTED_HANDLE]

#rnx
#Synced-CL 51621734

diff --git a/Engine/Source/Runtime/MassEntity/Public/MassProcessorDependencySolver.h b/Engine/Source/Runtime/MassEntity/Public/MassProcessorDependencySolver.h
--- a/Engine/Source/Runtime/MassEntity/Public/MassProcessorDependencySolver.h
+++ b/Engine/Source/Runtime/MassEntity/Public/MassProcessorDependencySolver.h
@@ -122,11 +122,11 @@ struct FMassProcessorDependencySolver
 		/** indicates how often given node can be found in dependencies sequence for other nodes  */
 		int32 TotalWaitingNodes = 0;
 		/**
-		 * Indicates the maximum execution priority represented by this node or any od the nodes
-		 * that depend on it in a logical sense - i.e. it does not include the nodes that are dependencies just by blocing required resources
-		 * @todo reword this comment
-		 * Note that we're using a larger type than UMassProcessor.ExecutionPriority (int32 vs int16)
-		 * to not have to handle overflow in UpdateExecutionPriority
+		 * The highest execution priority across this node and all nodes in its logical dependency chain.
+		 * Nodes that are only dependencies due to resource contention (i.e. blocking shared resources)
+		 * are excluded from this calculation.
+		 * Note: we use int32 rather than UMassProcessor.ExecutionPriority's int16 to avoid overflow
+		 * when propagating priorities in UpdateExecutionPriority.
 		 */
 		int32 MaxExecutionPriority = 0;
 		/** 
