[Mass] Added extended comment to FMassEntityManager::GetFragmentDataStruct

#rnx
#Synced-CL 51590771

diff --git a/Engine/Source/Runtime/MassEntity/Public/MassEntityManager.h b/Engine/Source/Runtime/MassEntity/Public/MassEntityManager.h
--- a/Engine/Source/Runtime/MassEntity/Public/MassEntityManager.h
+++ b/Engine/Source/Runtime/MassEntity/Public/MassEntityManager.h
@@ -698,6 +698,13 @@ public:
 		return (FragmentType*)InternalGetFragmentDataPtr(EntityHandle, FragmentType::StaticStruct());
 	}
 
+	/** Returns a FStructView for the given fragment on the specified entity. The view points
+	 *  directly into archetype chunk storage.
+	 *  WARNING: For initializing fragments on newly created entities, prefer the
+	 *  CreateEntity(TConstArrayView<FInstancedStruct>,...) overload which sets initial values
+	 *  atomically. The two-step pattern (CreateEntity + GetFragmentDataStruct + mutate) is only
+	 *  safe when no processors are running (e.g., during BeginPlay or subsystem Initialize).
+	 *  @return an empty FStructView if the entity does not have the fragment. */
 	FStructView GetFragmentDataStruct(FMassEntityHandle EntityHandle, const UScriptStruct* FragmentType) const
 	{
 		checkf(UE::Mass::IsA<FMassFragment>(FragmentType)
