[Mass] Split CFragment static_assert into targeted diagnostics

Adds MASS_STATIC_CHECK_FRAGMENT macro that provides two separate static_assert messages: one for inheritance (must derive from FMassFragment) and one for trivial copyability (FString, TArray, TMap forbidden). Previously these were combined into a single message that scrolled offscreen, making it unclear which condition failed.

#rnx
#Synced-CL 51590771

#Synced-CL 51590771

diff --git a/Engine/Plugins/Runtime/MassGameplay/Source/MassSpawner/Public/MassEntityTemplate.h b/Engine/Plugins/Runtime/MassGameplay/Source/MassSpawner/Public/MassEntityTemplate.h
--- a/Engine/Plugins/Runtime/MassGameplay/Source/MassSpawner/Public/MassEntityTemplate.h
+++ b/Engine/Plugins/Runtime/MassGameplay/Source/MassSpawner/Public/MassEntityTemplate.h
@@ -129,7 +129,7 @@ struct FMassEntityTemplateData
 	template<typename T>
 	void AddFragment()
 	{
-		static_assert(UE::Mass::CFragment<T>, MASS_INVALID_FRAGMENT_MSG);
+		MASS_STATIC_CHECK_FRAGMENT(T);
 		Composition.GetElementsBitSet().Add<T>();
 	}
 
@@ -166,7 +166,7 @@ struct FMassEntityTemplateData
 	template<typename T>
 	T& AddFragment_GetRef()
 	{
-		static_assert(UE::Mass::CFragment<T>, MASS_INVALID_FRAGMENT_MSG);
+		MASS_STATIC_CHECK_FRAGMENT(T);
 		if (!Composition.GetElementsBitSet().Contains<T>())
 		{
 			Composition.GetElementsBitSet().Add<T>();
@@ -183,7 +183,7 @@ struct FMassEntityTemplateData
 	template<typename T>
 	T* GetMutableFragment()
 	{
-		static_assert(UE::Mass::CFragment<T>, MASS_INVALID_FRAGMENT_MSG);
+		MASS_STATIC_CHECK_FRAGMENT(T);
 		FInstancedStruct* Fragment = InitialFragmentValues.FindByPredicate(FStructTypeEqualOperator(T::StaticStruct()));
 		return Fragment ? &Fragment->template GetMutable<T>() : (T*)nullptr;
 	}
diff --git a/Engine/Source/Runtime/MassEntity/Public/MassCommandBuffer.h b/Engine/Source/Runtime/MassEntity/Public/MassCommandBuffer.h
--- a/Engine/Source/Runtime/MassEntity/Public/MassCommandBuffer.h
+++ b/Engine/Source/Runtime/MassEntity/Public/MassCommandBuffer.h
@@ -182,7 +182,7 @@ struct FMassCommandBuffer
 	template<typename T>
 	void AddFragment(FMassEntityHandle Entity)
 	{
-		static_assert(UE::Mass::CFragment<T>, MASS_INVALID_FRAGMENT_MSG);
+		MASS_STATIC_CHECK_FRAGMENT(T);
 		PushCommand<FMassCommandAddFragmentsInternal<EMassCommandCheckTime::CompileTimeCheck, T>>(Entity);
 	}
 
@@ -196,7 +196,7 @@ struct FMassCommandBuffer
 	template<typename T>
 	void RemoveFragment(FMassEntityHandle Entity)
 	{
-		static_assert(UE::Mass::CFragment<T>, MASS_INVALID_FRAGMENT_MSG);
+		MASS_STATIC_CHECK_FRAGMENT(T);
 		PushCommand<FMassCommandRemoveFragmentsInternal<EMassCommandCheckTime::CompileTimeCheck, T>>(Entity);
 	}
 
diff --git a/Engine/Source/Runtime/MassEntity/Public/MassEntityConcepts.h b/Engine/Source/Runtime/MassEntity/Public/MassEntityConcepts.h
--- a/Engine/Source/Runtime/MassEntity/Public/MassEntityConcepts.h
+++ b/Engine/Source/Runtime/MassEntity/Public/MassEntityConcepts.h
@@ -67,3 +67,16 @@ namespace UE::Mass
 		return bIsSparse;
 	}
 }
+
+// Provides two targeted static_assert checks instead of the single combined CFragment concept check.
+// The first failing assert gives the user a clear, specific error message rather than the full
+// concept failure wall. Place this macro where static_assert(UE::Mass::CFragment<T>, ...) was used.
+#define MASS_STATIC_CHECK_FRAGMENT(T) \
+	static_assert(TIsDerivedFrom<UE::Mass::Clean<T>, FMassFragment>::Value, \
+		"Type must inherit from FMassFragment or one of its child-types."); \
+	static_assert(std::is_trivially_copyable_v<UE::Mass::Clean<T>> \
+		|| static_cast<bool>(TMassFragmentTraits<UE::Mass::Clean<T>>::AuthorAcceptsItsNotTriviallyCopyable), \
+		"Fragment type is not trivially copyable. Heap-allocating types like FString, TArray, " \
+		"TMap are forbidden in fragments. Use FName instead of FString, fixed arrays instead " \
+		"of TArray. To opt out: specialize TMassFragmentTraits<YourType> with " \
+		"AuthorAcceptsItsNotTriviallyCopyable = true.")
diff --git a/Engine/Source/Runtime/MassEntity/Public/MassEntityManager.h b/Engine/Source/Runtime/MassEntity/Public/MassEntityManager.h
--- a/Engine/Source/Runtime/MassEntity/Public/MassEntityManager.h
+++ b/Engine/Source/Runtime/MassEntity/Public/MassEntityManager.h
@@ -687,14 +687,14 @@ public:
 	template<typename FragmentType>
 	FragmentType& GetFragmentDataChecked(FMassEntityHandle EntityHandle) const
 	{
-		static_assert(UE::Mass::CFragment<FragmentType>, MASS_INVALID_FRAGMENT_MSG);
+		MASS_STATIC_CHECK_FRAGMENT(FragmentType);
 		return *((FragmentType*)InternalGetFragmentDataChecked(EntityHandle, FragmentType::StaticStruct()));
 	}
 
 	template<typename FragmentType>
 	FragmentType* GetFragmentDataPtr(FMassEntityHandle EntityHandle) const
 	{
-		static_assert(UE::Mass::CFragment<FragmentType>, MASS_INVALID_FRAGMENT_MSG);
+		MASS_STATIC_CHECK_FRAGMENT(FragmentType);
 		return (FragmentType*)InternalGetFragmentDataPtr(EntityHandle, FragmentType::StaticStruct());
 	}
 
diff --git a/Engine/Source/Runtime/MassEntity/Public/MassEntityView.h b/Engine/Source/Runtime/MassEntity/Public/MassEntityView.h
--- a/Engine/Source/Runtime/MassEntity/Public/MassEntityView.h
+++ b/Engine/Source/Runtime/MassEntity/Public/MassEntityView.h
@@ -48,24 +48,24 @@ struct FMassEntityView
 
 	FMassEntityHandle GetEntity() const	{ return Entity; }
 
-	/** will fail a check if the viewed entity doesn't have the given fragment */	
+	/** will fail a check if the viewed entity doesn't have the given fragment */
 	template<typename T>
 	T& GetFragmentData() const
 	{
 		static_assert(!UE::Mass::CTag<T>,
 			"Given struct doesn't represent a valid fragment type but a tag. Use HasTag instead.");
-		static_assert(UE::Mass::CFragment<T>, MASS_INVALID_FRAGMENT_MSG);
+		MASS_STATIC_CHECK_FRAGMENT(T);
 
 		return *((T*)GetFragmentPtrChecked(*T::StaticStruct()));
 	}
-		
+
 	/** if the viewed entity doesn't have the given fragment the function will return null */
 	template<typename T>
 	T* GetFragmentDataPtr() const
 	{
 		static_assert(!UE::Mass::CTag<T>,
 			"Given struct doesn't represent a valid fragment type but a tag. Use HasTag instead.");
-		static_assert(UE::Mass::CFragment<T>, MASS_INVALID_FRAGMENT_MSG);
+		MASS_STATIC_CHECK_FRAGMENT(T);
 
 		return (T*)GetFragmentPtr(*T::StaticStruct());
 	}
diff --git a/Engine/Source/Runtime/MassEntity/Public/MassRequirements.h b/Engine/Source/Runtime/MassEntity/Public/MassRequirements.h
--- a/Engine/Source/Runtime/MassEntity/Public/MassRequirements.h
+++ b/Engine/Source/Runtime/MassEntity/Public/MassRequirements.h
@@ -247,7 +247,7 @@ struct FMassFragmentRequirements
 		checkf(FragmentRequirements.FindByPredicate([](const FMassFragmentRequirementDescription& Item) { return Item.StructType == T::StaticStruct(); }) == nullptr
 			, TEXT("Duplicated requirements are not supported. %s already present"), *T::StaticStruct()->GetName());
 
-		static_assert(UE::Mass::CFragment<T>, MASS_INVALID_FRAGMENT_MSG);
+		MASS_STATIC_CHECK_FRAGMENT(T);
 
 		if constexpr (UE::Mass::CSparse<T>)
 		{
@@ -659,7 +659,7 @@ struct FMassFragmentRequirements
 	template<UE::Mass::CFragment T, typename TLinkFragment = FMassEntityLinkFragment>
 	FMassFragmentRequirements& AddLinkedEntityRequirement(const EMassFragmentAccess AccessMode)
 	{
-		static_assert(UE::Mass::CFragment<T>, MASS_INVALID_FRAGMENT_MSG);
+		MASS_STATIC_CHECK_FRAGMENT(T);
 		return AddLinkedEntityRequirement(T::StaticStruct(), AccessMode, TLinkFragment::StaticStruct());
 	}
 
@@ -709,7 +709,7 @@ struct FMassFragmentRequirements
 	template<UE::Mass::CFragment T>
 	FMassFragmentRequirements& AddIndirectFragmentRequirement(const EMassFragmentAccess AccessMode)
 	{
-		static_assert(UE::Mass::CFragment<T>, MASS_INVALID_FRAGMENT_MSG);
+		MASS_STATIC_CHECK_FRAGMENT(T);
 		return AddIndirectFragmentRequirement(T::StaticStruct(), AccessMode);
 	}
 
