[Mass] minor comments cleanup in MassEntity

- Add missing namespace closing comments (} // namespace X) across ~40 files
- Add section headers in .cpp files with multiple class implementations
- Fix trailing whitespace on #else directives

#rnx
#Synced-CL 51627411

diff --git a/Engine/Source/Runtime/MassEntity/Internal/MassArchetypeData.h b/Engine/Source/Runtime/MassEntity/Internal/MassArchetypeData.h
--- a/Engine/Source/Runtime/MassEntity/Internal/MassArchetypeData.h
+++ b/Engine/Source/Runtime/MassEntity/Internal/MassArchetypeData.h
@@ -148,7 +148,7 @@ namespace UE::Mass
 		};
 		FElementCounter SparseElementTypesCount;
 	};
-}
+} // namespace UE::Mass
 
 /** This type represents a single chunk within an archetype */
 struct FMassArchetypeChunk
diff --git a/Engine/Source/Runtime/MassEntity/Private/MassBitSetRegistry.cpp b/Engine/Source/Runtime/MassEntity/Private/MassBitSetRegistry.cpp
--- a/Engine/Source/Runtime/MassEntity/Private/MassBitSetRegistry.cpp
+++ b/Engine/Source/Runtime/MassEntity/Private/MassBitSetRegistry.cpp
@@ -23,4 +23,4 @@ namespace UE::Mass
 				})
 	{	
 	}
-}
+} // namespace UE::Mass
diff --git a/Engine/Source/Runtime/MassEntity/Private/MassCommandBuffer.cpp b/Engine/Source/Runtime/MassEntity/Private/MassCommandBuffer.cpp
--- a/Engine/Source/Runtime/MassEntity/Private/MassCommandBuffer.cpp
+++ b/Engine/Source/Runtime/MassEntity/Private/MassCommandBuffer.cpp
@@ -64,7 +64,7 @@ namespace UE::Mass::Command
 		}
 	}
 #endif
-} // UE::Mass::Command
+} // namespace UE::Mass::Command
 
 //-----------------------------------------------------------------------------
 // FMassBatchedCommand
diff --git a/Engine/Source/Runtime/MassEntity/Private/MassDebugger.cpp b/Engine/Source/Runtime/MassEntity/Private/MassDebugger.cpp
--- a/Engine/Source/Runtime/MassEntity/Private/MassDebugger.cpp
+++ b/Engine/Source/Runtime/MassEntity/Private/MassDebugger.cpp
@@ -1743,6 +1743,6 @@ namespace UE::Mass::Debug
 	{
 		return Context.GetWorld();
 	}
-}
+} // namespace UE::Mass::Debug
 
 #endif // WITH_MASSENTITY_DEBUG
diff --git a/Engine/Source/Runtime/MassEntity/Private/MassElement.cpp b/Engine/Source/Runtime/MassEntity/Private/MassElement.cpp
--- a/Engine/Source/Runtime/MassEntity/Private/MassElement.cpp
+++ b/Engine/Source/Runtime/MassEntity/Private/MassElement.cpp
@@ -37,4 +37,4 @@ namespace UE::Mass
 			break;
 		}
 	}
-}
+} // namespace UE::Mass
diff --git a/Engine/Source/Runtime/MassEntity/Private/MassEntityCollection.cpp b/Engine/Source/Runtime/MassEntity/Private/MassEntityCollection.cpp
--- a/Engine/Source/Runtime/MassEntity/Private/MassEntityCollection.cpp
+++ b/Engine/Source/Runtime/MassEntity/Private/MassEntityCollection.cpp
@@ -111,4 +111,4 @@ namespace UE::Mass
 		}
 		return true;
 	}
-}
\ No newline at end of file
+} // namespace UE::Mass
\ No newline at end of file
diff --git a/Engine/Source/Runtime/MassEntity/Private/MassEntityManager.cpp b/Engine/Source/Runtime/MassEntity/Private/MassEntityManager.cpp
--- a/Engine/Source/Runtime/MassEntity/Private/MassEntityManager.cpp
+++ b/Engine/Source/Runtime/MassEntity/Private/MassEntityManager.cpp
@@ -62,7 +62,7 @@ namespace UE::Mass::Private
 		return ElementType->IsChildOf(FMassSparseFragment::StaticStruct())
 			|| ElementType->IsChildOf(FMassSparseTag::StaticStruct());
 	}
-}
+} // namespace UE::Mass::Private
 
 //-----------------------------------------------------------------------------
 // FMassEntityManager
@@ -176,7 +176,7 @@ namespace UE::Mass::Private
 		
 		FMassEntityManager::FEntityStorageContainerType* EntityStorage = nullptr;
 	};
-}
+} // namespace UE::Mass::Private
 
 void FMassEntityManager::Initialize(const FMassEntityManagerStorageInitParams& InitializationParams)
 {
diff --git a/Engine/Source/Runtime/MassEntity/Private/MassEntityManagerStorage.cpp b/Engine/Source/Runtime/MassEntity/Private/MassEntityManagerStorage.cpp
--- a/Engine/Source/Runtime/MassEntity/Private/MassEntityManagerStorage.cpp
+++ b/Engine/Source/Runtime/MassEntity/Private/MassEntityManagerStorage.cpp
@@ -685,4 +685,4 @@ namespace UE::Mass
 			&& GenerationId == Other.GenerationId
 			&& bIsAllocated == Other.bIsAllocated;
 	}
-}
+} // namespace UE::Mass
diff --git a/Engine/Source/Runtime/MassEntity/Private/MassEntityQuery.cpp b/Engine/Source/Runtime/MassEntity/Private/MassEntityQuery.cpp
--- a/Engine/Source/Runtime/MassEntity/Private/MassEntityQuery.cpp
+++ b/Engine/Source/Runtime/MassEntity/Private/MassEntityQuery.cpp
@@ -31,7 +31,7 @@ namespace UE::Mass::Tweakables
 			{	TEXT("mass.AllowQueryParallelFor"), bAllowParallelExecution, TEXT("Controls whether EntityQueries are allowed to utilize ParallelFor construct"), ECVF_Cheat }
 		};
 	}
-}
+} // namespace UE::Mass::Tweakables
 
 //-----------------------------------------------------------------------------
 // FScopedEntityQueryContext
diff --git a/Engine/Source/Runtime/MassEntity/Private/MassEntityRelations.cpp b/Engine/Source/Runtime/MassEntity/Private/MassEntityRelations.cpp
--- a/Engine/Source/Runtime/MassEntity/Private/MassEntityRelations.cpp
+++ b/Engine/Source/Runtime/MassEntity/Private/MassEntityRelations.cpp
@@ -42,7 +42,7 @@ namespace UE::Mass
 		return FString::Printf(TEXT("[%s] %s [%s]"), *A.DebugGetDescription(), *DebugInFix, *B.DebugGetDescription());
 	}
 #endif // WITH_MASSENTITY_DEBUG
-}
+} // namespace UE::Mass
 
 //-----------------------------------------------------------------------------
 // FMassRelationRoleInstanceHandle
diff --git a/Engine/Source/Runtime/MassEntity/Private/MassEntitySubsystem.cpp b/Engine/Source/Runtime/MassEntity/Private/MassEntitySubsystem.cpp
--- a/Engine/Source/Runtime/MassEntity/Private/MassEntitySubsystem.cpp
+++ b/Engine/Source/Runtime/MassEntity/Private/MassEntitySubsystem.cpp
@@ -35,7 +35,7 @@ namespace UE::Mass::Private
 			}
 		};
 	}
-}
+} // namespace UE::Mass::Private
 
 //-----------------------------------------------------------------------------
 // UMassEntitySubsystem
diff --git a/Engine/Source/Runtime/MassEntity/Private/MassEntityTrace.cpp b/Engine/Source/Runtime/MassEntity/Private/MassEntityTrace.cpp
--- a/Engine/Source/Runtime/MassEntity/Private/MassEntityTrace.cpp
+++ b/Engine/Source/Runtime/MassEntity/Private/MassEntityTrace.cpp
@@ -123,6 +123,9 @@ enum class FFragmentType : uint8
 	SharedFragment
 };
 
+//----------------------------------------------------------------//
+//  FMassTrace
+//----------------------------------------------------------------//
 void FMassTrace::QueryCreated(const FMassEntityQuery* Query)
 {
 	UE_TRACE_LOG(MassTrace, QueryCreated, MassChannel)
@@ -157,6 +160,9 @@ void FMassTrace::QueryArchetypeAdded(const FMassEntityQuery* Query, const FMassA
 	}
 }
 
+//----------------------------------------------------------------//
+//  FMassTrace::FScopedQueryForEachTrace
+//----------------------------------------------------------------//
 FMassTrace::FScopedQueryForEachTrace::FScopedQueryForEachTrace(const FMassEntityQuery* InQuery)
 	: Query(InQuery)
 	, ArchetypeCount(0)
diff --git a/Engine/Source/Runtime/MassEntity/Private/MassEntityTypes.cpp b/Engine/Source/Runtime/MassEntity/Private/MassEntityTypes.cpp
--- a/Engine/Source/Runtime/MassEntity/Private/MassEntityTypes.cpp
+++ b/Engine/Source/Runtime/MassEntity/Private/MassEntityTypes.cpp
@@ -213,7 +213,7 @@ namespace UE::Mass::Private
 
 		return true;
 	}
-}
+} // namespace UE::Mass::Private
 
 bool FMassArchetypeSharedFragmentValues::HasSameValues(const FMassArchetypeSharedFragmentValues& Other) const
 {
diff --git a/Engine/Source/Runtime/MassEntity/Private/MassObserverManager.cpp b/Engine/Source/Runtime/MassEntity/Private/MassObserverManager.cpp
--- a/Engine/Source/Runtime/MassEntity/Private/MassObserverManager.cpp
+++ b/Engine/Source/Runtime/MassEntity/Private/MassObserverManager.cpp
@@ -22,7 +22,7 @@ namespace UE::Mass::ObserverManager
 	{
 		// Used as a template parameter for TInlineAllocator that we use when gathering UScriptStruct* of the observed types to process.
 		constexpr int InlineAllocatorElementsForOverlapTypes = 8;
-	} // Tweakables
+	} // namespace Tweakables
 
 	bool bCoalesceBufferedNotifications = false;
 
@@ -80,7 +80,7 @@ namespace UE::Mass::ObserverManager
 				}
 			}
 		};
-	} // Private
+	} // namespace Private
 
 	struct FDeprecationHelper
 	{
@@ -167,7 +167,7 @@ namespace UE::Mass::ObserverManager
 	private:
 		FMassObserverExecutionContext& RuntimeContext;
 	};
-} // UE::Mass::ObserverManager
+} // namespace UE::Mass::ObserverManager
 
 //----------------------------------------------------------------------//
 // FMassObserversMap
diff --git a/Engine/Source/Runtime/MassEntity/Private/MassObserverNotificationTypes.cpp b/Engine/Source/Runtime/MassEntity/Private/MassObserverNotificationTypes.cpp
--- a/Engine/Source/Runtime/MassEntity/Private/MassObserverNotificationTypes.cpp
+++ b/Engine/Source/Runtime/MassEntity/Private/MassObserverNotificationTypes.cpp
@@ -14,7 +14,7 @@ namespace UE::Mass::ObserverManager
 			static TSharedRef<FObserverLock> DummyObserverLock = MakeShareable(new FObserverLock());
 			return DummyObserverLock;
 		}
-	} // Private
+	} // namespace Private
 
 	//-----------------------------------------------------------------------------
 	// FObserverLock
@@ -115,4 +115,4 @@ namespace UE::Mass::ObserverManager
 
 		return true;
 	}
-} // UE::Mass::ObserverManager
+} // namespace UE::Mass::ObserverManager
diff --git a/Engine/Source/Runtime/MassEntity/Private/MassProcessingPhaseManager.cpp b/Engine/Source/Runtime/MassEntity/Private/MassProcessingPhaseManager.cpp
--- a/Engine/Source/Runtime/MassEntity/Private/MassProcessingPhaseManager.cpp
+++ b/Engine/Source/Runtime/MassEntity/Private/MassProcessingPhaseManager.cpp
@@ -33,7 +33,7 @@ namespace UE::Mass::Tweakables
 		{TEXT("mass.MakePrePhysicsTickFunctionHighPriority"), bMakePrePhysicsTickFunctionHighPriority, TEXT("Whether to make the PrePhysics tick function high priority - can minimise GameThread waits by starting parallel work as soon as possible")},
 		{TEXT("mass.UseProcessingQueue"), bUseProcessingQueue, TEXT("Enable MassProcessingQueue for executing processing phases.")},
 	};
-}
+} // namespace UE::Mass::Tweakables
 
 namespace UE::Mass::Private
 {
@@ -46,7 +46,7 @@ namespace UE::Mass::Private
 		ETickingGroup::TG_PostPhysics,	// EMassProcessingPhase::PostPhysics
 		ETickingGroup::TG_LastDemotable, // EMassProcessingPhase::FrameEnd
 	};
-} // UE::Mass::Private
+} // namespace UE::Mass::Private
 
 //----------------------------------------------------------------------//
 //  FMassProcessingPhase
diff --git a/Engine/Source/Runtime/MassEntity/Private/MassProcessingQueue.cpp b/Engine/Source/Runtime/MassEntity/Private/MassProcessingQueue.cpp
--- a/Engine/Source/Runtime/MassEntity/Private/MassProcessingQueue.cpp
+++ b/Engine/Source/Runtime/MassEntity/Private/MassProcessingQueue.cpp
@@ -27,7 +27,7 @@ namespace UE::Mass::Tweakables
 		{TEXT("mass.ProcessingQueue.MaxProcessorWorkers"), MaxProcessorWorkerCount, TEXT("Max number of concurrent Processor Workers")},
 		{TEXT("mass.ProcessingQueue.MaxChunkWorkers"), MaxChunkWorkerCount, TEXT("Max number of concurrent Chunk Workers")},
 	};
-}
+} // namespace UE::Mass::Tweakables
 
 namespace UE::Mass
 {
@@ -682,7 +682,7 @@ namespace UE::Mass
 			}
 		}
 	}
-}
+} // namespace UE::Mass
 
 #undef MASS_PROCESSING_QUEUE_LOGGING
 #undef REUSE_CHUNK_QUEUE_ENTRIES
diff --git a/Engine/Source/Runtime/MassEntity/Private/MassProcessingQueueTypes.cpp b/Engine/Source/Runtime/MassEntity/Private/MassProcessingQueueTypes.cpp
--- a/Engine/Source/Runtime/MassEntity/Private/MassProcessingQueueTypes.cpp
+++ b/Engine/Source/Runtime/MassEntity/Private/MassProcessingQueueTypes.cpp
@@ -26,7 +26,7 @@ namespace UE::Mass::Tweakables
 		{TEXT("mass.ProcessingQueue.FullLogging"), bFullLogging, TEXT("Log Everything...")},
 		{TEXT("mass.ProcessingQueue.LockLogging"), bLockLogging, TEXT("Log All lock operations...")},
 	};
-}
+} // namespace UE::Mass::Tweakables
 
 static inline const TCHAR* MassThreadTag(bool bMain)
 {
@@ -795,7 +795,7 @@ namespace UE::Mass
 		QueuedProcessor->DebugCurrentProcessorWorker = this;
 #endif
 	}
-}
+} // namespace UE::Mass
 
 #undef CHECK_THREAD_CONSISTENCY
 #undef LOG_PROC_WORKER
diff --git a/Engine/Source/Runtime/MassEntity/Private/MassProcessor.cpp b/Engine/Source/Runtime/MassEntity/Private/MassProcessor.cpp
--- a/Engine/Source/Runtime/MassEntity/Private/MassProcessor.cpp
+++ b/Engine/Source/Runtime/MassEntity/Private/MassProcessor.cpp
@@ -34,7 +34,7 @@ namespace UE::Mass::Debug
 		};
 	}
 
-}
+} // namespace UE::Mass::Debug
 
 // change to 1 to enable more detailed processing tasks logging
 #if 0
@@ -43,8 +43,8 @@ namespace UE::Mass::Debug
 #define PROCESSOR_TASK_LOG(...) 
 #endif // 0
 
-#else 
-#define PROCESSOR_TASK_LOG(...) 
+#else
+#define PROCESSOR_TASK_LOG(...)
 #endif // WITH_MASSENTITY_DEBUG
 
 class FMassProcessorTask
diff --git a/Engine/Source/Runtime/MassEntity/Private/MassProcessorDependencySolver.cpp b/Engine/Source/Runtime/MassEntity/Private/MassProcessorDependencySolver.cpp
--- a/Engine/Source/Runtime/MassEntity/Private/MassProcessorDependencySolver.cpp
+++ b/Engine/Source/Runtime/MassEntity/Private/MassProcessorDependencySolver.cpp
@@ -127,7 +127,7 @@ namespace UE::Mass::Private
 				ECVF_Default)
 		};
 	}
-}
+} // namespace UE::Mass::Private
 
 //----------------------------------------------------------------------//
 //  FMassExecutionRequirements
diff --git a/Engine/Source/Runtime/MassEntity/Private/MassRequirementAccessDetector.cpp b/Engine/Source/Runtime/MassEntity/Private/MassRequirementAccessDetector.cpp
--- a/Engine/Source/Runtime/MassEntity/Private/MassRequirementAccessDetector.cpp
+++ b/Engine/Source/Runtime/MassEntity/Private/MassRequirementAccessDetector.cpp
@@ -12,7 +12,7 @@ namespace UE::Mass::Private
 	
 	FAutoConsoleVariableRef CVarTrackRequirementsAccess(TEXT("mass.debug.TrackRequirementsAccess"), bTrackRequirementsAccess
 		, TEXT("Enables Mass processing debugging mode where we monitor thread-safety of query requirements access."));
-}
+} // namespace UE::Mass::Private
 
 void FMassRequirementAccessDetector::Initialize()
 {
@@ -76,6 +76,6 @@ namespace UE::Mass::Debug
 	{
 		EntityManager->GetRequirementAccessDetector().ReleaseAccess(Query);
 	}
-}
+} // namespace UE::Mass::Debug
 
 #endif // WITH_MASSENTITY_DEBUG
diff --git a/Engine/Source/Runtime/MassEntity/Private/MassRequirements.cpp b/Engine/Source/Runtime/MassEntity/Private/MassRequirements.cpp
--- a/Engine/Source/Runtime/MassEntity/Private/MassRequirements.cpp
+++ b/Engine/Source/Runtime/MassEntity/Private/MassRequirements.cpp
@@ -34,7 +34,7 @@ FScopedOutputDevice::~FScopedOutputDevice()
 	GOutputDevice = ActiveOutputDevice;
 }
 
-} // UE::Mass::ErrorReporting
+} // namespace UE::Mass::ErrorReporting
 
 namespace UE::Mass::Private
 {
@@ -74,7 +74,7 @@ namespace UE::Mass::Private
 			}
 		}
 	}
-}
+} // namespace UE::Mass::Private
 
 //-----------------------------------------------------------------------------
 // FMassSubsystemRequirements
diff --git a/Engine/Source/Runtime/MassEntity/Private/MassSubsystemAccess.cpp b/Engine/Source/Runtime/MassEntity/Private/MassSubsystemAccess.cpp
--- a/Engine/Source/Runtime/MassEntity/Private/MassSubsystemAccess.cpp
+++ b/Engine/Source/Runtime/MassEntity/Private/MassSubsystemAccess.cpp
@@ -10,7 +10,7 @@ namespace UE::Mass::Private
 	{
 		return *(reinterpret_cast<TSubclassOf<T>*>(&SubsystemClass));
 	}
-}
+} // namespace UE::Mass::Private
 
 //-----------------------------------------------------------------------------
 // FMassSubsystemAccess
diff --git a/Engine/Source/Runtime/MassEntity/Private/MassSubsystemBase.cpp b/Engine/Source/Runtime/MassEntity/Private/MassSubsystemBase.cpp
--- a/Engine/Source/Runtime/MassEntity/Private/MassSubsystemBase.cpp
+++ b/Engine/Source/Runtime/MassEntity/Private/MassSubsystemBase.cpp
@@ -61,7 +61,7 @@ namespace UE::Mass
 				, ECVF_Default }
 			};
 		}
-	} // UE::Mass::Private
+	} // namespace UE::Mass::Private
 } // namespace UE::Mass
 
 //-----------------------------------------------------------------------------
diff --git a/Engine/Source/Runtime/MassEntity/Private/MassTypeManager.cpp b/Engine/Source/Runtime/MassEntity/Private/MassTypeManager.cpp
--- a/Engine/Source/Runtime/MassEntity/Private/MassTypeManager.cpp
+++ b/Engine/Source/Runtime/MassEntity/Private/MassTypeManager.cpp
@@ -154,4 +154,4 @@ namespace UE::Mass
 	{
 		return UE::Mass::IsA<FMassRelation>(RelationOrElementType) && TypeDataMap.Contains(FTypeHandle(RelationOrElementType));
 	}
-}
+} // namespace UE::Mass
diff --git a/Engine/Source/Runtime/MassEntity/Private/Relations/MassChildOf.cpp b/Engine/Source/Runtime/MassEntity/Private/Relations/MassChildOf.cpp
--- a/Engine/Source/Runtime/MassEntity/Private/Relations/MassChildOf.cpp
+++ b/Engine/Source/Runtime/MassEntity/Private/Relations/MassChildOf.cpp
@@ -55,8 +55,8 @@ namespace UE::Mass
 		{
 			static FChildOfRelationTypeInitializer BuildInRelationHandlesInitializer;
 		}
-	}
-}
+	} // namespace Relations
+} // namespace UE::Mass
 
 //-----------------------------------------------------------------------------
 // UMassRelationEntityCreation
diff --git a/Engine/Source/Runtime/MassEntity/Public/MassCommandBuffer.h b/Engine/Source/Runtime/MassEntity/Public/MassCommandBuffer.h
--- a/Engine/Source/Runtime/MassEntity/Public/MassCommandBuffer.h
+++ b/Engine/Source/Runtime/MassEntity/Public/MassCommandBuffer.h
@@ -43,7 +43,7 @@ PRAGMA_DISABLE_DEPRECATION_WARNINGS
 	{
 	};
 PRAGMA_ENABLE_DEPRECATION_WARNINGS
-}
+} // namespace UE::Mass::Deprecation
 
 #define ASSERT_EXECUTE_DEPRECATION(CommandType) \
 		static_assert(!UE::Mass::Deprecation::TOverridesExecute<CommandType>::value, "Mass Commands: CONST Execute function is deprecated in 5.7 and will be removed by 5.9. Use Run instead.");
diff --git a/Engine/Source/Runtime/MassEntity/Public/MassCommands.h b/Engine/Source/Runtime/MassEntity/Public/MassCommands.h
--- a/Engine/Source/Runtime/MassEntity/Public/MassCommands.h
+++ b/Engine/Source/Runtime/MassEntity/Public/MassCommands.h
@@ -86,7 +86,7 @@ namespace UE::Mass::Command
 			RequiresUniqueHandling = false
 		};
 	};
-};
+}; // namespace UE::Mass::Command
 
 struct FMassBatchedCommand
 {
diff --git a/Engine/Source/Runtime/MassEntity/Public/MassDebugLogging.h b/Engine/Source/Runtime/MassEntity/Public/MassDebugLogging.h
--- a/Engine/Source/Runtime/MassEntity/Public/MassDebugLogging.h
+++ b/Engine/Source/Runtime/MassEntity/Public/MassDebugLogging.h
@@ -86,4 +86,4 @@ namespace UE::Mass::Debug
 		}
 #endif // WITH_MASSENTITY_DEBUG
 	};
-}
+} // namespace UE::Mass::Debug
diff --git a/Engine/Source/Runtime/MassEntity/Public/MassDebugger.h b/Engine/Source/Runtime/MassEntity/Public/MassDebugger.h
--- a/Engine/Source/Runtime/MassEntity/Public/MassDebugger.h
+++ b/Engine/Source/Runtime/MassEntity/Public/MassDebugger.h
@@ -463,7 +463,7 @@ struct FMassDebugger
 namespace UE::Mass::Debug
 {
 	MASSENTITY_API extern UWorld* GetDebugDrawWorld(const FMassExecutionContext& Context);
-}
+} // namespace UE::Mass::Debug
 
 #define MASS_IF_ENTITY_DEBUGGED(a, b) false
 #define MASS_BREAK_IF_ENTITY_DEBUGGED(a, b)
diff --git a/Engine/Source/Runtime/MassEntity/Public/MassDebuggerBreakpoints.h b/Engine/Source/Runtime/MassEntity/Public/MassDebuggerBreakpoints.h
--- a/Engine/Source/Runtime/MassEntity/Public/MassDebuggerBreakpoints.h
+++ b/Engine/Source/Runtime/MassEntity/Public/MassDebuggerBreakpoints.h
@@ -389,7 +389,7 @@ namespace UE::Mass::Debug
 		{
 		};
 	};
-}
+} // namespace UE::Mass::Debug
 
 #define MASS_BREAKPOINT(...)	\
 		do { } while (0)
diff --git a/Engine/Source/Runtime/MassEntity/Public/MassElement.h b/Engine/Source/Runtime/MassEntity/Public/MassElement.h
--- a/Engine/Source/Runtime/MassEntity/Public/MassElement.h
+++ b/Engine/Source/Runtime/MassEntity/Public/MassElement.h
@@ -218,7 +218,7 @@ namespace UE::Mass
 	DECLARE_NEWTYPEBITSET(FConstSharedFragmentBitSet, FMassConstSharedFragment, EElementType::ConstSharedFragment);
 
 	#undef DECLARE_NEWTYPEBITSET
-}
+} // namespace UE::Mass
 
 using FMassElementBitSet = UE::Mass::FElementBitSet;
 using FMassFragmentBitSet = UE::Mass::FFragmentBitSet;
diff --git a/Engine/Source/Runtime/MassEntity/Public/MassEntityCollection.h b/Engine/Source/Runtime/MassEntity/Public/MassEntityCollection.h
--- a/Engine/Source/Runtime/MassEntity/Public/MassEntityCollection.h
+++ b/Engine/Source/Runtime/MassEntity/Public/MassEntityCollection.h
@@ -199,4 +199,4 @@ namespace UE::Mass
 		 */
 		FMassArchetypeEntityCollection::EDuplicatesHandling CollectionCreationDuplicatesHandling = FMassArchetypeEntityCollection::NoDuplicates;
 	};
-}
+} // namespace UE::Mass
diff --git a/Engine/Source/Runtime/MassEntity/Public/MassEntityConcepts.h b/Engine/Source/Runtime/MassEntity/Public/MassEntityConcepts.h
--- a/Engine/Source/Runtime/MassEntity/Public/MassEntityConcepts.h
+++ b/Engine/Source/Runtime/MassEntity/Public/MassEntityConcepts.h
@@ -55,7 +55,7 @@ namespace UE::Mass
 				, std::conditional_t<CConstSharedFragment<T>, FMassConstSharedFragment
 				, void>>>>>;
 		};
-	}
+	} // namespace Private
 
 	template<typename T>
 	using TElementType = typename Private::TElementTypeHelper<T>::Type;
@@ -66,7 +66,7 @@ namespace UE::Mass
 		constexpr bool bIsSparse = CSparse<typename TRemoveReference<T>::Type>;
 		return bIsSparse;
 	}
-}
+} // namespace UE::Mass
 
 // Provides two targeted static_assert checks instead of the single combined CFragment concept check.
 // The first failing assert gives the user a clear, specific error message rather than the full
diff --git a/Engine/Source/Runtime/MassEntity/Public/MassEntityElementTypes.h b/Engine/Source/Runtime/MassEntity/Public/MassEntityElementTypes.h
--- a/Engine/Source/Runtime/MassEntity/Public/MassEntityElementTypes.h
+++ b/Engine/Source/Runtime/MassEntity/Public/MassEntityElementTypes.h
@@ -131,4 +131,4 @@ namespace UE::Mass
 	{
 		return Struct->IsChildOf(FMassSparseFragment::StaticStruct()) || Struct->IsChildOf(FMassSparseTag::StaticStruct());
 	}
-}
+} // namespace UE::Mass
diff --git a/Engine/Source/Runtime/MassEntity/Public/MassEntityManager.h b/Engine/Source/Runtime/MassEntity/Public/MassEntityManager.h
--- a/Engine/Source/Runtime/MassEntity/Public/MassEntityManager.h
+++ b/Engine/Source/Runtime/MassEntity/Public/MassEntityManager.h
@@ -49,18 +49,18 @@ namespace UE::Mass
 	namespace Private
 	{
 		struct FEntityStorageInitializer;
-	}
+	} // namespace Private
 	namespace ObserverManager
 	{
 		struct FObserverLock;
 		struct FCreationContext;
-	}
+	} // namespace ObserverManager
 
 	template<typename T>
 	concept CStructType = TIsDerivedFrom<Clean<T>, FInstancedStruct>::Value
 		|| TIsDerivedFrom<Clean<T>, FStructView>::Value
 		|| TIsDerivedFrom<Clean<T>, FConstStructView>::Value;
-}
+} // namespace UE::Mass
 
 // use REQUESTED_MASS_CONCURRENT_RESERVE=0 in your project's Build.cs file to disable concurrent storage
 // NOTE that it will always be enabled in WITH_EDITOR since editor code requires it. 
@@ -78,9 +78,9 @@ namespace UE::Mass
 #else
 	using FStorageType = FSingleThreadedEntityStorage;
 #endif //WITH_MASS_CONCURRENT_RESERVE
-}
+} // namespace UE::Mass
 
-/** 
+/**
  * The type responsible for hosting Entities managing Archetypes.
  * Entities are stored as FEntityData entries in a chunked array. 
  * Each valid entity is assigned to an Archetype that stored fragments associated with a given entity at the moment. 
diff --git a/Engine/Source/Runtime/MassEntity/Public/MassEntityRelations.h b/Engine/Source/Runtime/MassEntity/Public/MassEntityRelations.h
--- a/Engine/Source/Runtime/MassEntity/Public/MassEntityRelations.h
+++ b/Engine/Source/Runtime/MassEntity/Public/MassEntityRelations.h
@@ -16,7 +16,7 @@ class UMassObserverProcessor;
 namespace UE::Mass
 {
 	struct FRelationManager;
-}
+} // namespace UE::Mass
 
 namespace UE::Mass::Relations
 {
@@ -172,7 +172,7 @@ namespace UE::Mass::Relations
 		FString DebugInFix;
 #endif // WITH_MASSENTITY_DEBUG
 	};
-}
+} // namespace UE::Mass::Relations
 
 //-----------------------------------------------------------------------------
 // Relation types
@@ -338,6 +338,6 @@ namespace UE::Mass
 
 	template<typename T>
 	concept CRelation = TIsDerivedFrom<typename TRemoveReference<T>::Type, FMassRelation>::Value;
-};
+} // namespace UE::Mass
 
 #undef UE_API
diff --git a/Engine/Source/Runtime/MassEntity/Public/MassEntityTypes.h b/Engine/Source/Runtime/MassEntity/Public/MassEntityTypes.h
--- a/Engine/Source/Runtime/MassEntity/Public/MassEntityTypes.h
+++ b/Engine/Source/Runtime/MassEntity/Public/MassEntityTypes.h
@@ -68,7 +68,7 @@ namespace UE::Mass
 		No,
 		Yes
 	};
-}
+} // namespace UE::Mass
 
 /** The type summarily describing a composition of an entity or an archetype. It contains information on both the
  *  fragments and tags */
@@ -936,7 +936,7 @@ namespace UE::Mass
 		TArray<FType> FragmentInstances;
 	};
 
-} // UE::Mass
+} // namespace UE::Mass
 
 
 struct FMassArchetypeCreationParams
diff --git a/Engine/Source/Runtime/MassEntity/Public/MassExecutor.h b/Engine/Source/Runtime/MassEntity/Public/MassExecutor.h
--- a/Engine/Source/Runtime/MassEntity/Public/MassExecutor.h
+++ b/Engine/Source/Runtime/MassEntity/Public/MassExecutor.h
@@ -10,7 +10,7 @@ struct FMassRuntimePipeline;
 namespace UE::Mass
 {
 	struct FProcessingContext;
-}
+} // namespace UE::Mass
 struct FMassEntityHandle;
 struct FMassArchetypeEntityCollection;
 class UMassProcessor;
@@ -53,4 +53,4 @@ namespace UE::Mass::Executor
 	UE_DEPRECATED(5.6, "lvalue flavor of TriggerParallelTasks has been deprevate. Use the rvalue version.")
 	MASSENTITY_API FGraphEventRef TriggerParallelTasks(UMassProcessor& Processor, FProcessingContext& ProcessingContext, TFunction<void()> OnDoneNotification
 		, ENamedThreads::Type CurrentThread = ENamedThreads::GameThread);
-};
+} // namespace UE::Mass::Executor
diff --git a/Engine/Source/Runtime/MassEntity/Public/MassObserverManager.h b/Engine/Source/Runtime/MassEntity/Public/MassObserverManager.h
--- a/Engine/Source/Runtime/MassEntity/Public/MassObserverManager.h
+++ b/Engine/Source/Runtime/MassEntity/Public/MassObserverManager.h
@@ -33,10 +33,10 @@ namespace UE::Mass
 		struct FCreationNotificationHandle;
 		struct FObserverLock;
 		struct FCreationContext;
-	}
-}
+	} // namespace ObserverManager
+} // namespace UE::Mass
 
-/** 
+/**
  * A wrapper type for a TMap to support having array-of-maps UPROPERTY members in FMassObserverManager
  */
 USTRUCT()
diff --git a/Engine/Source/Runtime/MassEntity/Public/MassObserverProcessor.h b/Engine/Source/Runtime/MassEntity/Public/MassObserverProcessor.h
--- a/Engine/Source/Runtime/MassEntity/Public/MassObserverProcessor.h
+++ b/Engine/Source/Runtime/MassEntity/Public/MassObserverProcessor.h
@@ -11,7 +11,7 @@
 namespace UE::Mass::ObserverManager
 {
 	struct FObserverContextIterator;
-}
+} // namespace UE::Mass::ObserverManager
 
 /**
  * Instances of the type will be fed into FMassRuntimePipeline.AuxData and at execution time will
diff --git a/Engine/Source/Runtime/MassEntity/Public/MassProcessingTypes.h b/Engine/Source/Runtime/MassEntity/Public/MassProcessingTypes.h
--- a/Engine/Source/Runtime/MassEntity/Public/MassProcessingTypes.h
+++ b/Engine/Source/Runtime/MassEntity/Public/MassProcessingTypes.h
@@ -58,7 +58,7 @@ struct FMassProcessingContext_DEPRECATED
 namespace UE::Mass
 {
 	struct FProcessingContext;
-}
+} // namespace UE::Mass
 using FMassProcessingContext = UE::Mass::FProcessingContext;
 
 /** 
diff --git a/Engine/Source/Runtime/MassEntity/Public/MassProcessor.h b/Engine/Source/Runtime/MassEntity/Public/MassProcessor.h
--- a/Engine/Source/Runtime/MassEntity/Public/MassProcessor.h
+++ b/Engine/Source/Runtime/MassEntity/Public/MassProcessor.h
@@ -28,7 +28,7 @@ namespace UE::Mass
 	struct FQueryExecutor;
 	struct FProcessingQueue;
 	struct FQueuedProcessor;
-}
+} // namespace UE::Mass
 
 enum class EProcessorCompletionStatus : uint8
 {
diff --git a/Engine/Source/Runtime/MassEntity/Public/MassProcessorDependencySolver.h b/Engine/Source/Runtime/MassEntity/Public/MassProcessorDependencySolver.h
--- a/Engine/Source/Runtime/MassEntity/Public/MassProcessorDependencySolver.h
+++ b/Engine/Source/Runtime/MassEntity/Public/MassProcessorDependencySolver.h
@@ -14,7 +14,7 @@ class UMassProcessor;
 namespace UE::Mass
 {
 	struct FTypeManager;
-}
+} // namespace UE::Mass
 
 namespace EMassAccessOperation
 {
diff --git a/Engine/Source/Runtime/MassEntity/Public/MassRelationObservers.h b/Engine/Source/Runtime/MassEntity/Public/MassRelationObservers.h
--- a/Engine/Source/Runtime/MassEntity/Public/MassRelationObservers.h
+++ b/Engine/Source/Runtime/MassEntity/Public/MassRelationObservers.h
@@ -12,7 +12,7 @@ namespace UE::Mass::Relations
 {
 	struct FRelationTypeTraits;
 	enum class ERelationRole : uint8;
-}
+} // namespace UE::Mass::Relations
 
 UCLASS(MinimalAPI)
 class UMassRelationObserver : public UMassObserverProcessor
diff --git a/Engine/Source/Runtime/MassEntity/Public/MassRequirementAccessDetector.h b/Engine/Source/Runtime/MassEntity/Public/MassRequirementAccessDetector.h
--- a/Engine/Source/Runtime/MassEntity/Public/MassRequirementAccessDetector.h
+++ b/Engine/Source/Runtime/MassEntity/Public/MassRequirementAccessDetector.h
@@ -96,4 +96,4 @@ namespace UE::Mass::Debug
 		}
 #endif // WITH_MASSENTITY_DEBUG
 	};
-}
+} // namespace UE::Mass::Debug
diff --git a/Engine/Source/Runtime/MassEntity/Public/MassRequirements.h b/Engine/Source/Runtime/MassEntity/Public/MassRequirements.h
--- a/Engine/Source/Runtime/MassEntity/Public/MassRequirements.h
+++ b/Engine/Source/Runtime/MassEntity/Public/MassRequirements.h
@@ -34,7 +34,7 @@ private:
 	FOutputDevice* ActiveOutputDevice = nullptr;
 };
 
-} //UE::Mass::ErrorReporting
+} // namespace UE::Mass::ErrorReporting
 
 UENUM()
 enum class EMassFragmentAccess : uint8
diff --git a/Engine/Source/Runtime/MassEntity/Public/MassSparseElementsStorage.h b/Engine/Source/Runtime/MassEntity/Public/MassSparseElementsStorage.h
--- a/Engine/Source/Runtime/MassEntity/Public/MassSparseElementsStorage.h
+++ b/Engine/Source/Runtime/MassEntity/Public/MassSparseElementsStorage.h
@@ -476,6 +476,6 @@ inline FSparseElementIterator FSparseElementsStorage::CreateElementIterator(TNot
 	return TypePool ? FSparseElementIterator(*TypePool) : FSparseElementIterator();
 }
 
-}
+} // namespace UE::Mass
 
 #undef UE_API
diff --git a/Engine/Source/Runtime/MassEntity/Public/MassSubsystemBase.h b/Engine/Source/Runtime/MassEntity/Public/MassSubsystemBase.h
--- a/Engine/Source/Runtime/MassEntity/Public/MassSubsystemBase.h
+++ b/Engine/Source/Runtime/MassEntity/Public/MassSubsystemBase.h
@@ -25,10 +25,10 @@ namespace UE::Mass
 
 		MASSENTITY_API void RegisterSubsystemType(FSubsystemCollectionBase& Collection, TSubclassOf<USubsystem> SubsystemClass, FSubsystemTypeTraits&& Traits);
 		MASSENTITY_API void RegisterSubsystemType(TSharedRef<FMassEntityManager> EntityManager, TSubclassOf<USubsystem> SubsystemClass, FSubsystemTypeTraits&& Traits);
-	}
-}
+	} // namespace Subsystems
+} // namespace UE::Mass
 
-/** 
+/**
  * The sole responsibility of this world subsystem class is to serve functionality common to all 
  * Mass-related UWorldSubsystem-based subsystems, like whether the subsystems should get created at all. 
  */
diff --git a/Engine/Source/Runtime/MassEntity/Public/MassTypeManager.h b/Engine/Source/Runtime/MassEntity/Public/MassTypeManager.h
--- a/Engine/Source/Runtime/MassEntity/Public/MassTypeManager.h
+++ b/Engine/Source/Runtime/MassEntity/Public/MassTypeManager.h
@@ -308,6 +308,6 @@ namespace UE::Mass
 	{
 		return FTypeHandle(InTypeKey);
 	}
-}
+} // namespace UE::Mass
 
 #undef UE_API
diff --git a/Engine/Source/Runtime/MassEntity/Public/Relations/MassChildOf.h b/Engine/Source/Runtime/MassEntity/Public/Relations/MassChildOf.h
--- a/Engine/Source/Runtime/MassEntity/Public/Relations/MassChildOf.h
+++ b/Engine/Source/Runtime/MassEntity/Public/Relations/MassChildOf.h
@@ -13,7 +13,7 @@ namespace UE::Mass::Relations
 {
 	extern UE_API FTypeHandle ChildOfHandle;
 	UE_API void RegisterChildOfRelation();
-}
+} // namespace UE::Mass::Relations
 
 USTRUCT()
 struct FMassChildOfRelation : public FMassRelation
