Deprecate Mass SingleThreaded Storage implementation
* To be removed in 5.10

#Synced-CL 51731166

diff --git a/Engine/Source/Developer/MassEntityTestSuite/Private/MassEntityStorageTest.cpp b/Engine/Source/Developer/MassEntityTestSuite/Private/MassEntityStorageTest.cpp
--- a/Engine/Source/Developer/MassEntityTestSuite/Private/MassEntityStorageTest.cpp
+++ b/Engine/Source/Developer/MassEntityTestSuite/Private/MassEntityStorageTest.cpp
@@ -109,7 +109,9 @@ struct FEntityStorageTestBase : FAITestBase
 			FMassEntityManagerStorageInitParams InitializationParams;
 			if (Mode == EEntityManagerMode::Sequential)
 			{
+PRAGMA_DISABLE_DEPRECATION_WARNINGS
 				InitializationParams.Emplace<FMassEntityManager_InitParams_SingleThreaded>();
+PRAGMA_ENABLE_DEPRECATION_WARNINGS
 			}
 			else // Mode == EEntityManagerMode::Concurrent
 			{
diff --git a/Engine/Source/Developer/MassEntityTestSuite/Private/MassEntityTestFarmPlot.cpp b/Engine/Source/Developer/MassEntityTestSuite/Private/MassEntityTestFarmPlot.cpp
--- a/Engine/Source/Developer/MassEntityTestSuite/Private/MassEntityTestFarmPlot.cpp
+++ b/Engine/Source/Developer/MassEntityTestSuite/Private/MassEntityTestFarmPlot.cpp
@@ -231,7 +231,15 @@ void AMassEntityTestFarmPlot::BeginPlay()
 	}
 
 	FMassEntityManager& EntityManager = SharedEntityManager.Get();
+#if WITH_MASS_CONCURRENT_RESERVE
+	FMassEntityManagerStorageInitParams InitParams;
+	InitParams.Emplace<FMassEntityManager_InitParams_Concurrent>();
+	EntityManager.Initialize(FMassEntityManagerStorageInitParams(InitParams));
+#else
+PRAGMA_DISABLE_DEPRECATION_WARNINGS
 	EntityManager.Initialize();
+PRAGMA_ENABLE_DEPRECATION_WARNINGS
+#endif
 
 	FMassArchetypeHandle CropArchetype = EntityManager.CreateArchetype(TArray<const UScriptStruct*>{ FFarmWaterFragment::StaticStruct(), FFarmCropFragment::StaticStruct(), FHarvestTimerFragment::StaticStruct(), FFarmVisualFragment::StaticStruct(), FFarmGridCellData::StaticStruct() });
 	FMassArchetypeHandle FlowerArchetype = EntityManager.CreateArchetype(TArray<const UScriptStruct*>{ FFarmWaterFragment::StaticStruct(), FFarmFlowerFragment::StaticStruct(), FHarvestTimerFragment::StaticStruct(), FFarmVisualFragment::StaticStruct(), FFarmGridCellData::StaticStruct() });
diff --git a/Engine/Source/Developer/MassEntityTestSuite/Private/MassProcessingPhasesTest.cpp b/Engine/Source/Developer/MassEntityTestSuite/Private/MassProcessingPhasesTest.cpp
--- a/Engine/Source/Developer/MassEntityTestSuite/Private/MassProcessingPhasesTest.cpp
+++ b/Engine/Source/Developer/MassEntityTestSuite/Private/MassProcessingPhasesTest.cpp
@@ -31,7 +31,9 @@ bool FProcessingPhasesTestBase::SetUp()
 
 		InitializationParams.Emplace<FMassEntityManager_InitParams_Concurrent>();
 #else
+PRAGMA_DISABLE_DEPRECATION_WARNINGS
 		InitializationParams.Emplace<FMassEntityManager_InitParams_SingleThreaded>();
+PRAGMA_ENABLE_DEPRECATION_WARNINGS
 #endif // WITH_MASS_CONCURRENT_RESERVE
 
 		EntityManager->Initialize(InitializationParams);
diff --git a/Engine/Source/Developer/MassEntityTestSuite/Private/MassTestTypes.cpp b/Engine/Source/Developer/MassEntityTestSuite/Private/MassTestTypes.cpp
--- a/Engine/Source/Developer/MassEntityTestSuite/Private/MassTestTypes.cpp
+++ b/Engine/Source/Developer/MassEntityTestSuite/Private/MassTestTypes.cpp
@@ -24,7 +24,9 @@ bool FExecutionTestBase::SetUp()
 
 	InitializationParams.Emplace<FMassEntityManager_InitParams_Concurrent>();
 #else
+PRAGMA_DISABLE_DEPRECATION_WARNINGS
 	InitializationParams.Emplace<FMassEntityManager_InitParams_SingleThreaded>();
+PRAGMA_ENABLE_DEPRECATION_WARNINGS
 #endif // WITH_MASS_CONCURRENT_RESERVE
 
 	EntityManager->Initialize(InitializationParams);
diff --git a/Engine/Source/Runtime/MassEntity/Private/MassEntityManager.cpp b/Engine/Source/Runtime/MassEntity/Private/MassEntityManager.cpp
--- a/Engine/Source/Runtime/MassEntity/Private/MassEntityManager.cpp
+++ b/Engine/Source/Runtime/MassEntity/Private/MassEntityManager.cpp
@@ -151,7 +151,9 @@ void FMassEntityManager::AddReferencedObjects(FReferenceCollector& Collector)
 void FMassEntityManager::Initialize()
 {
 	FMassEntityManagerStorageInitParams InitializationParams;
+PRAGMA_DISABLE_DEPRECATION_WARNINGS
 	InitializationParams.Emplace<FMassEntityManager_InitParams_SingleThreaded>();
+PRAGMA_ENABLE_DEPRECATION_WARNINGS
 	Initialize(InitializationParams);
 }
 
@@ -159,11 +161,14 @@ namespace UE::Mass::Private
 {
 	struct FEntityStorageInitializer
 	{
+PRAGMA_DISABLE_DEPRECATION_WARNINGS
 		void operator()(const FMassEntityManager_InitParams_SingleThreaded& Params)
+
 		{
 			EntityStorage->Emplace<UE::Mass::FSingleThreadedEntityStorage>();
 			EntityStorage->Get<FSingleThreadedEntityStorage>().Initialize(Params);
 		}
+PRAGMA_ENABLE_DEPRECATION_WARNINGS
 		void operator()(const FMassEntityManager_InitParams_Concurrent& Params)
 		{
 #if WITH_MASS_CONCURRENT_RESERVE
diff --git a/Engine/Source/Runtime/MassEntity/Private/MassEntityManagerStorage.cpp b/Engine/Source/Runtime/MassEntity/Private/MassEntityManagerStorage.cpp
--- a/Engine/Source/Runtime/MassEntity/Private/MassEntityManagerStorage.cpp
+++ b/Engine/Source/Runtime/MassEntity/Private/MassEntityManagerStorage.cpp
@@ -31,7 +31,7 @@ namespace UE::Mass
 	//-----------------------------------------------------------------------------
 	// FSingleThreadedEntityStorage
 	//-----------------------------------------------------------------------------
-
+PRAGMA_DISABLE_DEPRECATION_WARNINGS
 	void FSingleThreadedEntityStorage::Initialize(const FMassEntityManager_InitParams_SingleThreaded&)
 	{
 		// Index 0 is reserved so we can treat that index as an invalid entity handle
@@ -239,6 +239,8 @@ namespace UE::Mass
 	{
 		return SerialNumber != 0 && CurrentArchetype.IsValid();
 	}
+	
+PRAGMA_ENABLE_DEPRECATION_WARNINGS
 
 	//-----------------------------------------------------------------------------
 	// FConcurrentEntityStorage
diff --git a/Engine/Source/Runtime/MassEntity/Private/MassEntitySubsystem.cpp b/Engine/Source/Runtime/MassEntity/Private/MassEntitySubsystem.cpp
--- a/Engine/Source/Runtime/MassEntity/Private/MassEntitySubsystem.cpp
+++ b/Engine/Source/Runtime/MassEntity/Private/MassEntitySubsystem.cpp
@@ -8,6 +8,8 @@
 
 namespace UE::Mass::Private
 {
+	// UE_DEPRECATED 5.8: bEnableMassConcurrentReserveRuntime and Mass.ConcurrentReserve.Enable are deprecated.
+	// Single threaded entity storage is being removed. Concurrent storage will always be used.
 	static bool bEnableMassConcurrentReserveRuntime = true;
 	static int32 ConcurrentReserveMaxEntityCount = 1 << 27;
 	static int32 ConcurrentReserveMaxEntitiesPerPage = 1 << 16;
@@ -18,7 +20,7 @@ namespace UE::Mass::Private
 			{
 				TEXT("Mass.ConcurrentReserve.Enable"),
 				bEnableMassConcurrentReserveRuntime,
-				TEXT("Enable Mass's concurrent reserve feature in runtime"),
+				TEXT("DEPRECATED 5.8: Single threaded entity storage is being removed. This CVar will be ignored in a future release."),
 				ECVF_Default
 			},
 			{
@@ -70,7 +72,9 @@ void UMassEntitySubsystem::Initialize(FSubsystemCollectionBase& Collection)
 	else
 #endif // WITH_MASS_CONCURRENT_RESERVE
 	{
+PRAGMA_DISABLE_DEPRECATION_WARNINGS
 		InitializationParams.Emplace<FMassEntityManager_InitParams_SingleThreaded>();
+PRAGMA_ENABLE_DEPRECATION_WARNINGS
 	}
 	
 	EntityManager->Initialize(InitializationParams);
diff --git a/Engine/Source/Runtime/MassEntity/Public/MassEntityManager.h b/Engine/Source/Runtime/MassEntity/Public/MassEntityManager.h
--- a/Engine/Source/Runtime/MassEntity/Public/MassEntityManager.h
+++ b/Engine/Source/Runtime/MassEntity/Public/MassEntityManager.h
@@ -62,9 +62,9 @@ namespace UE::Mass
 		|| TIsDerivedFrom<Clean<T>, FConstStructView>::Value;
 } // namespace UE::Mass
 
-// use REQUESTED_MASS_CONCURRENT_RESERVE=0 in your project's Build.cs file to disable concurrent storage
-// NOTE that it will always be enabled in WITH_EDITOR since editor code requires it. 
-// @see WITH_MASS_CONCURRENT_RESERVE
+// UE_DEPRECATED 5.8: REQUESTED_MASS_CONCURRENT_RESERVE and WITH_MASS_CONCURRENT_RESERVE are deprecated.
+// Single threaded entity storage is being removed. Concurrent storage will always be used.
+// Setting REQUESTED_MASS_CONCURRENT_RESERVE=0 is no longer supported and will be ignored in a future release.
 #ifndef REQUESTED_MASS_CONCURRENT_RESERVE
 #define REQUESTED_MASS_CONCURRENT_RESERVE 1
 #endif
@@ -138,6 +138,7 @@ public:
 	UE_API void GetResourceSizeEx(FResourceSizeEx& CumulativeResourceSize);
 
 	// Default to use single threaded implementation
+	UE_DEPRECATED(5.8, "Single threaded entity storage is being removed. Use Initialize(FMassEntityManagerStorageInitParams) with FMassEntityManager_InitParams_Concurrent instead.")
 	UE_API void Initialize();
 	UE_API void Initialize(const FMassEntityManagerStorageInitParams& InitializationParams);
 	UE_API void PostInitialize();
@@ -1352,7 +1353,9 @@ private:
 				checkf(false, TEXT("Attempt to use EntityStorageInterface without initialization"));
 				return nullptr;
 			}
+			PRAGMA_DISABLE_DEPRECATION_WARNINGS
 			UE::Mass::IEntityStorageInterface* operator()(FSingleThreadedEntityStorage& Storage) const
+			PRAGMA_ENABLE_DEPRECATION_WARNINGS
 			{
 				return &Storage;
 			}
@@ -1375,10 +1378,12 @@ private:
 private:
 
 	friend struct UE::Mass::Private::FEntityStorageInitializer;
+PRAGMA_DISABLE_DEPRECATION_WARNINGS
 	using FEntityStorageContainerType = TVariant<
 		FEmptyVariantState,
 		UE::Mass::FSingleThreadedEntityStorage,
 		UE::Mass::FConcurrentEntityStorage>;
+PRAGMA_ENABLE_DEPRECATION_WARNINGS
 	mutable FEntityStorageContainerType EntityStorage;
 #if WITH_MASSENTITY_DEBUG
 	UE::Mass::FStorageType* DebugEntityStoragePtr = nullptr;
diff --git a/Engine/Source/Runtime/MassEntity/Public/MassEntityManagerStorage.h b/Engine/Source/Runtime/MassEntity/Public/MassEntityManagerStorage.h
--- a/Engine/Source/Runtime/MassEntity/Public/MassEntityManagerStorage.h
+++ b/Engine/Source/Runtime/MassEntity/Public/MassEntityManagerStorage.h
@@ -14,7 +14,9 @@ struct FMassEntityHandle;
  * Initialization parameters to configure MassEntityManager to reserve entities only single threaded
  * Supported in all build configurations
  */
-struct FMassEntityManager_InitParams_SingleThreaded {};
+struct
+	UE_DEPRECATED(5.8, "Single threaded entity storage is being removed. Use FMassEntityManager_InitParams_Concurrent instead.")
+	FMassEntityManager_InitParams_SingleThreaded {};
 
 /**
  * Initialization parameters to configure MassEntityManager to concurrently reserve entities
@@ -41,7 +43,9 @@ struct FMassEntityManager_InitParams_Concurrent
 	uint32 MaxEntitiesPerPage = 1 << 16; // 65536
 };
 
+PRAGMA_DISABLE_DEPRECATION_WARNINGS 
 using FMassEntityManagerStorageInitParams = TVariant<FMassEntityManager_InitParams_SingleThreaded, FMassEntityManager_InitParams_Concurrent>;
+PRAGMA_ENABLE_DEPRECATION_WARNINGS
 
 namespace UE::Mass
 {
@@ -134,10 +138,14 @@ namespace UE::Mass
 	 * This storage backend should be used when the user of MassEntityManager can guarantee
 	 * that all entity management will be done on a single thread.
 	 */
-	class FSingleThreadedEntityStorage final : public IEntityStorageInterface
+	class
+		UE_DEPRECATED(5.8, "Single threaded entity storage is being removed. Use FConcurrentEntityStorage instead.")
+		FSingleThreadedEntityStorage final : public IEntityStorageInterface
 	{
 	public:
+PRAGMA_DISABLE_DEPRECATION_WARNINGS
 		void Initialize(const FMassEntityManager_InitParams_SingleThreaded&);
+PRAGMA_ENABLE_DEPRECATION_WARNINGS
 		virtual FMassArchetypeData* GetArchetype(int32 Index) override;
 		virtual const FMassArchetypeData* GetArchetype(int32 Index) const override;
 		virtual TSharedPtr<FMassArchetypeData>& GetArchetypeAsShared(int32 Index) override;
