[Mass] Fix FMassCommandAddFragmentInstances ignoring OperationType parameter

The constructor accepted an OperationType argument but the initializer
list hardcoded EMassCommandOperationType::Set instead of forwarding it.
This caused FMassCommandBuildEntity (which passes Create) to be sorted
as Set during command buffer flush, executing after AddTag/AddFragment
commands and hitting a checkf on entity state == Reserved.

#Synced-CL 51823590

diff --git a/Engine/Source/Runtime/MassEntity/Public/MassCommands.h b/Engine/Source/Runtime/MassEntity/Public/MassCommands.h
--- a/Engine/Source/Runtime/MassEntity/Public/MassCommands.h
+++ b/Engine/Source/Runtime/MassEntity/Public/MassCommands.h
@@ -405,7 +405,7 @@ struct FMassCommandAddFragmentInstances : public FMassBatchedEntityCommand
 	using Super = FMassBatchedEntityCommand;
 
 	FMassCommandAddFragmentInstances(EMassCommandOperationType OperationType = EMassCommandOperationType::Set DEBUG_NAME_PARAM("AddFragmentInstanceList"))
-		: Super(EMassCommandOperationType::Set FORWARD_DEBUG_NAME_PARAM)
+		: Super(OperationType FORWARD_DEBUG_NAME_PARAM)
 		, FragmentsAffected(UE::Mass::Utils::ConstructFragmentBitSet<EMassCommandCheckTime::CompileTimeCheck, TOthers...>())
 	{}
 
