[Mass] Propagate debug names from source archetype to newly created archetypes during shared fragment add and remove operations in MassEntityManager

#rnx
#Synced-CL 51846703

diff --git a/Engine/Source/Runtime/MassEntity/Private/MassEntityManager.cpp b/Engine/Source/Runtime/MassEntity/Private/MassEntityManager.cpp
--- a/Engine/Source/Runtime/MassEntity/Private/MassEntityManager.cpp
+++ b/Engine/Source/Runtime/MassEntity/Private/MassEntityManager.cpp
@@ -1662,6 +1662,8 @@ bool FMassEntityManager::AddConstSharedFragmentToEntity(const FMassEntityHandle
 	FMassArchetypeData* NewArchetype = NewArchetypeHandle.DataPtr.Get();
 	check(NewArchetype);
 
+	NewArchetype->CopyDebugNamesFrom(*CurrentArchetype);
+
 	const FMassArchetypeSharedFragmentValues& OldSharedFragmentValues = CurrentArchetype->GetSharedFragmentValues(Entity.Index);
 	check(!OldSharedFragmentValues.ContainsType(StructType));
 	FMassArchetypeSharedFragmentValues NewSharedFragmentValues(OldSharedFragmentValues);
@@ -1705,6 +1707,8 @@ bool FMassEntityManager::RemoveConstSharedFragmentFromEntity(const FMassEntityHa
 	check(NewArchetypeHandle.IsValid());
 	FMassArchetypeData* NewArchetype = NewArchetypeHandle.DataPtr.Get();
 	check(NewArchetype);
+
+	NewArchetype->CopyDebugNamesFrom(*CurrentArchetype);
 	
 	const FMassArchetypeSharedFragmentValues& OldSharedFragmentValues = CurrentArchetype->GetSharedFragmentValues(Entity.Index);
 	check(OldSharedFragmentValues.ContainsType(ConstSharedFragmentType));
@@ -1765,6 +1769,8 @@ bool FMassEntityManager::AddSharedFragmentToEntity(const FMassEntityHandle Entit
 	FMassArchetypeData* NewArchetype = NewArchetypeHandle.DataPtr.Get();
 	check(NewArchetype);
 
+	NewArchetype->CopyDebugNamesFrom(*CurrentArchetype);
+
 	const FMassArchetypeSharedFragmentValues& OldSharedFragmentValues = CurrentArchetype->GetSharedFragmentValues(Entity.Index);
 	check(!OldSharedFragmentValues.ContainsType(StructType));
 	FMassArchetypeSharedFragmentValues NewSharedFragmentValues(OldSharedFragmentValues);
@@ -1809,6 +1815,8 @@ bool FMassEntityManager::RemoveSharedFragmentFromEntity(const FMassEntityHandle
 	FMassArchetypeData* NewArchetype = NewArchetypeHandle.DataPtr.Get();
 	check(NewArchetype);
 
+	NewArchetype->CopyDebugNamesFrom(*CurrentArchetype);
+
 	const FMassArchetypeSharedFragmentValues& OldSharedFragmentValues = CurrentArchetype->GetSharedFragmentValues(Entity.Index);
 	check(OldSharedFragmentValues.ContainsType(SharedFragmentType));
 	FMassArchetypeSharedFragmentValues NewSharedFragmentValues(OldSharedFragmentValues);
