[MassDebugger] Fixed wrong hierarchy representation of unnamed archetypes in MassDebugger by adding a composition hash identifier for unnamed archetypes and filtering out the None entry from debug names
Added IsUnnamedArchetype helper to detect archetypes created with default NAME_None

#rnx
#Synced-CL 51846703

diff --git a/Engine/Source/Runtime/MassEntity/Internal/MassArchetypeData.h b/Engine/Source/Runtime/MassEntity/Internal/MassArchetypeData.h
--- a/Engine/Source/Runtime/MassEntity/Internal/MassArchetypeData.h
+++ b/Engine/Source/Runtime/MassEntity/Internal/MassArchetypeData.h
@@ -606,7 +606,7 @@ public:
 	void CopyDebugNamesFrom(const FMassArchetypeData& Other)
 	{ 
 #if WITH_MASSENTITY_DEBUG
-		DebugNames = Other.DebugNames; 
+		DebugNames = Other.DebugNames;
 #endif // WITH_MASSENTITY_DEBUG
 	}
 
@@ -615,11 +615,25 @@ public:
 	void DebugGetEntityMemoryNumbers(SIZE_T& OutActiveChunksMemorySize, SIZE_T& OutActiveEntitiesMemorySize) const;
 
 	/** Adds new debug name associated with the archetype. */
-	void AddUniqueDebugName(const FName& Name) { DebugNames.AddUnique(Name); }
-	
+	void AddUniqueDebugName(const FName& Name)
+	{
+		DebugNames.AddUnique(Name);
+	}
+
+	/** Whether the archetype was created using the default name (NAME_None) */
+	bool IsUnnamedArchetype() const
+	{
+		return DebugNames.Num() >= 1 && DebugNames[0].IsNone();
+	}
+
 	/** @return array of debug names associated with this archetype. */
-	const TConstArrayView<FName> GetDebugNames() const { return DebugNames; }
-	
+	TConstArrayView<FName> GetDebugNames() const
+	{
+		// For unnamed archetype we chop the "None" entry so the first name will be the unique identifier from the composition hash
+		const TConstArrayView<FName> AllNames = DebugNames;
+		return IsUnnamedArchetype() ? AllNames.RightChop(1) : AllNames;
+	}
+
 	/** @return string of all debug names combined */
 	FString GetCombinedDebugNamesAsString() const;
 
diff --git a/Engine/Source/Runtime/MassEntity/Private/MassArchetypeData.cpp b/Engine/Source/Runtime/MassEntity/Private/MassArchetypeData.cpp
--- a/Engine/Source/Runtime/MassEntity/Private/MassArchetypeData.cpp
+++ b/Engine/Source/Runtime/MassEntity/Private/MassArchetypeData.cpp
@@ -250,6 +250,11 @@ void FMassArchetypeData::Initialize(const FMassEntityManager& EntityManager, con
 	EntityListOffsetWithinChunk = 0;
 
 #if WITH_MASSENTITY_DEBUG
+	// For unnamed archetype (DebugNames[0] == NAME_None) we push an extra identifier to allow debug tools to differentiate them
+	if (IsUnnamedArchetype())
+	{
+		DebugNames.Add(FName(FString::Printf(TEXT("0x%X"), GetTypeHash(CompositionBitSet))));
+	}
 	SetDebugColor(DebugColor);
 #endif // WITH_MASSENTITY_DEBUG
 }
@@ -1304,7 +1309,7 @@ FString FMassArchetypeData::GetCombinedDebugNamesAsString() const
 	{
 		if (i > 0)
 		{
-			StringBuilder.Append(TEXT(", "));;
+			StringBuilder.Append(TEXT(", "));
 		}
 		StringBuilder.Append(DebugNames[i].ToString());
 	}
