[Mass] Use MoveAssignScriptStruct instead of CopyScriptStruct where source data is transient

Switches batch-spawn fragment initialization in FMassArchetypeData to use
MoveAssignScriptStruct, which leverages C++ move assignment for types that
opt in via WithMoveAssign (falling back to copy for types without the trait).

Adds FSparseElementsStorage::FTypePool::AddMove() for the same purpose in
sparse element storage.

#Synced-CL 52004206

diff --git a/Engine/Source/Runtime/MassEntity/Private/MassArchetypeData.cpp b/Engine/Source/Runtime/MassEntity/Private/MassArchetypeData.cpp
--- a/Engine/Source/Runtime/MassEntity/Private/MassArchetypeData.cpp
+++ b/Engine/Source/Runtime/MassEntity/Private/MassArchetypeData.cpp
@@ -1786,9 +1786,12 @@ void FMassArchetypeData::BatchSetFragmentValues(TConstArrayView<FMassArchetypeEn
 			{
 				const int32 FragmentIndex = FragmentIndexMap.FindChecked(FragmentType);
 				void* Dst = FragmentConfigs[FragmentIndex].GetFragmentData(Chunk.GetRawMemory(), EntityRange.SubchunkStart);
-				const void* Src = FragmentPayload.GetDataAt(EntitiesHandled);
+				void* Src = FragmentPayload.GetDataAt(EntitiesHandled);
 
-				FragmentType->CopyScriptStruct(Dst, Src, RangeLength);
+				// MoveAssignScriptStruct uses C++ move assignment for types that opt in via WithMoveAssign,
+				// leaving the source in a valid moved-from state. For types without the trait this falls
+				// back to CopyScriptStruct (zero regression).
+				FragmentType->MoveAssignScriptStruct(Dst, Src, RangeLength);
 			}
 			else
 			{
diff --git a/Engine/Source/Runtime/MassEntity/Private/MassSparseElementsStorage.cpp b/Engine/Source/Runtime/MassEntity/Private/MassSparseElementsStorage.cpp
--- a/Engine/Source/Runtime/MassEntity/Private/MassSparseElementsStorage.cpp
+++ b/Engine/Source/Runtime/MassEntity/Private/MassSparseElementsStorage.cpp
@@ -71,6 +71,15 @@ TNotNull<uint8*> FSparseElementsStorage::FTypePool::Add(const uint32 EntityIndex
 	return ElementPtr;
 }
 
+TNotNull<uint8*> FSparseElementsStorage::FTypePool::AddMove(const uint32 EntityIndex, TNotNull<uint8*> SourceData)
+{
+	TNotNull<uint8*> ElementPtr = Add(EntityIndex);
+	// MoveAssignScriptStruct uses C++ move assignment for types that opt in via WithMoveAssign,
+	// falling back to CopyScriptStruct for types without the trait.
+	Config.GetType()->MoveAssignScriptStruct(ElementPtr, SourceData);
+	return ElementPtr;
+}
+
 FStructView FSparseElementsStorage::FTypePool::Add_GetView(const uint32 EntityIndex)
 {	
 	return { Config.GetType(), Add(EntityIndex) };
diff --git a/Engine/Source/Runtime/MassEntity/Public/MassSparseElementsStorage.h b/Engine/Source/Runtime/MassEntity/Public/MassSparseElementsStorage.h
--- a/Engine/Source/Runtime/MassEntity/Public/MassSparseElementsStorage.h
+++ b/Engine/Source/Runtime/MassEntity/Public/MassSparseElementsStorage.h
@@ -223,6 +223,7 @@ private:
 
 		TNotNull<uint8*> Add(uint32 EntityIndex);
 		TNotNull<uint8*> Add(uint32 EntityIndex, TNotNull<const uint8*> SourceData);
+		TNotNull<uint8*> AddMove(uint32 EntityIndex, TNotNull<uint8*> SourceData);
 		FStructView Add_GetView(uint32 EntityIndex);
 		bool Remove(uint32 EntityIndex);
 		FConstStructView Get(uint32 EntityIndex) const;
