[MassEntity] Coding standard compliance pass, and spelling fixed

#rnx
#Synced-CL 52057779

diff --git a/Engine/Source/Runtime/MassEntity/Internal/MassArchetypeData.h b/Engine/Source/Runtime/MassEntity/Internal/MassArchetypeData.h
--- a/Engine/Source/Runtime/MassEntity/Internal/MassArchetypeData.h
+++ b/Engine/Source/Runtime/MassEntity/Internal/MassArchetypeData.h
@@ -171,7 +171,7 @@ public:
 	{
 		
 		LLM_SCOPE_BYNAME(TEXT("Mass/ArchetypeChunk"));
-		RawMemory = (uint8*)FMemory::Malloc(AllocSize);
+		RawMemory = static_cast<uint8*>(FMemory::Malloc(AllocSize));
 	}
 
 	~FMassArchetypeChunk()
@@ -287,7 +287,10 @@ public:
 		return SerialModificationNumber;
 	}
 
-	FStructView GetMutableChunkFragmentViewChecked(const int32 Index) { return FStructView(ChunkFragmentData[Index]); }
+	FStructView GetMutableChunkFragmentViewChecked(const int32 Index)
+	{
+		return FStructView(ChunkFragmentData[Index]);
+	}
 
 	FInstancedStruct* FindMutableChunkFragment(const UScriptStruct* Type)
 	{
@@ -307,7 +310,7 @@ public:
 		// If this chunk previously had entity and it does not anymore, we might have to reallocate the memory as it was freed to save memory
 		if (RawMemory == nullptr)
 		{
-			RawMemory = (uint8*)FMemory::Malloc(AllocSize);
+			RawMemory = static_cast<uint8*>(FMemory::Malloc(AllocSize));
 		}
 	}
 
@@ -317,11 +320,20 @@ public:
 	}
 
 #if WITH_MASSENTITY_DEBUG
-	int32 DebugGetChunkFragmentCount() const { return ChunkFragmentData.Num(); }
+	int32 DebugGetChunkFragmentCount() const
+	{
+		return ChunkFragmentData.Num();
+	}
 #endif // WITH_MASSENTITY_DEBUG
 
-	FMassArchetypeSharedFragmentValues& GetMutableSharedFragmentValues() { return SharedFragmentValues; }
-	const FMassArchetypeSharedFragmentValues& GetSharedFragmentValues() const { return SharedFragmentValues; }
+	FMassArchetypeSharedFragmentValues& GetMutableSharedFragmentValues()
+	{
+		return SharedFragmentValues;
+	}
+	const FMassArchetypeSharedFragmentValues& GetSharedFragmentValues() const
+	{
+		return SharedFragmentValues;
+	}
 };
 
 // Information for a single fragment type in an archetype
@@ -409,7 +421,10 @@ private:
 public:
 	explicit FMassArchetypeData(const FMassArchetypeCreationParams& CreationParams = FMassArchetypeCreationParams());
 
-	TConstArrayView<FMassArchetypeFragmentConfig> GetFragmentConfigs() const { return FragmentConfigs; }
+	TConstArrayView<FMassArchetypeFragmentConfig> GetFragmentConfigs() const
+	{
+		return FragmentConfigs;
+	}
 	const FMassElementBitSet& GetCompositionBitSet() const;
 	FMassArchetypeCompositionDescriptor GetCompositionDescriptor() const;
 
@@ -512,16 +527,37 @@ public:
 	bool ContainsSparseElement(TNotNull<const UScriptStruct*> ElementType) const;
 	const FMassElementBitSet& GetSparseElementsBitSet() const;
 
-	FORCEINLINE const int32* GetInternalIndexForEntity(const int32 EntityIndex) const { return EntityMap.Find(EntityIndex); }
-	FORCEINLINE int32 GetInternalIndexForEntityChecked(const int32 EntityIndex) const { return EntityMap.FindChecked(EntityIndex); }
-	int32 GetNumEntitiesPerChunk() const { return NumEntitiesPerChunk; }
-	SIZE_T GetBytesPerEntity() const { return TotalBytesPerEntity; }
+	FORCEINLINE const int32* GetInternalIndexForEntity(const int32 EntityIndex) const
+	{
+		return EntityMap.Find(EntityIndex);
+	}
+	FORCEINLINE int32 GetInternalIndexForEntityChecked(const int32 EntityIndex) const
+	{
+		return EntityMap.FindChecked(EntityIndex);
+	}
+	int32 GetNumEntitiesPerChunk() const
+	{
+		return NumEntitiesPerChunk;
+	}
+	SIZE_T GetBytesPerEntity() const
+	{
+		return TotalBytesPerEntity;
+	}
 
-	int32 GetNumEntities() const { return EntityMap.Num(); }
+	int32 GetNumEntities() const
+	{
+		return EntityMap.Num();
+	}
 
-	SIZE_T GetChunkAllocSize() const { return ChunkMemorySize; }
+	SIZE_T GetChunkAllocSize() const
+	{
+		return ChunkMemorySize;
+	}
 
-	int32 GetChunkCount() const { return Chunks.Num(); }
+	int32 GetChunkCount() const
+	{
+		return Chunks.Num();
+	}
 	int32 GetNonEmptyChunkCount() const;
 
 	FORCEINLINE static int32 CalculateRangeLength(FMassArchetypeEntityCollection::FArchetypeEntityRange EntityRange, const FMassArchetypeChunk& Chunk)
@@ -659,8 +695,14 @@ public:
 
 	//////////////////////////////////////////////////////////////////////
 	// low level api
-	FORCEINLINE const int32* GetFragmentIndex(const UScriptStruct* FragmentType) const { return FragmentIndexMap.Find(FragmentType); }
-	FORCEINLINE int32 GetFragmentIndexChecked(const UScriptStruct* FragmentType) const { return FragmentIndexMap.FindChecked(FragmentType); }
+	FORCEINLINE const int32* GetFragmentIndex(const UScriptStruct* FragmentType) const
+	{
+		return FragmentIndexMap.Find(FragmentType);
+	}
+	FORCEINLINE int32 GetFragmentIndexChecked(const UScriptStruct* FragmentType) const
+	{
+		return FragmentIndexMap.FindChecked(FragmentType);
+	}
 
 	FORCEINLINE void* GetFragmentData(const int32 FragmentIndex, const FMassRawEntityInChunkData RawEntityInChunkHandle) const
 	{
@@ -707,7 +749,10 @@ public:
 		return MakeEntityHandle(Entity.Index); 
 	}
 
-	bool IsInitialized() const { return TotalBytesPerEntity > 0 && FragmentConfigs.IsEmpty() == false; }
+	bool IsInitialized() const
+	{
+		return TotalBytesPerEntity > 0 && FragmentConfigs.IsEmpty() == false;
+	}
 
 	//////////////////////////////////////////////////////////////////////
 	// batched api
@@ -796,7 +841,10 @@ private:
 
 struct FMassArchetypeHelper
 {
-	FORCEINLINE static FMassArchetypeData* ArchetypeDataFromHandle(const FMassArchetypeHandle& ArchetypeHandle) { return ArchetypeHandle.DataPtr.Get(); }
+	FORCEINLINE static FMassArchetypeData* ArchetypeDataFromHandle(const FMassArchetypeHandle& ArchetypeHandle)
+	{
+		return ArchetypeHandle.DataPtr.Get();
+	}
 	FORCEINLINE static FMassArchetypeData& ArchetypeDataFromHandleChecked(const FMassArchetypeHandle& ArchetypeHandle)
 	{
 		check(ArchetypeHandle.IsValid());
@@ -910,7 +958,7 @@ void FMassArchetypeData::SetFragmentData(FMassArchetypeEntityCollection::FConstE
 	
 			for (FMassArchetypeChunkIterator ChunkIterator(EntityRangeContainer); ChunkIterator; ++ChunkIterator)
 			{
-				uint8* FragmentMemory = (uint8*)FragmentConfigs[*FragmentIndex].GetFragmentData(Chunks[ChunkIterator->ChunkIndex].GetRawMemory(), ChunkIterator->SubchunkStart);
+				uint8* FragmentMemory = static_cast<uint8*>(FragmentConfigs[*FragmentIndex].GetFragmentData(Chunks[ChunkIterator->ChunkIndex].GetRawMemory(), ChunkIterator->SubchunkStart));
 				// Note that we cannot handle the copying with a single FragmentType->CopyScriptStruct call
 				// because we have a single FragmentSource instance we're using as a copy source (i.e. we don't have a source array).
 				for (int i = ChunkIterator->Length; i; --i, FragmentMemory += FragmentTypeSize)
diff --git a/Engine/Source/Runtime/MassEntity/Private/MassCommandBuffer.cpp b/Engine/Source/Runtime/MassEntity/Private/MassCommandBuffer.cpp
--- a/Engine/Source/Runtime/MassEntity/Private/MassCommandBuffer.cpp
+++ b/Engine/Source/Runtime/MassEntity/Private/MassCommandBuffer.cpp
@@ -31,7 +31,7 @@ namespace UE::Mass::Command
 		TEXT("Set to true create a dedicated stat per type of command."), ECVF_Default);
 
 	/** CSV stat names */
-	static FString DefaultBatchedName = "BatchedCommand";
+	static FString DefaultBatchedName = TEXT("BatchedCommand");
 	static TMap<FName, TPair<FString, FAnsiString>> CommandBatchedFNames;
 
 	/** CSV custom stat names (ANSI) */
@@ -136,12 +136,19 @@ bool FMassCommandBuffer::Flush(FMassEntityManager& EntityManager)
 	{
 		FBatchedCommandsSortedIndex(const int32 InIndex, const int32 InGroupOrder)
 			: Index(InIndex), GroupOrder(InGroupOrder)
-		{}
+		{
+		}
 
 		const int32 Index = -1;
 		const int32 GroupOrder = MAX_int32;
-		bool IsValid() const { return GroupOrder < MAX_int32; }
-		bool operator<(const FBatchedCommandsSortedIndex& Other) const { return GroupOrder < Other.GroupOrder; }
+		bool IsValid() const
+		{
+			return GroupOrder < MAX_int32;
+		}
+		bool operator<(const FBatchedCommandsSortedIndex& Other) const
+		{
+			return GroupOrder < Other.GroupOrder;
+		}
 	};
 		
 	TArray<FBatchedCommandsSortedIndex> CommandsOrder;
@@ -151,12 +158,12 @@ bool FMassCommandBuffer::Flush(FMassEntityManager& EntityManager)
 	for (int32 i = 0; i < OwnedCommandsCount; ++i)
 	{
 		const TUniquePtr<FMassBatchedCommand>& Command = CommandInstances[i];
-		CommandsOrder.Add(FBatchedCommandsSortedIndex(i, (Command && Command->HasWork())? CommandTypeOrder[(int)Command->GetOperationType()] : MAX_int32));
+		CommandsOrder.Add(FBatchedCommandsSortedIndex(i, (Command && Command->HasWork())? CommandTypeOrder[static_cast<int32>(Command->GetOperationType())] : MAX_int32));
 	}
 	for (int32 i = 0; i < AppendedCommandInstances.Num(); ++i)
 	{
 		const TUniquePtr<FMassBatchedCommand>& Command = AppendedCommandInstances[i];
-		CommandsOrder.Add(FBatchedCommandsSortedIndex(i + OwnedCommandsCount, (Command && Command->HasWork()) ? CommandTypeOrder[(int)Command->GetOperationType()] : MAX_int32));
+		CommandsOrder.Add(FBatchedCommandsSortedIndex(i + OwnedCommandsCount, (Command && Command->HasWork()) ? CommandTypeOrder[static_cast<int32>(Command->GetOperationType())] : MAX_int32));
 	}
 	CommandsOrder.StableSort();
 
diff --git a/Engine/Source/Runtime/MassEntity/Private/MassEntityBuilder.cpp b/Engine/Source/Runtime/MassEntity/Private/MassEntityBuilder.cpp
--- a/Engine/Source/Runtime/MassEntity/Private/MassEntityBuilder.cpp
+++ b/Engine/Source/Runtime/MassEntity/Private/MassEntityBuilder.cpp
@@ -523,7 +523,7 @@ void FEntityBuilder::CacheEntityHandle() const
 {
 	if (EntityHandle.IsValid() == false)
 	{
-		checkf(State != EState::Committed, TEXT("Reserving an entity while the builder has already committed. This should not happen. Indicates an error during builder copying from another instance."))
+		checkf(State != EState::Committed, TEXT("Reserving an entity while the builder has already committed. This should not happen. Indicates an error during builder copying from another instance."));
 		EntityHandle = EntityManager->ReserveEntity();
 	}
 }
diff --git a/Engine/Source/Runtime/MassEntity/Private/MassEntityManagerConstants.h b/Engine/Source/Runtime/MassEntity/Private/MassEntityManagerConstants.h
--- a/Engine/Source/Runtime/MassEntity/Private/MassEntityManagerConstants.h
+++ b/Engine/Source/Runtime/MassEntity/Private/MassEntityManagerConstants.h
@@ -8,4 +8,4 @@ namespace UE::Mass::Private
 {
 	// Index 0 is a sentinel for an Empty/Unset EntityHandle
 	constexpr int32 InvalidEntityIndex = 0;
-};
\ No newline at end of file
+} // namespace UE::Mass::Private
\ No newline at end of file
diff --git a/Engine/Source/Runtime/MassEntity/Private/MassEntityManagerStorage.cpp b/Engine/Source/Runtime/MassEntity/Private/MassEntityManagerStorage.cpp
--- a/Engine/Source/Runtime/MassEntity/Private/MassEntityManagerStorage.cpp
+++ b/Engine/Source/Runtime/MassEntity/Private/MassEntityManagerStorage.cpp
@@ -374,8 +374,8 @@ PRAGMA_ENABLE_DEPRECATION_WARNINGS
 
 		// Size of page pointer array
 		const uint32 MaxEntities = 1 << MaximumEntityCountShift;
-		const uint32 MagPageCount = (MaxEntities / MaxEntitiesPerPage);
-		const SIZE_T PagePointerArraySizeBytes = MagPageCount * sizeof(FEntityData**);
+		const uint32 MaxPageCount = (MaxEntities / MaxEntitiesPerPage);
+		const SIZE_T PagePointerArraySizeBytes = MaxPageCount * sizeof(FEntityData**);
 	
 		return PageAllocatedSizeBytes + PagePointerArraySizeBytes + EntityFreeListSizeBytes;
 	}
@@ -617,7 +617,7 @@ PRAGMA_ENABLE_DEPRECATION_WARNINGS
 
 	int32 FConcurrentEntityStorage::Num() const
 	{
-		return MaxEntitiesPerPage * PageCount;;
+		return MaxEntitiesPerPage * PageCount;
 	}
 
 	int32 FConcurrentEntityStorage::ComputeFreeSize() const
diff --git a/Engine/Source/Runtime/MassEntity/Private/MassEntitySettings.cpp b/Engine/Source/Runtime/MassEntity/Private/MassEntitySettings.cpp
--- a/Engine/Source/Runtime/MassEntity/Private/MassEntitySettings.cpp
+++ b/Engine/Source/Runtime/MassEntity/Private/MassEntitySettings.cpp
@@ -23,9 +23,9 @@
 UMassEntitySettings::UMassEntitySettings(const FObjectInitializer& ObjectInitializer)
 	: Super(ObjectInitializer)
 {
-	for (int i = 0; i < (int)EMassProcessingPhase::MAX; ++i)
+	for (int32 i = 0; i < static_cast<int32>(EMassProcessingPhase::MAX); ++i)
 	{
-		ProcessingPhasesConfig[i].PhaseName = *UEnum::GetDisplayValueAsText(EMassProcessingPhase(i)).ToString();
+		ProcessingPhasesConfig[i].PhaseName = *UEnum::GetDisplayValueAsText(static_cast<EMassProcessingPhase>(i)).ToString();
 	}
 
 	FCoreDelegates::GetOnPostEngineInit().AddUObject(this, &UMassEntitySettings::OnPostEngineInit);
@@ -104,13 +104,13 @@ void UMassEntitySettings::BuildPhases()
 #if WITH_EDITOR
 	if (GIsEditor)
 	{
-		for (int i = 0; i < int(EMassProcessingPhase::MAX); ++i)
+		for (int32 i = 0; i < static_cast<int32>(EMassProcessingPhase::MAX); ++i)
 		{
 			FMassProcessingPhaseConfig& PhaseConfig = ProcessingPhasesConfig[i];
 			PhaseConfig.PhaseProcessor = NewObject<UMassCompositeProcessor>(this, PhaseConfig.PhaseGroupClass
 				, *FString::Printf(TEXT("ProcessingPhase_%s"), *PhaseConfig.PhaseName.ToString()));
 			PhaseConfig.PhaseProcessor->SetGroupName(PhaseConfig.PhaseName);
-			PhaseConfig.PhaseProcessor->SetProcessingPhase(EMassProcessingPhase(i));
+			PhaseConfig.PhaseProcessor->SetProcessingPhase(static_cast<EMassProcessingPhase>(i));
 			const FString PhaseDumpDependencyGraphFileName = !DumpDependencyGraphFileName.IsEmpty() ? DumpDependencyGraphFileName + TEXT("_") + PhaseConfig.PhaseName.ToString() : FString();
 
 			FMassRuntimePipeline TmpPipeline(EProcessorExecutionFlags::All);
@@ -133,19 +133,19 @@ void UMassEntitySettings::BuildProcessorList()
 		PhaseConfig.ProcessorCDOs.Reset();
 	}
 
-	TArray<UClass*> SubClassess;
-	GetDerivedClasses(UMassProcessor::StaticClass(), SubClassess);
+	TArray<UClass*> SubClasses;
+	GetDerivedClasses(UMassProcessor::StaticClass(), SubClasses);
 
-	for (int i = SubClassess.Num() - 1; i >= 0; --i)
+	for (int i = SubClasses.Num() - 1; i >= 0; --i)
 	{
-		if (SubClassess[i]->HasAnyClassFlags(CLASS_Abstract))
+		if (SubClasses[i]->HasAnyClassFlags(CLASS_Abstract))
 		{
 			continue;
 		}
 
-		UMassProcessor* ProcessorCDO = GetMutableDefault<UMassProcessor>(SubClassess[i]);
+		UMassProcessor* ProcessorCDO = GetMutableDefault<UMassProcessor>(SubClasses[i]);
 		// we explicitly restrict adding UMassCompositeProcessor. If needed by specific project a derived class can be added
-		if (ProcessorCDO && SubClassess[i] != UMassCompositeProcessor::StaticClass()
+		if (ProcessorCDO && SubClasses[i] != UMassCompositeProcessor::StaticClass()
 #if WITH_EDITOR
 			&& ProcessorCDO->ShouldShowUpInSettings()
 #endif // WITH_EDITOR
@@ -161,7 +161,7 @@ void UMassEntitySettings::BuildProcessorList()
 			ProcessorCDOs.Add(ProcessorCDO);
 			if (ProcessorCDO->ShouldAutoAddToGlobalList())
 			{
-				ProcessingPhasesConfig[int(ProcessorCDO->GetProcessingPhase())].ProcessorCDOs.Add(ProcessorCDO);
+				ProcessingPhasesConfig[static_cast<int32>(ProcessorCDO->GetProcessingPhase())].ProcessorCDOs.Add(ProcessorCDO);
 			}
 		}
 	}
@@ -199,7 +199,7 @@ void UMassEntitySettings::AddToActiveProcessorsList(TSubclassOf<UMassProcessor>
 TConstArrayView<FMassProcessingPhaseConfig> UMassEntitySettings::GetProcessingPhasesConfig()
 {
 	BuildProcessorListAndPhases();
-	return MakeArrayView(ProcessingPhasesConfig, int32(EMassProcessingPhase::MAX));
+	return MakeArrayView(ProcessingPhasesConfig, static_cast<int32>(EMassProcessingPhase::MAX));
 }
 
 #if WITH_EDITOR
diff --git a/Engine/Source/Runtime/MassEntity/Private/MassEntityTrace.cpp b/Engine/Source/Runtime/MassEntity/Private/MassEntityTrace.cpp
--- a/Engine/Source/Runtime/MassEntity/Private/MassEntityTrace.cpp
+++ b/Engine/Source/Runtime/MassEntity/Private/MassEntityTrace.cpp
@@ -115,7 +115,7 @@ UE_TRACE_EVENT_BEGIN(MassTrace, QueryForEachComplete)
 	UE_TRACE_EVENT_FIELD(int32, EntityCount)
 UE_TRACE_EVENT_END()
 
-enum class FFragmentType : uint8
+enum class EFragmentType : uint8
 {
 	Unknown = 0,
 	Fragment,
@@ -144,7 +144,7 @@ void FMassTrace::QueryRegisteredToProcessor(const FMassEntityQuery* Query, TNotN
 {
 	UE_TRACE_LOG(MassTrace, QueryRegisteredToProcessor, MassChannel)
 		<< QueryRegisteredToProcessor.QueryID(reinterpret_cast<uint64>(Query))
-		<< QueryRegisteredToProcessor.ProcessorID(reinterpret_cast<uint64>((const UMassProcessor*)(Processor)))
+		<< QueryRegisteredToProcessor.ProcessorID(reinterpret_cast<uint64>(static_cast<const UMassProcessor*>(Processor)))
 		<< QueryRegisteredToProcessor.ProcessorName(*Processor->GetProcessorName());
 }
 
@@ -197,23 +197,23 @@ void FMassTrace::FScopedQueryForEachTrace::ReportArchetype(const FMassArchetypeD
 void FMassTrace::OutputRegisterFragment(const UScriptStruct* Struct)
 {
 	const uint64 FragmentId = reinterpret_cast<uint64>(Struct);
-	const FFragmentType FragmentType = [](const UScriptStruct* Struct)
+	const EFragmentType FragmentType = [](const UScriptStruct* Struct)
 		{
 			if (Struct->IsChildOf<FMassFragment>())
 			{
-				return FFragmentType::Fragment;
+				return EFragmentType::Fragment;
 			}
 			else if (Struct->IsChildOf<FMassTag>())
 			{
-				return FFragmentType::Tag;
+				return EFragmentType::Tag;
 			}
 			else if (Struct->IsChildOf<FMassSharedFragment>())
 			{
-				return FFragmentType::SharedFragment;
+				return EFragmentType::SharedFragment;
 			}
 			else
 			{
-				return FFragmentType::Unknown;
+				return EFragmentType::Unknown;
 			}
 		}(Struct);
 
@@ -353,7 +353,7 @@ void FMassTrace::EntitiesDestroyed(TConstArrayView<FMassEntityHandle> Entities)
 	}
 }
 
-uint64_t FMassTrace::OutputBeginPhaseWithID(const TCHAR* PhaseName)
+uint64 FMassTrace::OutputBeginPhaseWithID(const TCHAR* PhaseName)
 {
 	if (UE_TRACE_CHANNELEXPR_IS_ENABLED(MassChannel))
 	{
diff --git a/Engine/Source/Runtime/MassEntity/Private/MassExecutionContext.cpp b/Engine/Source/Runtime/MassEntity/Private/MassExecutionContext.cpp
--- a/Engine/Source/Runtime/MassEntity/Private/MassExecutionContext.cpp
+++ b/Engine/Source/Runtime/MassEntity/Private/MassExecutionContext.cpp
@@ -260,7 +260,7 @@ FMassEntityHandle FMassExecutionContext::CreateEntity(const FMassArchetypeHandle
 #if WITH_MASSENTITY_DEBUG
 	// this check is a bit slow and very unlikely to fail for anyone other than the programmer implementing the processors calling this
 	const FMassArchetypeData& ArchetypeData = FMassArchetypeHelper::ArchetypeDataFromHandleChecked(ArchetypeHandle);
-	checkf(ActiveEntityCreationRequirements->Contains(ArchetypeData.GetCompositionBitSet()), TEXT("Declared FMassEntityCreationRequirements do not allow creating entites of the requested archetype"));
+	checkf(ActiveEntityCreationRequirements->Contains(ArchetypeData.GetCompositionBitSet()), TEXT("Declared FMassEntityCreationRequirements do not allow creating entities of the requested archetype"));
 #endif
 	FMassEntityHandle EntityHandle = EntityManager->CreateEntityAsyncUnsafe(ArchetypeHandle, SharedValues);
 
@@ -272,7 +272,7 @@ void FMassExecutionContext::BuildEntity(const FMassArchetypeHandle& ArchetypeHan
 	checkf(ActiveEntityCreationRequirements != nullptr, TEXT("Attempting to create an entity without declaring EntityCreationRequirements"));
 #if WITH_MASSENTITY_DEBUG
 	const FMassArchetypeData& ArchetypeData = FMassArchetypeHelper::ArchetypeDataFromHandleChecked(ArchetypeHandle);
-	checkf(ActiveEntityCreationRequirements->Contains(ArchetypeData.GetCompositionBitSet()), TEXT("Declared FMassEntityCreationRequirements do not allow creating entites of the requested archetype"));
+	checkf(ActiveEntityCreationRequirements->Contains(ArchetypeData.GetCompositionBitSet()), TEXT("Declared FMassEntityCreationRequirements do not allow creating entities of the requested archetype"));
 #endif
 	EntityManager->BuildEntityAsyncUnsafe(ArchetypeHandle, EntityHandle, SharedValues);
 }
@@ -283,7 +283,7 @@ void FMassExecutionContext::BatchCreateEntities(const FMassArchetypeHandle& Arch
 	checkf(ActiveEntityCreationRequirements != nullptr, TEXT("Attempting to create an entity without declaring FMassEntityCreationRequirements"));
 #if WITH_MASSENTITY_DEBUG
 	const FMassArchetypeData& ArchetypeData = FMassArchetypeHelper::ArchetypeDataFromHandleChecked(ArchetypeHandle);
-	checkf(ActiveEntityCreationRequirements->Contains(ArchetypeData.GetCompositionBitSet()), TEXT("Declared FMassEntityCreationRequirements do not allow creating entites of the requested archetype"));
+	checkf(ActiveEntityCreationRequirements->Contains(ArchetypeData.GetCompositionBitSet()), TEXT("Declared FMassEntityCreationRequirements do not allow creating entities of the requested archetype"));
 #endif
 
 	EntityManager->BatchCreateEntitiesAsyncUnsafe(ArchetypeHandle, SharedFragmentValues, Count, InOutEntities);
diff --git a/Engine/Source/Runtime/MassEntity/Private/MassExecutor.cpp b/Engine/Source/Runtime/MassEntity/Private/MassExecutor.cpp
--- a/Engine/Source/Runtime/MassEntity/Private/MassExecutor.cpp
+++ b/Engine/Source/Runtime/MassEntity/Private/MassExecutor.cpp
@@ -81,7 +81,7 @@ void RunProcessorsView(TArrayView<UMassProcessor* const> Processors, FProcessing
 #if WITH_MASSENTITY_DEBUG
 	if (Processors.Find(nullptr) != INDEX_NONE)
 	{
-		UE_LOG(LogMass, Error, TEXT("%s input Processors contains nullptr. Baling out."), ANSI_TO_TCHAR(__FUNCTION__));
+		UE_LOG(LogMass, Error, TEXT("%s input Processors contains nullptr. Bailing out."), ANSI_TO_TCHAR(__FUNCTION__));
 		return;
 	}
 #endif // WITH_MASSENTITY_DEBUG
@@ -127,7 +127,10 @@ struct FMassExecutorDoneTask
 		return DesiredThread;
 	}
 
-	static ESubsequentsMode::Type GetSubsequentsMode() { return ESubsequentsMode::TrackSubsequents; }
+	static ESubsequentsMode::Type GetSubsequentsMode()
+	{
+		return ESubsequentsMode::TrackSubsequents;
+	}
 	void DoTask(ENamedThreads::Type CurrentThread, const FGraphEventRef& MyCompletionGraphEvent)
 	{
 		TRACE_CPUPROFILER_EVENT_SCOPE_STR("Flush Deferred Commands Parallel");
@@ -177,8 +180,14 @@ struct FMassExecutorDispatchTask
 		RETURN_QUICK_DECLARE_CYCLE_STAT(FMassExecutorDispatchTask, STATGROUP_TaskGraphTasks);
 	}
 
-	static ENamedThreads::Type GetDesiredThread() { return ENamedThreads::AnyHiPriThreadHiPriTask; }
-	static ESubsequentsMode::Type GetSubsequentsMode() { return ESubsequentsMode::TrackSubsequents; }
+	static ENamedThreads::Type GetDesiredThread()
+	{
+		return ENamedThreads::AnyHiPriThreadHiPriTask;
+	}
+	static ESubsequentsMode::Type GetSubsequentsMode()
+	{
+		return ESubsequentsMode::TrackSubsequents;
+	}
 	void DoTask(ENamedThreads::Type /*CurrentThread*/, const FGraphEventRef& MyCompletionGraphEvent)
 	{
 		TRACE_CPUPROFILER_EVENT_SCOPE_STR("Dispatch Processors Async")
diff --git a/Engine/Source/Runtime/MassEntity/Private/MassObserverManager.cpp b/Engine/Source/Runtime/MassEntity/Private/MassObserverManager.cpp
--- a/Engine/Source/Runtime/MassEntity/Private/MassObserverManager.cpp
+++ b/Engine/Source/Runtime/MassEntity/Private/MassObserverManager.cpp
@@ -206,7 +206,7 @@ FMassObserverManager::FMassObserverManager(FMassEntityManager& Owner)
 void FMassObserverManager::DebugGatherUniqueProcessors(TArray<const UMassProcessor*>& OutProcessors) const
 {
 #if WITH_MASSENTITY_DEBUG
-	for (int32 OperationIndex = 0; OperationIndex < static_cast<uint32>(EMassObservedOperation::MAX); ++OperationIndex)
+	for (int32 OperationIndex = 0; OperationIndex < static_cast<int32>(EMassObservedOperation::MAX); ++OperationIndex)
 	{
 		ElementObservers[OperationIndex].DebugAddUniqueProcessors(OutProcessors);
 	}
@@ -225,7 +225,7 @@ void FMassObserverManager::Initialize()
 
 	TMap<const TSubclassOf<UMassProcessor>, UMassProcessor*> TransientProcessorsMap;
 	using UE::Mass::ObserverManager::Private::AddRegisteredObserverProcessorInstances;
-	for (int i = 0; i < (int)EMassObservedOperation::MAX; ++i)
+	for (int32 i = 0; i < static_cast<int32>(EMassObservedOperation::MAX); ++i)
 	{
 		AddRegisteredObserverProcessorInstances(EntityManager, WorldExecutionFlags, *Owner, Registry.ElementObserverMaps[i], ObservedElements[i], ElementObservers[i], TransientProcessorsMap);
 	}
@@ -256,7 +256,7 @@ void FMassObserverManager::DeInitialize()
 		ModulesUnloadedHandle.Reset();
 	}
 
-	for (int32 OperationIndex = 0; OperationIndex < (int32)EMassObservedOperation::MAX; ++OperationIndex)
+	for (int32 OperationIndex = 0; OperationIndex < static_cast<int32>(EMassObservedOperation::MAX); ++OperationIndex)
 	{
 		(*ElementObservers[OperationIndex]).Empty();
 		ObservedElements[OperationIndex].Reset();
@@ -678,7 +678,7 @@ void FMassObserverManager::OnModulePackagesUnloaded(TConstArrayView<UPackage*> P
 		}
 	};
 
-	for (int32 OperationIndex = 0; OperationIndex < static_cast<uint32>(EMassObservedOperation::MAX); ++OperationIndex)
+	for (int32 OperationIndex = 0; OperationIndex < static_cast<int32>(EMassObservedOperation::MAX); ++OperationIndex)
 	{
 		ProcessObserversMap(ElementObservers[OperationIndex]);
 	}
diff --git a/Engine/Source/Runtime/MassEntity/Private/MassProcessingContext.cpp b/Engine/Source/Runtime/MassEntity/Private/MassProcessingContext.cpp
--- a/Engine/Source/Runtime/MassEntity/Private/MassProcessingContext.cpp
+++ b/Engine/Source/Runtime/MassEntity/Private/MassProcessingContext.cpp
@@ -4,7 +4,7 @@
 
 namespace UE::Mass
 {
-	// @todo remove the depracation disabling once CommandBuffer and EntityManager are moved to `protected` and un-deprecated (around UE5.8)
+	// @todo remove the deprecation disabling once CommandBuffer and EntityManager are moved to `protected` and un-deprecated (around UE5.8)
 	PRAGMA_DISABLE_DEPRECATION_WARNINGS
 	FProcessingContext::~FProcessingContext()
 	{	
diff --git a/Engine/Source/Runtime/MassEntity/Private/MassProcessingPhaseManager.cpp b/Engine/Source/Runtime/MassEntity/Private/MassProcessingPhaseManager.cpp
--- a/Engine/Source/Runtime/MassEntity/Private/MassProcessingPhaseManager.cpp
+++ b/Engine/Source/Runtime/MassEntity/Private/MassProcessingPhaseManager.cpp
@@ -37,7 +37,7 @@ namespace UE::Mass::Tweakables
 
 namespace UE::Mass::Private
 {
-	ETickingGroup PhaseToTickingGroup[int(EMassProcessingPhase::MAX)]
+	ETickingGroup PhaseToTickingGroup[static_cast<int32>(EMassProcessingPhase::MAX)]
 	{
 		ETickingGroup::TG_PrePhysics, // EMassProcessingPhase::PrePhysics
 		ETickingGroup::TG_StartPhysics, // EMassProcessingPhase::StartPhysics
@@ -84,7 +84,7 @@ void FMassProcessingPhase::ExecuteTick(float DeltaTime, ELevelTick TickType, ENa
 
 	bIsDuringMassProcessing = true;	
 
-	if (UE::Mass::Tweakables::bUseProcessingQueue&& PhaseManager->IsPaused() == false)
+	if (UE::Mass::Tweakables::bUseProcessingQueue && PhaseManager->IsPaused() == false)
 	{
 		ObserversLockPtr = EntityManager.GetOrMakeObserversLock().ToSharedPtr();
 		FMassProcessingContext Context(EntityManager, DeltaTime);
@@ -299,9 +299,9 @@ void FMassProcessingPhaseManager::Initialize(UObject& InOwner, TConstArrayView<F
 	ProcessorExecutionFlags = UE::Mass::Utils::DetermineProcessorExecutionFlags(World, ProcessorExecutionFlags);
 	const uint8 SupportedTickTypes = UE::Mass::Utils::DetermineProcessorSupportedTickTypes(World);
 
-	for (int PhaseAsInt = 0; PhaseAsInt < int(EMassProcessingPhase::MAX); ++PhaseAsInt)
+	for (int32 PhaseAsInt = 0; PhaseAsInt < static_cast<int32>(EMassProcessingPhase::MAX); ++PhaseAsInt)
 	{		
-		const EMassProcessingPhase Phase = EMassProcessingPhase(PhaseAsInt);
+		const EMassProcessingPhase Phase = static_cast<EMassProcessingPhase>(PhaseAsInt);
 		FMassProcessingPhase& ProcessingPhase = ProcessingPhases[PhaseAsInt];
 
 		UMassCompositeProcessor* PhaseProcessor = NewObject<UMassCompositeProcessor>(&InOwner, UMassCompositeProcessor::StaticClass()
@@ -358,7 +358,7 @@ const FGraphEventRef& FMassProcessingPhaseManager::TriggerPhase(const EMassProce
 
 	if (bIsAllowedToTick)
 	{
-		ProcessingPhases[(int)Phase].ExecuteTick(DeltaTime, LEVELTICK_All, CurrentThread, MyCompletionGraphEvent);
+		ProcessingPhases[static_cast<int32>(Phase)].ExecuteTick(DeltaTime, LEVELTICK_All, CurrentThread, MyCompletionGraphEvent);
 	}
 
 	return MyCompletionGraphEvent;
@@ -536,7 +536,7 @@ void FMassProcessingPhaseManager::OnPhaseStart(FMassProcessingPhase& Phase)
 	ensure(CurrentPhase == EMassProcessingPhase::MAX);
 	CurrentPhase = Phase.Phase;
 
-	const int32 PhaseAsInt = int32(Phase.Phase);
+	const int32 PhaseAsInt = static_cast<int32>(Phase.Phase);
 
 	// The VERY FIRST thing we do in the first phase is to change the Pause state if needed.
 	// This way any code that depends on knowing the pause state (if any) gets consistent results.
@@ -649,7 +649,7 @@ void FMassProcessingPhaseManager::RegisterDynamicProcessor(UMassProcessor& Proce
 	if (ensureMsgf(Processor.GetProcessingPhase() != EMassProcessingPhase::MAX
 		, TEXT("%hs, Misconfigured processor %s, marked as ProcessingPhase == MAX"), __FUNCTION__, *Processor.GetName()))
 	{
-		PendingDynamicProcessors[int32(Processor.GetProcessingPhase())].Enqueue(&Processor, EDynamicProcessorOperationType::Add);
+		PendingDynamicProcessors[static_cast<int32>(Processor.GetProcessingPhase())].Enqueue(&Processor, EDynamicProcessorOperationType::Add);
 	}
 }
 
@@ -669,7 +669,7 @@ void FMassProcessingPhaseManager::UnregisterDynamicProcessor(UMassProcessor& Pro
 	if (ensureMsgf(Processor.GetProcessingPhase() != EMassProcessingPhase::MAX
 		, TEXT("%hs, Misconfigured processor %s, marked as ProcessingPhase == MAX"), __FUNCTION__, *Processor.GetName()))
 	{
-		PendingDynamicProcessors[int32(Processor.GetProcessingPhase())].Enqueue(&Processor, EDynamicProcessorOperationType::Remove);
+		PendingDynamicProcessors[static_cast<int32>(Processor.GetProcessingPhase())].Enqueue(&Processor, EDynamicProcessorOperationType::Remove);
 	}
 }
 
@@ -679,7 +679,7 @@ void FMassProcessingPhaseManager::UnregisterDynamicProcessorInternal(TNotNull<UM
 	if (DynamicProcessors.Find(Processor, Index))
 	{
 		DynamicProcessors.RemoveAtSwap(Index, 1, EAllowShrinking::No);
-		ProcessingGraphBuildStates[int32(Processor->GetProcessingPhase())].bProcessorsNeedRebuild = true;
+		ProcessingGraphBuildStates[static_cast<int32>(Processor->GetProcessingPhase())].bProcessorsNeedRebuild = true;
 
 		// it's possible that the given dynamic processor is a part of processing graph at the moment
 		// we need to store the information about its removal and use it when rebuilding the graph next time around.
diff --git a/Engine/Source/Runtime/MassEntity/Private/MassProcessingQueue.cpp b/Engine/Source/Runtime/MassEntity/Private/MassProcessingQueue.cpp
--- a/Engine/Source/Runtime/MassEntity/Private/MassProcessingQueue.cpp
+++ b/Engine/Source/Runtime/MassEntity/Private/MassProcessingQueue.cpp
@@ -135,7 +135,7 @@ namespace UE::Mass
 			}
 		}
 
-		// There is a very small chance that the entity manager bookeeping for the ProcessingScope count d
+		// There is a very small chance that the entity manager bookkeeping for the ProcessingScope count d
 		while (EntityManager.IsProcessing())
 		{
 			FPlatformProcess::YieldCycles(1);
@@ -188,16 +188,16 @@ namespace UE::Mass
 		return ActiveContentionTracker.TryActivate(RequestorIndex);
 	}
 
-	bool FProcessingQueue::ReleaseAndTryReserve(int32 IndexToRealease, int32 IndexToReserve)
+	bool FProcessingQueue::ReleaseAndTryReserve(int32 IndexToRelease, int32 IndexToReserve)
 	{
 		// unlock:
-		ActiveContentionTracker.Deactivate(IndexToRealease);
+		ActiveContentionTracker.Deactivate(IndexToRelease);
 		return TryReserve(IndexToReserve);
 	}
 
-	void FProcessingQueue::ReleaseAndReserve(int32 IndexToRealease, int32 IndexToReserve)
+	void FProcessingQueue::ReleaseAndReserve(int32 IndexToRelease, int32 IndexToReserve)
 	{
-		if (ReleaseAndTryReserve(IndexToRealease, IndexToReserve))
+		if (ReleaseAndTryReserve(IndexToRelease, IndexToReserve))
 		{
 			return;
 		}
@@ -416,7 +416,7 @@ namespace UE::Mass
 
 	void FProcessingQueue::Run(FMassExecutionContext& ExecutionContext)
 	{
-		FChangeListener ChaneListener(QueueChangeTracker);
+		FChangeListener ChangeListener(QueueChangeTracker);
 		for (;;)
 		{
 			int32 ProcessorsRemaining = PendingProcessorCount.load() - PendingMainThreadProcessorCount.load();
@@ -442,7 +442,7 @@ namespace UE::Mass
 			{
 				break;
 			}
-			ChaneListener.WaitForChange();
+			ChangeListener.WaitForChange();
 		}
 	}
 
diff --git a/Engine/Source/Runtime/MassEntity/Private/MassProcessor.cpp b/Engine/Source/Runtime/MassEntity/Private/MassProcessor.cpp
--- a/Engine/Source/Runtime/MassEntity/Private/MassProcessor.cpp
+++ b/Engine/Source/Runtime/MassEntity/Private/MassProcessor.cpp
@@ -62,7 +62,10 @@ public:
 		RETURN_QUICK_DECLARE_CYCLE_STAT(FMassProcessorTask, STATGROUP_TaskGraphTasks);
 	}
 
-	static ESubsequentsMode::Type GetSubsequentsMode() { return ESubsequentsMode::TrackSubsequents; }
+	static ESubsequentsMode::Type GetSubsequentsMode()
+	{
+		return ESubsequentsMode::TrackSubsequents;
+	}
 
 	static ENamedThreads::Type GetDesiredThread()
 	{
@@ -132,7 +135,7 @@ UMassProcessor::UMassProcessor(const FObjectInitializer& ObjectInitializer)
 }
 
 UMassProcessor::UMassProcessor()
-	: ExecutionFlags((int32)(EProcessorExecutionFlags::Server | EProcessorExecutionFlags::Standalone))
+	: ExecutionFlags(static_cast<int32>(EProcessorExecutionFlags::Server | EProcessorExecutionFlags::Standalone))
 {
 }
 
@@ -410,9 +413,9 @@ void UMassProcessor::ExportArchetypesCreated(TArray<FMassArchetypeHandle>& OutAr
 
 void UMassProcessor::RegisterQuery(FMassEntityQuery& Query)
 {
-	const uintptr_t ThisStart = (uintptr_t)this;
+	const uintptr_t ThisStart = reinterpret_cast<uintptr_t>(this);
 	const uintptr_t ThisEnd = ThisStart + GetClass()->GetStructureSize();
-	const uintptr_t QueryStart = (uintptr_t)&Query;
+	const uintptr_t QueryStart = reinterpret_cast<uintptr_t>(&Query);
 	const uintptr_t QueryEnd = QueryStart + sizeof(FMassEntityQuery);
 
 	if (QueryStart >= ThisStart && QueryEnd <= ThisEnd)
@@ -798,15 +801,6 @@ void UMassProcessor::Initialize(UObject& Owner)
 	}
 }
 
-//void UMassCompositeProcessor::Initialize(UObject& Owner)
-//{
-//	FMassEntityManager* EntityManager = UE::Mass::Utils::GetEntityManager(Owner.GetWorld());
-//	if (ensureMsgf(EntityManager, TEXT("Unable to determine the current MassEntityManager")))
-//	{
-//		InitializeInternal(Owner, EntityManager->AsShared());
-//	}
-//}
-
 void UMassCompositeProcessor::SetChildProcessors(TArray<UMassProcessor*>&& InProcessors)
 {
 	SetChildProcessors(MakeArrayView(InProcessors.GetData(), InProcessors.Num()));
diff --git a/Engine/Source/Runtime/MassEntity/Private/MassProcessorDependencySolver.cpp b/Engine/Source/Runtime/MassEntity/Private/MassProcessorDependencySolver.cpp
--- a/Engine/Source/Runtime/MassEntity/Private/MassProcessorDependencySolver.cpp
+++ b/Engine/Source/Runtime/MassEntity/Private/MassProcessorDependencySolver.cpp
@@ -64,7 +64,7 @@ namespace UE::Mass::Private
 			}
 		}
 
-		// Finding the cycle length by finding the other occurence of cycle-starting node
+		// Finding the cycle length by finding the other occurrence of cycle-starting node
 		const int32 CycleLength = MakeArrayView(&CycleIndices[CycleStartElementIndex + 1], CycleIndices.Num() - CycleStartElementIndex - 1).Find(CycleIndices[CycleStartElementIndex]) + 1;
 		check(CycleLength > 0);
 
diff --git a/Engine/Source/Runtime/MassEntity/Private/MassQueryExecutor.cpp b/Engine/Source/Runtime/MassEntity/Private/MassQueryExecutor.cpp
--- a/Engine/Source/Runtime/MassEntity/Private/MassQueryExecutor.cpp
+++ b/Engine/Source/Runtime/MassEntity/Private/MassQueryExecutor.cpp
@@ -38,15 +38,15 @@ void FQueryExecutor::ConfigureQuery(FMassSubsystemRequirements& ProcessorRequire
 #if WITH_MASSENTITY_DEBUG
 void FQueryExecutor::ValidateAccessors()
 {
-	const uintptr_t ExecutorStart = (uintptr_t)this;
-	const uintptr_t ExecutorEnd = ExecutorStart + DebugSize;
+	const UPTRINT ExecutorStart = reinterpret_cast<UPTRINT>(this);
+	const UPTRINT ExecutorEnd = ExecutorStart + DebugSize;
 
 	if (AccessorsPtr)
 	{
-		const uintptr_t AccessorsStart = (uintptr_t)this;
+		const UPTRINT AccessorsStart = reinterpret_cast<UPTRINT>(this);
 		checkf(ExecutorStart <= AccessorsStart && AccessorsStart <= ExecutorEnd, TEXT("Accessors assigned to a FQueryExecutor must be member variables of that struct."));
 	}
-};
+}
 #endif //WITH_MASSENTITY_DEBUG
 
 } // namespace UE::Mass
\ No newline at end of file
diff --git a/Engine/Source/Runtime/MassEntity/Private/MassRelationManager.cpp b/Engine/Source/Runtime/MassEntity/Private/MassRelationManager.cpp
--- a/Engine/Source/Runtime/MassEntity/Private/MassRelationManager.cpp
+++ b/Engine/Source/Runtime/MassEntity/Private/MassRelationManager.cpp
@@ -6,7 +6,6 @@
 #include "MassEntityManager.h"
 #include "MassRelationObservers.h"
 #include "MassExecutionContext.h"
-#include "Interfaces/ITextureFormat.h"
 #include "MassCommands.h"
 #include "MassArchetypeGroupCommands.h"
 #include "Algo/RemoveIf.h"
@@ -84,7 +83,7 @@ namespace UE::Mass
 		const int32 QueriedRoleIndex = static_cast<int32>(QueriedRole);
 
 		ensureMsgf(Traits.RoleTraits[QueriedRoleIndex].RequiresExternalMapping == EExternalMappingRequired::Yes
-			, TEXT("Fetching relation participants or role %s, while role traits explicitly prevent generic mapping")
+			, TEXT("Fetching relation participants for role %s, while role traits explicitly prevent generic mapping")
 			, *LexToString(QueriedRole));
 
 		TArray<FMassEntityHandle> ReturnSubjects;
diff --git a/Engine/Source/Runtime/MassEntity/Private/MassRelationObservers.cpp b/Engine/Source/Runtime/MassEntity/Private/MassRelationObservers.cpp
--- a/Engine/Source/Runtime/MassEntity/Private/MassRelationObservers.cpp
+++ b/Engine/Source/Runtime/MassEntity/Private/MassRelationObservers.cpp
@@ -84,7 +84,7 @@ UMassRelationEntityCreation::UMassRelationEntityCreation()
 {
 	ObservedOperations = EMassObservedOperationFlags::CreateEntity;
 	ExecutionPriority = RelationCreationObserverExecutionPriority;
-};
+}
 
 void UMassRelationEntityCreation::Execute(FMassEntityManager& EntityManager, FMassExecutionContext& Context)
 {
@@ -139,7 +139,7 @@ void UMassRelationEntityGuardDog::Execute(FMassEntityManager& EntityManager, FMa
 UMassRelationEntityDestruction::UMassRelationEntityDestruction()
 {
 	ObservedOperations = EMassObservedOperationFlags::DestroyEntity;
-};
+}
 
 void UMassRelationEntityDestruction::Execute(FMassEntityManager& EntityManager, FMassExecutionContext& Context)
 {
@@ -185,7 +185,7 @@ UMassRelationRoleDestruction::UMassRelationRoleDestruction()
 	ObservedOperations = EMassObservedOperationFlags::DestroyEntity;
 	bAutoAddRelationFragmentRequirement = false;
 	bAutoAddRelationTagRequirement = false;
-};
+}
 
 void UMassRelationRoleDestruction::Execute(FMassEntityManager& EntityManager, FMassExecutionContext& Context)
 {
@@ -331,7 +331,7 @@ bool UMassRelationRoleDestruction::ConfigureRelationObserver(UE::Mass::FTypeHand
 			bExecutionFunctionAssigned = true;
 			break;
 		case UE::Mass::ERemovalPolicy::Splice:
-			// use my Object as my Subject's Object, and vive versa
+			// use my Object as my Subject's Object, and vice versa
 			ExecuteFunction = [InRegisteredTypeHandle](FMassExecutionContext& ExecutionContext)
 			{
 				static auto DestroyRelationEntitiesAndGetRoleEntities = [](
@@ -366,7 +366,7 @@ bool UMassRelationRoleDestruction::ConfigureRelationObserver(UE::Mass::FTypeHand
 			};
 			bExecutionFunctionAssigned = true;
 			break;
-		case UE::Mass::ERemovalPolicy::Custom: // fall through on purse
+		case UE::Mass::ERemovalPolicy::Custom: // fall through on purpose
 		default: ;
 	}
 
@@ -375,9 +375,9 @@ bool UMassRelationRoleDestruction::ConfigureRelationObserver(UE::Mass::FTypeHand
 
 void UMassRelationRoleDestruction::AddObserverInstances(FMassObserverManager& ObserverManager, UE::Mass::FTypeHandle InRegisteredTypeHandle, const UE::Mass::FRelationTypeTraits& Traits)
 {
-	static_assert(static_cast<int32>(UE::Mass::ERelationRole::MAX) == 2, "This implementation is tailored the there being only two roles to any relation");
+	static_assert(static_cast<int32>(UE::Mass::ERelationRole::MAX) == 2, "This implementation is tailored to there being only two roles to any relation");
 
-	// is we both roles have the same policy and the same element, we can get away with creating only one observer
+	// if both roles have the same policy and the same element, we can get away with creating only one observer
 	// as long as we configure it appropriately
 	if (Traits.RoleTraits[1].Element == Traits.RoleTraits[0].Element && Traits.RoleTraits[1].DestructionPolicy == Traits.RoleTraits[0].DestructionPolicy)
 	{
diff --git a/Engine/Source/Runtime/MassEntity/Private/MassRequirementAccessDetector.cpp b/Engine/Source/Runtime/MassEntity/Private/MassRequirementAccessDetector.cpp
--- a/Engine/Source/Runtime/MassEntity/Private/MassRequirementAccessDetector.cpp
+++ b/Engine/Source/Runtime/MassEntity/Private/MassRequirementAccessDetector.cpp
@@ -39,10 +39,10 @@ void FMassRequirementAccessDetector::RequireAccess(const FMassEntityQuery& Query
 		Operation(Query.RequiredConstSubsystems, &FRWAccessDetector::AcquireReadAccess);
 		Operation(Query.RequiredMutableSubsystems, &FRWAccessDetector::AcquireWriteAccess);
 		
-		Aquire(Query.FragmentRequirements);
-		Aquire(Query.ChunkFragmentRequirements);
-		Aquire(Query.ConstSharedFragmentRequirements);
-		Aquire(Query.SharedFragmentRequirements);
+		Acquire(Query.FragmentRequirements);
+		Acquire(Query.ChunkFragmentRequirements);
+		Acquire(Query.ConstSharedFragmentRequirements);
+		Acquire(Query.SharedFragmentRequirements);
 	}
 }
 
diff --git a/Engine/Source/Runtime/MassEntity/Private/Relations/MassChildOf.cpp b/Engine/Source/Runtime/MassEntity/Private/Relations/MassChildOf.cpp
--- a/Engine/Source/Runtime/MassEntity/Private/Relations/MassChildOf.cpp
+++ b/Engine/Source/Runtime/MassEntity/Private/Relations/MassChildOf.cpp
@@ -39,7 +39,7 @@ namespace UE::Mass
 				ChildOfRelation.RoleTraits[static_cast<int32>(ERelationRole::Subject)].Element = FMassChildOfFragment::StaticStruct();
 
 				ChildOfRelation.bHierarchical = true;
-				ChildOfRelation.RelationName = "ChildOf";
+				ChildOfRelation.RelationName = TEXT("ChildOf");
 				
 				ChildOfRelation.RegisteredGroupType = InOutTypeManager.GetEntityManager().FindOrAddArchetypeGroupType(ChildOfRelation.RelationName);
 
@@ -59,7 +59,7 @@ namespace UE::Mass
 } // namespace UE::Mass
 
 //-----------------------------------------------------------------------------
-// UMassRelationEntityCreation
+// UMassChildOfRelationEntityCreation
 //-----------------------------------------------------------------------------
 void UMassChildOfRelationEntityCreation::Execute(FMassEntityManager& EntityManager, FMassExecutionContext& Context)
 {
diff --git a/Engine/Source/Runtime/MassEntity/Public/MassArchetypeGroupCommands.h b/Engine/Source/Runtime/MassEntity/Public/MassArchetypeGroupCommands.h
--- a/Engine/Source/Runtime/MassEntity/Public/MassArchetypeGroupCommands.h
+++ b/Engine/Source/Runtime/MassEntity/Public/MassArchetypeGroupCommands.h
@@ -50,4 +50,4 @@ namespace UE::Mass::Command
 		TMap<FArchetypeGroupHandle, TArray<FMassEntityHandle>> Groups;
 	};
 
-} // namespace UE::Mass
+} // namespace UE::Mass::Command
diff --git a/Engine/Source/Runtime/MassEntity/Public/MassCommandBuffer.h b/Engine/Source/Runtime/MassEntity/Public/MassCommandBuffer.h
--- a/Engine/Source/Runtime/MassEntity/Public/MassCommandBuffer.h
+++ b/Engine/Source/Runtime/MassEntity/Public/MassCommandBuffer.h
@@ -296,7 +296,10 @@ struct FMassCommandBuffer
 	{
 		return ActiveCommandsCounter > 0;
 	}
-	bool IsFlushing() const { return bIsFlushing; }
+	bool IsFlushing() const
+	{
+		return bIsFlushing;
+	}
 
 	/**
 	 * Removes any pending command instances
@@ -341,11 +344,11 @@ private:
 		}
 		else if (CommandInstances[Index])
 		{
-			return (T&)(*CommandInstances[Index].Get());
+			return static_cast<T&>(*CommandInstances[Index].Get());
 		}
 
 		CommandInstances[Index] = MakeUnique<T>();
-		return (T&)(*CommandInstances[Index].Get());
+		return static_cast<T&>(*CommandInstances[Index].Get());
 	}
 
 	/** 
diff --git a/Engine/Source/Runtime/MassEntity/Public/MassCommands.h b/Engine/Source/Runtime/MassEntity/Public/MassCommands.h
--- a/Engine/Source/Runtime/MassEntity/Public/MassCommands.h
+++ b/Engine/Source/Runtime/MassEntity/Public/MassCommands.h
@@ -86,7 +86,7 @@ namespace UE::Mass::Command
 			RequiresUniqueHandling = false
 		};
 	};
-}; // namespace UE::Mass::Command
+} // namespace UE::Mass::Command
 
 struct FMassBatchedCommand
 {
@@ -123,8 +123,14 @@ PRAGMA_ENABLE_DEPRECATION_WARNINGS
 		bHasWork = false;
 	}
 
-	bool HasWork() const { return bHasWork; }
-	EMassCommandOperationType GetOperationType() const { return OperationType; }
+	bool HasWork() const
+	{
+		return bHasWork;
+	}
+	EMassCommandOperationType GetOperationType() const
+	{
+		return OperationType;
+	}
 	
 	template<typename T>
 	UE_AUTORTFM_ALWAYS_OPEN
@@ -138,7 +144,10 @@ PRAGMA_ENABLE_DEPRECATION_WARNINGS
 
 #if CSV_PROFILER_STATS || WITH_MASSENTITY_DEBUG
 	virtual int32 GetNumOperationsStat() const = 0;
-	FName GetFName() const { return DebugName; }
+	FName GetFName() const
+	{
+		return DebugName;
+	}
 #endif // CSV_PROFILER_STATS || WITH_MASSENTITY_DEBUG
 
 protected:
@@ -196,7 +205,10 @@ protected:
 	}
 
 #if CSV_PROFILER_STATS || WITH_MASSENTITY_DEBUG
-	virtual int32 GetNumOperationsStat() const override { return TargetEntities.Num(); }
+	virtual int32 GetNumOperationsStat() const override
+	{
+		return TargetEntities.Num();
+	}
 #endif // CSV_PROFILER_STATS || WITH_MASSENTITY_DEBUG
 
 	UE_MT_DECLARE_TS_RW_ACCESS_DETECTOR(EntitiesAccessDetector); 
diff --git a/Engine/Source/Runtime/MassEntity/Public/MassDebugger.h b/Engine/Source/Runtime/MassEntity/Public/MassDebugger.h
--- a/Engine/Source/Runtime/MassEntity/Public/MassDebugger.h
+++ b/Engine/Source/Runtime/MassEntity/Public/MassDebugger.h
@@ -135,7 +135,7 @@ namespace UE::Mass::Debug
 
 	inline EMessageSeverity::Type MassSeverityToMessageSeverity(EMessageSeverity::Type OriginalSeverity, EMassDebugMessageSeverity MassSeverity)
 	{
-		static constexpr EMessageSeverity::Type ConversionMap[int(EMassDebugMessageSeverity::MAX)] =
+		static constexpr EMessageSeverity::Type ConversionMap[static_cast<int32>(EMassDebugMessageSeverity::MAX)] =
 		{
 			/*EMassDebugMessageSeverity::Error=*/EMessageSeverity::Error,
 			/*EMassDebugMessageSeverity::Warning=*/EMessageSeverity::Warning,
@@ -143,7 +143,7 @@ namespace UE::Mass::Debug
 		};
 		return MassSeverity == EMassDebugMessageSeverity::Default 
 			? OriginalSeverity
-			: ConversionMap[int(MassSeverity)];
+			: ConversionMap[static_cast<int32>(MassSeverity)];
 	}
 } // namespace UE::Mass::Debug
 
@@ -174,7 +174,10 @@ struct FMassDebugger
 		explicit FEnvironment(FMassEntityManager& InEntityManager);
 		~FEnvironment();
 
-		bool IsValid() const { return EntityManager.IsValid(); }
+		bool IsValid() const
+		{
+			return EntityManager.IsValid();
+		}
 
 		void ClearBreakpoints();
 	};
@@ -209,7 +212,7 @@ struct FMassDebugger
 	static MASSENTITY_API const FMassArchetypeData* GetArchetypeData(const FMassArchetypeHandle& ArchetypeHandle);
 	static MASSENTITY_API void EnumerateChunks(const FMassArchetypeData& Archetype, TFunctionRef<void(const FMassArchetypeChunk&)> Fn);
 
-	static MASSENTITY_API void GetArchetypeEntityStats(const FMassArchetypeHandle & ArchetypeHandle, UE::Mass::Debug::FArchetypeStats & OutStats);
+	static MASSENTITY_API void GetArchetypeEntityStats(const FMassArchetypeHandle& ArchetypeHandle, UE::Mass::Debug::FArchetypeStats& OutStats);
 	static MASSENTITY_API const TConstArrayView<FName> GetArchetypeDebugNames(const FMassArchetypeHandle& ArchetypeHandle);
 	static MASSENTITY_API TArray<FMassEntityHandle> GetEntitiesOfArchetype(const FMassArchetypeHandle& ArchetypeHandle);
 
@@ -278,7 +281,10 @@ struct FMassDebugger
 	 */
 	static MASSENTITY_API void RegisterProcessorDataProvider(FName ProviderName, const TSharedRef<FMassEntityManager>& EntityManager, const UE::Mass::Debug::FProcessorProviderFunction& ProviderFunction);
 
-	static TConstArrayView<FEnvironment> GetEnvironments() { return ActiveEnvironments; }
+	static TConstArrayView<FEnvironment> GetEnvironments()
+	{
+		return ActiveEnvironments;
+	}
 	static MASSENTITY_API FEnvironment* FindEnvironmentForEntityManager(const FMassEntityManager& EntityManager);
 
 	/**
@@ -454,10 +460,22 @@ struct FMassExecutionContext;
 
 struct FMassDebugger
 {
-	static FString GetSingleRequirementDescription(const FMassFragmentRequirementDescription&) { return TEXT("[no debug information]"); }
-	static FString GetRequirementsDescription(const FMassFragmentRequirements&) { return TEXT("[no debug information]"); }
-	static FString GetArchetypeRequirementCompatibilityDescription(const FMassFragmentRequirements&, const FMassArchetypeHandle&) { return TEXT("[no debug information]"); }
-	static FString GetArchetypeRequirementCompatibilityDescription(const FMassFragmentRequirements&, const FMassArchetypeCompositionDescriptor&) { return TEXT("[no debug information]"); }
+	static FString GetSingleRequirementDescription(const FMassFragmentRequirementDescription&)
+	{
+		return TEXT("[no debug information]");
+	}
+	static FString GetRequirementsDescription(const FMassFragmentRequirements&)
+	{
+		return TEXT("[no debug information]");
+	}
+	static FString GetArchetypeRequirementCompatibilityDescription(const FMassFragmentRequirements&, const FMassArchetypeHandle&)
+	{
+		return TEXT("[no debug information]");
+	}
+	static FString GetArchetypeRequirementCompatibilityDescription(const FMassFragmentRequirements&, const FMassArchetypeCompositionDescriptor&)
+	{
+		return TEXT("[no debug information]");
+	}
 };
 
 namespace UE::Mass::Debug
diff --git a/Engine/Source/Runtime/MassEntity/Public/MassElement.h b/Engine/Source/Runtime/MassEntity/Public/MassElement.h
--- a/Engine/Source/Runtime/MassEntity/Public/MassElement.h
+++ b/Engine/Source/Runtime/MassEntity/Public/MassElement.h
@@ -38,7 +38,7 @@ namespace UE::Mass
 		using FElementBitSetBase::FElementBitSetBase;
 		using FElementBitSetBase::operator=;
 		using FElementBitSetBase::HasAny;
-		using FElementBitSetBase::HasAll;;
+		using FElementBitSetBase::HasAll;
 
 		FElementBitSet(const FElementBitSetBase& OtherBitSet)
 			: FElementBitSetBase(OtherBitSet)
diff --git a/Engine/Source/Runtime/MassEntity/Public/MassEntityHandle.h b/Engine/Source/Runtime/MassEntity/Public/MassEntityHandle.h
--- a/Engine/Source/Runtime/MassEntity/Public/MassEntityHandle.h
+++ b/Engine/Source/Runtime/MassEntity/Public/MassEntityHandle.h
@@ -36,7 +36,10 @@ struct alignas(8) FMassEntityHandle
 	}
 
 	/** Has meaning only for sorting purposes */
-	bool operator<(const FMassEntityHandle Other) const { return Index < Other.Index; }
+	bool operator<(const FMassEntityHandle Other) const
+	{
+		return Index < Other.Index;
+	}
 
 	/** Note that this function is merely checking if Index and SerialNumber are set. There's no way to validate if 
 	 *  these indicate a valid entity in an EntitySubsystem without asking the system. */
@@ -56,12 +59,16 @@ struct alignas(8) FMassEntityHandle
 	}
 
 	/** Allows the entity handle to be shared anonymously. */
-	uint64 AsNumber() const { return *reinterpret_cast<const uint64*>(this); } // Relying on the fact that this struct only stores 2 integers and is aligned correctly.
+	// Relying on the fact that this struct only stores 2 integers and is aligned correctly.
+	uint64 AsNumber() const
+	{
+		return *reinterpret_cast<const uint64*>(this);
+	}
 	/** Reconstruct the entity handle from an anonymously shared integer. */
 	static FMassEntityHandle FromNumber(uint64 Value) 
 	{ 
 		FMassEntityHandle Result;
-		*reinterpret_cast<uint64_t*>(&Result) = Value;
+		*reinterpret_cast<uint64*>(&Result) = Value;
 		return Result;
 	}
 
diff --git a/Engine/Source/Runtime/MassEntity/Public/MassEntityRelations.h b/Engine/Source/Runtime/MassEntity/Public/MassEntityRelations.h
--- a/Engine/Source/Runtime/MassEntity/Public/MassEntityRelations.h
+++ b/Engine/Source/Runtime/MassEntity/Public/MassEntityRelations.h
@@ -252,8 +252,8 @@ struct alignas(8) FMassRelationRoleInstanceHandle
 		ReturnHandle.SetRoleEntityIndex(RoleHandle.Index);
 		ReturnHandle.SetRelationEntityIndex(RelationEntityHandle.Index);
 		ReturnHandle.SetRole(Role);
-		check(ReturnHandle.GetRoleEntityIndex() == RoleHandle.Index)
-		check(ReturnHandle.GetRelationEntityIndex() == RelationEntityHandle.Index)
+		check(ReturnHandle.GetRoleEntityIndex() == RoleHandle.Index);
+		check(ReturnHandle.GetRelationEntityIndex() == RelationEntityHandle.Index);
 
 		return ReturnHandle;
 	}
diff --git a/Engine/Source/Runtime/MassEntity/Public/MassEntitySettings.h b/Engine/Source/Runtime/MassEntity/Public/MassEntitySettings.h
--- a/Engine/Source/Runtime/MassEntity/Public/MassEntitySettings.h
+++ b/Engine/Source/Runtime/MassEntity/Public/MassEntitySettings.h
@@ -34,9 +34,16 @@ public:
 	MASSENTITY_API void AddToActiveProcessorsList(TSubclassOf<UMassProcessor> ProcessorClass);
 
 	MASSENTITY_API TConstArrayView<FMassProcessingPhaseConfig> GetProcessingPhasesConfig();
-	const FMassProcessingPhaseConfig& GetProcessingPhaseConfig(const EMassProcessingPhase ProcessingPhase) const { check(ProcessingPhase != EMassProcessingPhase::MAX); return ProcessingPhasesConfig[int(ProcessingPhase)]; }
+	const FMassProcessingPhaseConfig& GetProcessingPhaseConfig(const EMassProcessingPhase ProcessingPhase) const
+	{
+		check(ProcessingPhase != EMassProcessingPhase::MAX);
+		return ProcessingPhasesConfig[static_cast<int32>(ProcessingPhase)];
+	}
 
-	FOnInitialized& GetOnInitialized() { return OnInitializedEvent; }
+	FOnInitialized& GetOnInitialized()
+	{
+		return OnInitializedEvent;
+	}
 
 	UE_DEPRECATED(5.8, "Use the instance method GetOnInitialized")
 	static FOnInitialized& GetOnInitializedEvent()
@@ -46,9 +53,15 @@ public:
 	}
 
 #if WITH_EDITOR
-	FOnSettingsChange& GetOnSettingsChange() { return OnSettingsChange; }
+	FOnSettingsChange& GetOnSettingsChange()
+	{
+		return OnSettingsChange;
+	}
 
-	bool IsInitialized2() const { return bInitialized; }
+	bool IsInitialized2() const
+	{
+		return bInitialized;
+	}
 
 	UE_DEPRECATED(5.8, "Use the instance method IsInitialized2")
 	static bool IsInitialized()
@@ -88,7 +101,7 @@ public:
 
 	/** Lets users configure processing phases including the composite processor class to be used as a container for the phases' processors. */
 	UPROPERTY(EditDefaultsOnly, Category = Mass, config)
-	FMassProcessingPhaseConfig ProcessingPhasesConfig[(uint8)EMassProcessingPhase::MAX];
+	FMassProcessingPhaseConfig ProcessingPhasesConfig[static_cast<uint8>(EMassProcessingPhase::MAX)];
 
 	/** This list contains all the processors available in the given binary (including plugins). The contents are sorted by display name.*/
 	UPROPERTY(VisibleAnywhere, Category = Mass, Transient, Instanced, EditFixedSize)
diff --git a/Engine/Source/Runtime/MassEntity/Public/MassEntityTrace.h b/Engine/Source/Runtime/MassEntity/Public/MassEntityTrace.h
--- a/Engine/Source/Runtime/MassEntity/Public/MassEntityTrace.h
+++ b/Engine/Source/Runtime/MassEntity/Public/MassEntityTrace.h
@@ -38,7 +38,7 @@ public:
 		int32 EntityCount = 0;
 	};
 
-	MASSENTITY_API static uint64_t OutputBeginPhaseWithID(const TCHAR* PhaseName);
+	MASSENTITY_API static uint64 OutputBeginPhaseWithID(const TCHAR* PhaseName);
 	MASSENTITY_API static void OutputBeginPhaseRegion(const TCHAR* PhaseName);
 	MASSENTITY_API static void OutputEndPhaseRegion(const TCHAR* PhaseName);
 	MASSENTITY_API static void OutputEndPhaseRegion(uint64 PhaseId);
diff --git a/Engine/Source/Runtime/MassEntity/Public/MassEntityTypes.h b/Engine/Source/Runtime/MassEntity/Public/MassEntityTypes.h
--- a/Engine/Source/Runtime/MassEntity/Public/MassEntityTypes.h
+++ b/Engine/Source/Runtime/MassEntity/Public/MassEntityTypes.h
@@ -127,7 +127,7 @@ struct FMassArchetypeCompositionDescriptor
 	{}
 
 	/**
-	 * Constructor deliberately not marked as eplicit to support near-future change where we get rid of FMassArchetypeCompositionDescriptor
+	 * Constructor deliberately not marked as explicit to support near-future change where we get rid of FMassArchetypeCompositionDescriptor
 	 * since it's just a wrapper for FMassElementBitSet now.
 	 */
 	FMassArchetypeCompositionDescriptor(const FMassElementBitSet& InElements)
@@ -234,30 +234,30 @@ struct FMassArchetypeCompositionDescriptor
 	FMassSharedFragmentBitSet DebugGetSharedFragments() const;
 	FMassConstSharedFragmentBitSet DebugGetConstSharedFragments() const;
 
-	UE_DEPRECATED(5.8, "Getting a subset of compossition's bitsets is slow now. Consider refactoring your code to work directly on FMassElementBitSet")
+	UE_DEPRECATED(5.8, "Getting a subset of composition's bitsets is slow now. Consider refactoring your code to work directly on FMassElementBitSet")
 	FMassFragmentBitSet GetFragments() const;
-	UE_DEPRECATED(5.8, "Getting a subset of compossition's bitsets is slow now. Consider refactoring your code to work directly on FMassElementBitSet")
+	UE_DEPRECATED(5.8, "Getting a subset of composition's bitsets is slow now. Consider refactoring your code to work directly on FMassElementBitSet")
 	FMassTagBitSet GetTags() const;
-	UE_DEPRECATED(5.8, "Getting a subset of compossition's bitsets is slow now. Consider refactoring your code to work directly on FMassElementBitSet")
+	UE_DEPRECATED(5.8, "Getting a subset of composition's bitsets is slow now. Consider refactoring your code to work directly on FMassElementBitSet")
 	FMassChunkFragmentBitSet GetChunkFragments() const;
-	UE_DEPRECATED(5.8, "Getting a subset of compossition's bitsets is slow now. Consider refactoring your code to work directly on FMassElementBitSet")
+	UE_DEPRECATED(5.8, "Getting a subset of composition's bitsets is slow now. Consider refactoring your code to work directly on FMassElementBitSet")
 	FMassSharedFragmentBitSet GetSharedFragments() const;
-	UE_DEPRECATED(5.8, "Getting a subset of compossition's bitsets is slow now. Consider refactoring your code to work directly on FMassElementBitSet")
+	UE_DEPRECATED(5.8, "Getting a subset of composition's bitsets is slow now. Consider refactoring your code to work directly on FMassElementBitSet")
 	FMassConstSharedFragmentBitSet GetConstSharedFragments() const;
 
 	/**
 	 * Functions setting specific type of bits. These functions are a lot slower now.
 	 * Consider changing your code to work directly on FMassElementBitSet
 	 */
-	UE_DEPRECATED(5.8, "Setting a subset of compossition's bitsets is slow now. Consider refactoring your code to work directly on FMassElementBitSet")
+	UE_DEPRECATED(5.8, "Setting a subset of composition's bitsets is slow now. Consider refactoring your code to work directly on FMassElementBitSet")
 	void SetFragments(const FMassFragmentBitSet& InBitSet);
-	UE_DEPRECATED(5.8, "Setting a subset of compossition's bitsets is slow now. Consider refactoring your code to work directly on FMassElementBitSet")
+	UE_DEPRECATED(5.8, "Setting a subset of composition's bitsets is slow now. Consider refactoring your code to work directly on FMassElementBitSet")
 	void SetTags(const FMassTagBitSet& InBitSet);
-	UE_DEPRECATED(5.8, "Setting a subset of compossition's bitsets is slow now. Consider refactoring your code to work directly on FMassElementBitSet")
+	UE_DEPRECATED(5.8, "Setting a subset of composition's bitsets is slow now. Consider refactoring your code to work directly on FMassElementBitSet")
 	void SetChunkFragments(const FMassChunkFragmentBitSet& InBitSet);
-	UE_DEPRECATED(5.8, "Setting a subset of compossition's bitsets is slow now. Consider refactoring your code to work directly on FMassElementBitSet")
+	UE_DEPRECATED(5.8, "Setting a subset of composition's bitsets is slow now. Consider refactoring your code to work directly on FMassElementBitSet")
 	void SetSharedFragments(const FMassSharedFragmentBitSet& InBitSet);
-	UE_DEPRECATED(5.8, "Setting a subset of compossition's bitsets is slow now. Consider refactoring your code to work directly on FMassElementBitSet")
+	UE_DEPRECATED(5.8, "Setting a subset of composition's bitsets is slow now. Consider refactoring your code to work directly on FMassElementBitSet")
 	void SetConstSharedFragments(const FMassConstSharedFragmentBitSet& InBitSet);
 	
 	void SetAllSharedElements(const FMassElementBitSet& InBitSet);
@@ -619,7 +619,7 @@ struct FMassArchetypeSharedFragmentValues
 		UE::TWriteScopeLock ScopeLock{ HashLock };
 		HashCache = UINT32_MAX;
 		// we consider a single shared fragment as being "sorted"
-		bSorted = (SharedFragments.Num() + ConstSharedFragments.Num() <= 1) ;
+		bSorted = (SharedFragments.Num() + ConstSharedFragments.Num() <= 1);
 	}
 
 	/** Returns the cached-or-calculated hash */
@@ -731,14 +731,17 @@ enum class EMassExecutionContextType : uint8
 struct FMassGenericPayloadView
 {
 	FMassGenericPayloadView() = default;
-	FMassGenericPayloadView(TArray<FStructArrayView>&SourceData)
+	FMassGenericPayloadView(TArray<FStructArrayView>& SourceData)
 		: Content(SourceData)
 	{}
 	FMassGenericPayloadView(TArrayView<FStructArrayView> SourceData)
 		: Content(SourceData)
 	{}
 
-	int32 Num() const { return Content.Num(); }
+	int32 Num() const
+	{
+		return Content.Num();
+	}
 
 	void Reset()
 	{
@@ -882,12 +885,18 @@ namespace UE::Mass
 			A.Add(FStructArrayView(FragmentInstances));
 		}
 
-		void GetheredAffectedFragments(FMassFragmentBitSet& OutBitSet) const
+		void GatherAffectedFragments(FMassFragmentBitSet& OutBitSet) const
 		{
-			Super::GetheredAffectedFragments(OutBitSet);
+			Super::GatherAffectedFragments(OutBitSet);
 			OutBitSet += FMassFragmentBitSet::GetTypeBitSet<FType>();
 		}
 
+		UE_DEPRECATED(5.8, "Function name's spelling has been fixed. Use GatherAffectedFragments now")
+		void GetheredAffectedFragments(FMassFragmentBitSet& OutBitSet) const
+		{
+			GatherAffectedFragments(OutBitSet);
+		}
+
 		void Reset()
 		{
 			Super::Reset();
@@ -929,11 +938,17 @@ namespace UE::Mass
 			A.Add(FStructArrayView(FragmentInstances));
 		}
 
-		void GetheredAffectedFragments(FMassFragmentBitSet& OutBitSet) const
+		void GatherAffectedFragments(FMassFragmentBitSet& OutBitSet) const
 		{
 			OutBitSet += FMassFragmentBitSet::GetTypeBitSet<FType>();
 		}
 
+		UE_DEPRECATED(5.8, "Function name's spelling has been fixed. Use GatherAffectedFragments now")
+		void GetheredAffectedFragments(FMassFragmentBitSet& OutBitSet) const
+		{
+			GatherAffectedFragments(OutBitSet);
+		}
+
 		void Reset()
 		{
 			FragmentInstances.Reset();
diff --git a/Engine/Source/Runtime/MassEntity/Public/MassEntityView.h b/Engine/Source/Runtime/MassEntity/Public/MassEntityView.h
--- a/Engine/Source/Runtime/MassEntity/Public/MassEntityView.h
+++ b/Engine/Source/Runtime/MassEntity/Public/MassEntityView.h
@@ -46,7 +46,10 @@ struct FMassEntityView
 	 */
 	static UE_API FMassEntityView TryMakeView(const FMassEntityManager& EntityManager, FMassEntityHandle Entity);
 
-	FMassEntityHandle GetEntity() const	{ return Entity; }
+	FMassEntityHandle GetEntity() const
+	{
+		return Entity;
+	}
 
 	/** will fail a check if the viewed entity doesn't have the given fragment */
 	template<typename T>
@@ -56,7 +59,7 @@ struct FMassEntityView
 			"Given struct doesn't represent a valid fragment type but a tag. Use HasTag instead.");
 		MASS_STATIC_CHECK_FRAGMENT(T);
 
-		return *((T*)GetFragmentPtrChecked(*T::StaticStruct()));
+		return *static_cast<T*>(GetFragmentPtrChecked(*T::StaticStruct()));
 	}
 
 	/** if the viewed entity doesn't have the given fragment the function will return null */
@@ -67,7 +70,7 @@ struct FMassEntityView
 			"Given struct doesn't represent a valid fragment type but a tag. Use HasTag instead.");
 		MASS_STATIC_CHECK_FRAGMENT(T);
 
-		return (T*)GetFragmentPtr(*T::StaticStruct());
+		return static_cast<T*>(GetFragmentPtr(*T::StaticStruct()));
 	}
 
 	FStructView GetFragmentDataStruct(const UScriptStruct* FragmentType) const
@@ -97,7 +100,7 @@ struct FMassEntityView
 		static_assert(UE::Mass::CConstSharedFragment<T>,
 			"Given struct doesn't represent a valid const shared fragment type. Make sure to inherit from FMassConstSharedFragment or one of its child-types.");
 
-		return (const T*)GetConstSharedFragmentPtr(*T::StaticStruct());
+		return static_cast<const T*>(GetConstSharedFragmentPtr(*T::StaticStruct()));
 	}
 
 	/** will fail a check if the viewed entity doesn't have the given const shared fragment */
@@ -107,7 +110,7 @@ struct FMassEntityView
 		static_assert(UE::Mass::CConstSharedFragment<T>,
 			"Given struct doesn't represent a valid const shared fragment type. Make sure to inherit from FMassConstSharedFragment or one of its child-types.");
 
-		return *((const T*)GetConstSharedFragmentPtrChecked(*T::StaticStruct()));
+		return *static_cast<const T*>(GetConstSharedFragmentPtrChecked(*T::StaticStruct()));
 	}
 
 	FConstStructView GetConstSharedFragmentDataStruct(const UScriptStruct* FragmentType) const
@@ -120,14 +123,14 @@ struct FMassEntityView
 	template<UE::Mass::CSharedFragment T>
 	T& GetSharedFragmentData() const
 	{
-		return *((T*)GetSharedFragmentPtrChecked(*T::StaticStruct()));
+		return *static_cast<T*>(GetSharedFragmentPtrChecked(*T::StaticStruct()));
 	}
 
 	/** if the viewed entity doesn't have the given shared fragment the function will return null */
 	template<UE::Mass::CSharedFragment T>
 	T* GetSharedFragmentDataPtr() const
 	{
-		return (T*)GetSharedFragmentPtr(*T::StaticStruct());
+		return static_cast<T*>(GetSharedFragmentPtr(*T::StaticStruct()));
 	}
 
 	FStructView GetSharedFragmentDataStruct(const UScriptStruct* FragmentType) const
diff --git a/Engine/Source/Runtime/MassEntity/Public/MassExecutionContext.h b/Engine/Source/Runtime/MassEntity/Public/MassExecutionContext.h
--- a/Engine/Source/Runtime/MassEntity/Public/MassExecutionContext.h
+++ b/Engine/Source/Runtime/MassEntity/Public/MassExecutionContext.h
@@ -46,10 +46,18 @@ private:
 		FMassFragmentRequirementDescription Requirement;
 		ViewType FragmentView;
 
-		TFragmentView() {}
-		explicit TFragmentView(const FMassFragmentRequirementDescription& InRequirement) : Requirement(InRequirement) {}
+		TFragmentView()
+		{
+		}
+		explicit TFragmentView(const FMassFragmentRequirementDescription& InRequirement)
+			: Requirement(InRequirement)
+		{
+		}
 
-		bool operator==(const UScriptStruct* FragmentType) const { return Requirement.StructType == FragmentType; }
+		bool operator==(const UScriptStruct* FragmentType) const
+		{
+			return Requirement.StructType == FragmentType;
+		}
 	};
 	using FFragmentView = TFragmentView<TArrayView<FMassFragment>>;
 	TArray<FFragmentView, TInlineAllocator<8>> FragmentViews;
@@ -148,10 +156,22 @@ private:
 	/** Active permissions for archetype creation (processor-level when no query is active; query-level otherwise) */
 	const FMassEntityCreationRequirements* ActiveEntityCreationRequirements = nullptr;
 
-	TArrayView<FFragmentView> GetMutableRequirements() { return FragmentViews; }
-	TArrayView<FChunkFragmentView> GetMutableChunkRequirements() { return ChunkFragmentViews; }
-	TArrayView<FConstSharedFragmentView> GetMutableConstSharedRequirements() { return ConstSharedFragmentViews; }
-	TArrayView<FSharedFragmentView> GetMutableSharedRequirements() { return SharedFragmentViews; }
+	TArrayView<FFragmentView> GetMutableRequirements()
+	{
+		return FragmentViews;
+	}
+	TArrayView<FChunkFragmentView> GetMutableChunkRequirements()
+	{
+		return ChunkFragmentViews;
+	}
+	TArrayView<FConstSharedFragmentView> GetMutableConstSharedRequirements()
+	{
+		return ConstSharedFragmentViews;
+	}
+	TArrayView<FSharedFragmentView> GetMutableSharedRequirements()
+	{
+		return SharedFragmentViews;
+	}
 
 	void GetSubsystemRequirementBits(FMassExternalSubsystemBitSet& OutConstSubsystemsBitSet, FMassExternalSubsystemBitSet& OutMutableSubsystemsBitSet)
 	{
@@ -185,12 +205,27 @@ public:
 	const TSharedRef<FMassEntityManager>& GetSharedEntityManager();
 
 #if WITH_MASSENTITY_DEBUG
-	const FString& DebugGetExecutionDesc() const { return DebugExecutionDescription; }
-	void DebugSetExecutionDesc(const FString& Description) { DebugExecutionDescription = Description; }
-	void DebugSetExecutionDesc(FString&& Description) { DebugExecutionDescription = MoveTemp(Description); }
+	const FString& DebugGetExecutionDesc() const
+	{
+		return DebugExecutionDescription;
+	}
+	void DebugSetExecutionDesc(const FString& Description)
+	{
+		DebugExecutionDescription = Description;
+	}
+	void DebugSetExecutionDesc(FString&& Description)
+	{
+		DebugExecutionDescription = MoveTemp(Description);
+	}
 
-	UMassProcessor* DebugGetProcessor() const { return DebugProcessor.Get(); }
-	void DebugSetProcessor(UMassProcessor* Processor) { DebugProcessor = Processor; }
+	UMassProcessor* DebugGetProcessor() const
+	{
+		return DebugProcessor.Get();
+	}
+	void DebugSetProcessor(UMassProcessor* Processor)
+	{
+		DebugProcessor = Processor;
+	}
 #endif
 
 	void PushQuery(FMassEntityQuery& InQuery);
@@ -312,7 +347,7 @@ public:
 	 */
 	struct FSparseEntityIterator : FEntityIterator
 	{
-		/** Note that the default constructor results in an "invalid" (i.e. useless) itertor. */
+		/** Note that the default constructor results in an "invalid" (i.e. useless) iterator. */
 		FSparseEntityIterator() = default;
 
 		explicit FSparseEntityIterator(FEntityIterator&& Source)
@@ -387,11 +422,24 @@ public:
 
 	MASSENTITY_API UWorld* GetWorld() const;
 
-	TSharedPtr<FMassCommandBuffer> GetSharedDeferredCommandBuffer() const { return DeferredCommandBuffer; }
-	FMassCommandBuffer& Defer() const { checkSlow(DeferredCommandBuffer.IsValid()); return *DeferredCommandBuffer.Get(); }
+	TSharedPtr<FMassCommandBuffer> GetSharedDeferredCommandBuffer() const
+	{
+		return DeferredCommandBuffer;
+	}
+	FMassCommandBuffer& Defer() const
+	{
+		checkSlow(DeferredCommandBuffer.IsValid());
+		return *DeferredCommandBuffer.Get();
+	}
 
-	TConstArrayView<FMassEntityHandle> GetEntities() const { return EntityListView; }
-	int32 GetNumEntities() const { return EntityListView.Num(); }
+	TConstArrayView<FMassEntityHandle> GetEntities() const
+	{
+		return EntityListView;
+	}
+	int32 GetNumEntities() const
+	{
+		return EntityListView.Num();
+	}
 
 	FMassEntityHandle GetEntity(const int32 Index) const
 	{
@@ -577,7 +625,7 @@ public:
 		const FFragmentView* View = FragmentViews.FindByPredicate([FragmentType](const FFragmentView& Element) { return Element.Requirement.StructType == FragmentType; });
 		MASS_CHECK_IF_VALID(View, FragmentType);
 		MASS_CHECK_IF_READWRITE(View);
-		return MakeArrayView<TFragment>((TFragment*)View->FragmentView.GetData(), View->FragmentView.Num());
+		return MakeArrayView<TFragment>(static_cast<TFragment*>(View->FragmentView.GetData()), View->FragmentView.Num());
 	}
 
 	template<typename TFragment>
@@ -586,14 +634,14 @@ public:
 		const UScriptStruct* FragmentType = TFragment::StaticStruct();
 		const FFragmentView* View = FragmentViews.FindByPredicate([FragmentType](const FFragmentView& Element) { return Element.Requirement.StructType == FragmentType; });
 		MASS_CHECK_IF_VALID(View, TFragment::StaticStruct());
-		return TConstArrayView<TFragment>((const TFragment*)View->FragmentView.GetData(), View->FragmentView.Num());
+		return TConstArrayView<TFragment>(static_cast<const TFragment*>(View->FragmentView.GetData()), View->FragmentView.Num());
 	}
 
 	TConstArrayView<FMassFragment> GetFragmentView(const UScriptStruct* FragmentType) const
 	{
 		const FFragmentView* View = FragmentViews.FindByPredicate([FragmentType](const FFragmentView& Element) { return Element.Requirement.StructType == FragmentType; });
 		MASS_CHECK_IF_VALID(View, FragmentType);
-		return TConstArrayView<FMassFragment>((const FMassFragment*)View->FragmentView.GetData(), View->FragmentView.Num());;
+		return TConstArrayView<FMassFragment>(static_cast<const FMassFragment*>(View->FragmentView.GetData()), View->FragmentView.Num());
 	}
 
 	TArrayView<FMassFragment> GetMutableFragmentView(const UScriptStruct* FragmentType) 
@@ -617,7 +665,7 @@ public:
 	{
 		check(FragmentType->IsChildOf(TFragmentBase::StaticStruct()));
 		TArrayView<FMassFragment> View = GetMutableFragmentView(FragmentType);
-		return TArrayView<TFragmentBase>(reinterpret_cast<TFragmentBase*>(View.GetData()), View.Num());;
+		return TArrayView<TFragmentBase>(reinterpret_cast<TFragmentBase*>(View.GetData()), View.Num());
 	}
 
 	template<typename T>
@@ -789,10 +837,19 @@ public:
 	MASSENTITY_API bool HasSparseElement(TNotNull<const UScriptStruct*> ElementType, const FEntityIterator& EntityIterator) const;
 
 	/** Sparse chunk related operation */
-	const FMassArchetypeEntityCollection& GetEntityCollection() const { return EntityCollection; }
+	const FMassArchetypeEntityCollection& GetEntityCollection() const
+	{
+		return EntityCollection;
+	}
 
-	const FInstancedStruct& GetAuxData() const { return AuxData; }
-	FInstancedStruct& GetMutableAuxData() { return AuxData; }
+	const FInstancedStruct& GetAuxData() const
+	{
+		return AuxData;
+	}
+	FInstancedStruct& GetMutableAuxData()
+	{
+		return AuxData;
+	}
 	
 	template<typename TFragment>
 	bool ValidateAuxDataType() const
diff --git a/Engine/Source/Runtime/MassEntity/Public/MassExecutor.h b/Engine/Source/Runtime/MassEntity/Public/MassExecutor.h
--- a/Engine/Source/Runtime/MassEntity/Public/MassExecutor.h
+++ b/Engine/Source/Runtime/MassEntity/Public/MassExecutor.h
@@ -50,7 +50,7 @@ namespace UE::Mass::Executor
 		, ENamedThreads::Type CurrentThread = ENamedThreads::GameThread
 		, TFunction<void()> OnPreCommandFlushNotification = TFunction<void()>());
 
-	UE_DEPRECATED(5.6, "lvalue flavor of TriggerParallelTasks has been deprevate. Use the rvalue version.")
+	UE_DEPRECATED(5.6, "lvalue flavor of TriggerParallelTasks has been deprecated. Use the rvalue version.")
 	MASSENTITY_API FGraphEventRef TriggerParallelTasks(UMassProcessor& Processor, FProcessingContext& ProcessingContext, TFunction<void()> OnDoneNotification
 		, ENamedThreads::Type CurrentThread = ENamedThreads::GameThread);
 } // namespace UE::Mass::Executor
diff --git a/Engine/Source/Runtime/MassEntity/Public/MassObserverManager.h b/Engine/Source/Runtime/MassEntity/Public/MassObserverManager.h
--- a/Engine/Source/Runtime/MassEntity/Public/MassObserverManager.h
+++ b/Engine/Source/Runtime/MassEntity/Public/MassObserverManager.h
@@ -143,7 +143,7 @@ struct FMassObserverManager
 
 protected:
 	/**
-	 * helper struct for holding either a single enity handle or an archetype collection reference.
+	 * helper struct for holding either a single entity handle or an archetype collection reference.
 	 * The type is highly transient, do not store instances of it. Its only function is to 
 	 * allow us to have a single OnCompositionChange implementation rather than having two separate
 	 * implementations, one for entity handle, the other for archetype collection.
@@ -232,7 +232,10 @@ protected:
 
 	MASSENTITY_API TSharedRef<FObserverLock> GetOrMakeObserverLock();
 	TSharedPtr<FObserverLock> GetObserverLock() const;
-	bool IsLocked() const { return ActiveObserverLock.IsValid(); }
+	bool IsLocked() const
+	{
+		return ActiveObserverLock.IsValid();
+	}
 	MASSENTITY_API TSharedRef<FCreationContext> GetOrMakeCreationContext();
 	MASSENTITY_API TSharedRef<FCreationContext> GetOrMakeCreationContext(TConstArrayView<FMassEntityHandle> ReservedEntities, FMassArchetypeEntityCollection&& EntityCollection);
 	TSharedPtr<FCreationContext> GetCreationContext() const;
diff --git a/Engine/Source/Runtime/MassEntity/Public/MassObserverRegistry.h b/Engine/Source/Runtime/MassEntity/Public/MassObserverRegistry.h
--- a/Engine/Source/Runtime/MassEntity/Public/MassObserverRegistry.h
+++ b/Engine/Source/Runtime/MassEntity/Public/MassObserverRegistry.h
@@ -15,7 +15,7 @@ struct FMassObserverManager;
 /**
  * A wrapper type for a TArray to support having map-of-arrays UPROPERTY members in FMassEntityObserverClassesMap
  */
-USTRUCT(meta = (Deprecated = "5.7", DepracationMessage = "FMassProcessorClassCollection is no longer being used and will be removed in the upcoming engine released"))
+USTRUCT(meta = (Deprecated = "5.7", DeprecationMessage = "FMassProcessorClassCollection is no longer being used and will be removed in the upcoming engine release"))
 struct FMassProcessorClassCollection
 {
 	GENERATED_BODY()
@@ -27,7 +27,7 @@ struct FMassProcessorClassCollection
 /**
  * A wrapper type for a TMap to support having array-of-maps UPROPERTY members in UMassObserverRegistry
  */
-USTRUCT(meta = (Deprecated = "5.7", DepracationMessage = "FMassEntityObserverClassesMap is no longer being used and will be removed in the upcoming engine released"))
+USTRUCT(meta = (Deprecated = "5.7", DeprecationMessage = "FMassEntityObserverClassesMap is no longer being used and will be removed in the upcoming engine release"))
 struct FMassEntityObserverClassesMap
 {
 	GENERATED_BODY()
@@ -56,8 +56,14 @@ class UMassObserverRegistry : public UObject
 public:
 	UE_API UMassObserverRegistry();
 
-	static UMassObserverRegistry& GetMutable() { return *GetMutableDefault<UMassObserverRegistry>(); }
-	static const UMassObserverRegistry& Get() { return *GetDefault<UMassObserverRegistry>(); }
+	static UMassObserverRegistry& GetMutable()
+	{
+		return *GetMutableDefault<UMassObserverRegistry>();
+	}
+	static const UMassObserverRegistry& Get()
+	{
+		return *GetDefault<UMassObserverRegistry>();
+	}
 
 	UE_API void RegisterObserver(TNotNull<const UScriptStruct*> ObservedType, uint8 OperationFlags, const TSubclassOf<UMassProcessor>& ObserverClass);
 
@@ -80,11 +86,11 @@ protected:
 
 	UE_DEPRECATED(5.7, "Use ElementObserverMaps instead")
 	UPROPERTY()
-	FMassEntityObserverClassesMap FragmentObservers[(uint8)EMassObservedOperation::MAX];
+	FMassEntityObserverClassesMap FragmentObservers[static_cast<uint8>(EMassObservedOperation::MAX)];
 
 	UE_DEPRECATED(5.7, "Use ElementObserverMaps instead")
 	UPROPERTY()
-	FMassEntityObserverClassesMap TagObservers[(uint8)EMassObservedOperation::MAX];
+	FMassEntityObserverClassesMap TagObservers[static_cast<uint8>(EMassObservedOperation::MAX)];
 };
 
 #undef UE_API
diff --git a/Engine/Source/Runtime/MassEntity/Public/MassProcessingContext.h b/Engine/Source/Runtime/MassEntity/Public/MassProcessingContext.h
--- a/Engine/Source/Runtime/MassEntity/Public/MassProcessingContext.h
+++ b/Engine/Source/Runtime/MassEntity/Public/MassProcessingContext.h
@@ -38,7 +38,7 @@ namespace UE::Mass
 		UE_DEPRECATED_FORGAME(5.6, "Direct access to FProcessingContext.EntityManager has been deprecated. Use GetEntityManager instead.")
 		TSharedRef<FMassEntityManager> EntityManager;
 
-		UE_DEPRECATED_FORGAME(5.6, "Direct access to FProcessingContex.DeltaSeconds has been deprecated. Set it via the constructor instead, read via the getter.")
+		UE_DEPRECATED_FORGAME(5.6, "Direct access to FProcessingContext.DeltaSeconds has been deprecated. Set it via the constructor instead, read via the getter.")
 		float DeltaSeconds = 0.f;
 
 		FInstancedStruct AuxData;
@@ -47,7 +47,7 @@ namespace UE::Mass
 		 * If set to "true" the MassExecutor will flush commands at the end of given execution function. 
 		 * If "false" the caller is responsible for manually flushing the commands.
 		 */
-		UE_DEPRECATED_FORGAME(5.6, "Direct access to FProcessingContex.bFlushCommandBuffer has been deprecated. Set it via the constructor instead, read via the GetWillFlushCommands getter.")
+		UE_DEPRECATED_FORGAME(5.6, "Direct access to FProcessingContext.bFlushCommandBuffer has been deprecated. Set it via the constructor instead, read via the GetWillFlushCommands getter.")
 		bool bFlushCommandBuffer = true; 
 
 		UE_DEPRECATED_FORGAME(5.6, "Direct access to FProcessingContext's CommandBuffer has been deprecated. Use SetCommandBuffer()")
@@ -61,7 +61,7 @@ namespace UE::Mass
 	//----------------------------------------------------------------------//
 	// INLINES
 	//----------------------------------------------------------------------//
-	// @todo remove the depracation disabling once CommandBuffer and EntityManager are moved to `protected` and un-deprecated (around UE5.8)
+	// @todo remove the deprecation disabling once CommandBuffer and EntityManager are moved to `protected` and un-deprecated (around UE5.8)
 	PRAGMA_DISABLE_DEPRECATION_WARNINGS
 	inline FProcessingContext::FProcessingContext(FMassEntityManager& InEntityManager, const float InDeltaSeconds, const bool bInFlushCommandBuffer)
 		: FProcessingContext(/*MoveTemp*/(InEntityManager.AsShared()), InDeltaSeconds, bInFlushCommandBuffer)
diff --git a/Engine/Source/Runtime/MassEntity/Public/MassProcessingPhaseManager.h b/Engine/Source/Runtime/MassEntity/Public/MassProcessingPhaseManager.h
--- a/Engine/Source/Runtime/MassEntity/Public/MassProcessingPhaseManager.h
+++ b/Engine/Source/Runtime/MassEntity/Public/MassProcessingPhaseManager.h
@@ -68,16 +68,34 @@ protected:
 
 	MASSENTITY_API void OnParallelExecutionDone(const float DeltaTime);
 
-	bool IsConfiguredForParallelMode() const { return bRunInParallelMode; }
-	void ConfigureForParallelMode() { bRunInParallelMode = true; }
-	void ConfigureForSingleThreadMode() { bRunInParallelMode = false; }
+	bool IsConfiguredForParallelMode() const
+	{
+		return bRunInParallelMode;
+	}
+	void ConfigureForParallelMode()
+	{
+		bRunInParallelMode = true;
+	}
+	void ConfigureForSingleThreadMode()
+	{
+		bRunInParallelMode = false;
+	}
 
-	bool ShouldTick(const ELevelTick TickType) const { return SupportedTickTypes & (1 << TickType); }
+	bool ShouldTick(const ELevelTick TickType) const
+	{
+		return SupportedTickTypes & (1 << TickType);
+	}
 
 public:
 	MASSENTITY_API void Initialize(FMassProcessingPhaseManager& InPhaseManager, const EMassProcessingPhase InPhase, const ETickingGroup InTickGroup, UMassCompositeProcessor& InPhaseProcessor);
-	void AddSupportedTickType(const ELevelTick TickType) { SupportedTickTypes |= (1 << TickType); }
-	void RemoveSupportedTickType(const ELevelTick TickType) { SupportedTickTypes &= ~(1 << TickType); }
+	void AddSupportedTickType(const ELevelTick TickType)
+	{
+		SupportedTickTypes |= (1 << TickType);
+	}
+	void RemoveSupportedTickType(const ELevelTick TickType)
+	{
+		SupportedTickTypes &= ~(1 << TickType);
+	}
 
 #if WITH_MASSENTITY_DEBUG
 	const UMassCompositeProcessor* DebugGetPhaseProcessor() const;
@@ -150,13 +168,26 @@ public:
 	FMassProcessingPhaseManager(const FMassProcessingPhaseManager& Other) = delete;
 	FMassProcessingPhaseManager& operator=(const FMassProcessingPhaseManager& Other) = delete;
 
-	const TSharedPtr<FMassEntityManager>& GetEntityManager() const { return EntityManager; }
-	FMassEntityManager& GetEntityManagerRef() { check(EntityManager); return *EntityManager.Get(); }
+	const TSharedPtr<FMassEntityManager>& GetEntityManager() const
+	{
+		return EntityManager;
+	}
+	FMassEntityManager& GetEntityManagerRef()
+	{
+		check(EntityManager);
+		return *EntityManager.Get();
+	}
 
 	/** Retrieves OnPhaseStart multicast delegate's reference for a given Phase */
-	FMassProcessingPhase::FOnPhaseEvent& GetOnPhaseStart(const EMassProcessingPhase Phase) { return ProcessingPhases[uint8(Phase)].OnPhaseStart; }
+	FMassProcessingPhase::FOnPhaseEvent& GetOnPhaseStart(const EMassProcessingPhase Phase)
+	{
+		return ProcessingPhases[static_cast<uint8>(Phase)].OnPhaseStart;
+	}
 	/** Retrieves OnPhaseEnd multicast delegate's reference for a given Phase */
-	FMassProcessingPhase::FOnPhaseEvent& GetOnPhaseEnd(const EMassProcessingPhase Phase) { return ProcessingPhases[uint8(Phase)].OnPhaseEnd; }
+	FMassProcessingPhase::FOnPhaseEvent& GetOnPhaseEnd(const EMassProcessingPhase Phase)
+	{
+		return ProcessingPhases[static_cast<uint8>(Phase)].OnPhaseEnd;
+	}
 
 	/** 
 	 *  Populates hosted FMassProcessingPhase instances with Processors read from MassEntitySettings configuration.
@@ -180,7 +211,10 @@ public:
 	 */
 	MASSENTITY_API void Start(const TSharedRef<FMassEntityManager>& InEntityManager);
 	MASSENTITY_API void Stop();
-	bool IsRunning() const { return EntityManager.IsValid(); }
+	bool IsRunning() const
+	{
+		return EntityManager.IsValid();
+	}
 
 	/**
 	 * Determine if this Phase Manager is currently paused.
@@ -223,15 +257,15 @@ public:
 	FTickFunction& GetProcessingPhaseTickFunction(EMassProcessingPhase Phase)
 	{
 		check(Phase < EMassProcessingPhase::MAX);
-		check((int)Phase >= 0 );
-		return ProcessingPhases[(int)Phase];
+		check(static_cast<int32>(Phase) >= 0);
+		return ProcessingPhases[static_cast<int32>(Phase)];
 	}
-	
+
 	const FTickFunction& GetProcessingPhaseTickFunction(EMassProcessingPhase Phase) const
 	{
 		check(Phase < EMassProcessingPhase::MAX);
-		check((int)Phase >= 0 );
-		return ProcessingPhases[(int)Phase];
+		check(static_cast<int32>(Phase) >= 0);
+		return ProcessingPhases[static_cast<int32>(Phase)];
 	}
 
 #if WITH_MASSENTITY_DEBUG
@@ -275,8 +309,8 @@ protected:
 	MASSENTITY_API void OnNewArchetype(const FMassArchetypeHandle& NewArchetype);
 
 protected:
-	FMassProcessingPhase ProcessingPhases[(uint8)EMassProcessingPhase::MAX];
-	FPhaseGraphBuildState ProcessingGraphBuildStates[(uint8)EMassProcessingPhase::MAX];
+	FMassProcessingPhase ProcessingPhases[static_cast<uint8>(EMassProcessingPhase::MAX)];
+	FPhaseGraphBuildState ProcessingGraphBuildStates[static_cast<uint8>(EMassProcessingPhase::MAX)];
 	TArray<FMassProcessingPhaseConfig> ProcessingPhasesConfig;
 	TArray<TObjectPtr<UMassProcessor>> DynamicProcessors;
 	TArray<TWeakObjectPtr<UMassProcessor>> RemovedDynamicProcessors;
@@ -287,7 +321,7 @@ protected:
 	};
 	/** using TStrongObjectPtr to not worry about GC while the processor instances are waiting in PendingDynamicProcessors */
 	using FDynamicProcessorOperation = TPair<TStrongObjectPtr<UMassProcessor>, EDynamicProcessorOperationType>;
-	TMpscQueue<FDynamicProcessorOperation> PendingDynamicProcessors[(uint8)EMassProcessingPhase::MAX];
+	TMpscQueue<FDynamicProcessorOperation> PendingDynamicProcessors[static_cast<uint8>(EMassProcessingPhase::MAX)];
 
 	TSharedPtr<FMassEntityManager> EntityManager;
 
diff --git a/Engine/Source/Runtime/MassEntity/Public/MassProcessingQueue.h b/Engine/Source/Runtime/MassEntity/Public/MassProcessingQueue.h
--- a/Engine/Source/Runtime/MassEntity/Public/MassProcessingQueue.h
+++ b/Engine/Source/Runtime/MassEntity/Public/MassProcessingQueue.h
@@ -52,7 +52,7 @@ namespace UE::Mass
 			Any = MainThread | OffMainThread
 		};
 
-		/** Gets the next processor that is ready to execute, sets its state to running and updates contention tracking and other bookeeping */
+		/** Gets the next processor that is ready to execute, sets its state to running and updates contention tracking and other bookkeeping */
 		FQueuedProcessor* GetNextReadyProcessor(TArray<int32>& InOutRemainingIndices);
 
 		bool IsComplete() const;
@@ -69,9 +69,9 @@ namespace UE::Mass
 		void Reserve(int32 RequestorIndex);
 
 		/** Release a reservation and immediately attempt TryReserve on a different index */
-		bool ReleaseAndTryReserve(int32 IndexToRealease, int32 IndexToReserve);
+		bool ReleaseAndTryReserve(int32 IndexToRelease, int32 IndexToReserve);
 		/** ReleaseAndReserve, then if that fails immediately Reserve (which will wait) */
-		void ReleaseAndReserve(int32 IndexToRealease, int32 IndexToReserve);
+		void ReleaseAndReserve(int32 IndexToRelease, int32 IndexToReserve);
 
 		bool TryActivateChunk(FQueuedChunk& InChunk, int32 DirectChunkContentionIndex);
 		void DeactivateChunk(FQueuedChunk& InChunk);
@@ -79,7 +79,7 @@ namespace UE::Mass
 		void ReserveChunkWorkerSlot();
 		void ReleaseChunkWorkerSlot();
 
-		/** Report that a processor is finished, set its state internal state and update queue bookeeping */
+		/** Report that a processor is finished, set its internal state and update queue bookkeeping */
 		void ProcessorFinished(FQueuedProcessor& CompletedProcessor);
 
 		/**
diff --git a/Engine/Source/Runtime/MassEntity/Public/MassProcessingQueueTypes.h b/Engine/Source/Runtime/MassEntity/Public/MassProcessingQueueTypes.h
--- a/Engine/Source/Runtime/MassEntity/Public/MassProcessingQueueTypes.h
+++ b/Engine/Source/Runtime/MassEntity/Public/MassProcessingQueueTypes.h
@@ -315,7 +315,7 @@ namespace UE::Mass
 
 		void ReleaseAccess(FQueuedChunk& Chunk, int32 DirectChunkContentionIndex)
 		{
-			const int32 SlotIndex = (int32)Chunk.ActiveChunkSlotIndex;
+			const int32 SlotIndex = static_cast<int32>(Chunk.ActiveChunkSlotIndex);
 			checkSlow(SlotIndex >= 0 && SlotIndex < MaxActiveChunks);
 
 			const uint64 ChunkID = Chunk.GetChunkID();
@@ -550,7 +550,7 @@ namespace UE::Mass
 
 				while (ContenderWord != 0)
 				{
-					const int32 TrailingZeroCount = (int32)FMath::CountTrailingZeros64(ContenderWord);
+					const int32 TrailingZeroCount = static_cast<int32>(FMath::CountTrailingZeros64(ContenderWord));
 					const int32 NeighborIndex = BaseBitIndex + TrailingZeroCount;
 
 					if (NeighborIndex < MatrixSize)
@@ -753,13 +753,13 @@ namespace UE::Mass
 
 	/**
 	 * Increments a change-tracking atomic variable when leaving a scope.
-	 * Used with FChangeListener for efficient syncronizaion of workers waiting for resources.
+	 * Used with FChangeListener for efficient synchronization of workers waiting for resources.
 	 */
 	template <typename T>
 	struct FChangeScope
 	{
 		explicit FChangeScope(std::atomic<T>& InTrackingValue)
-			:Tracker(InTrackingValue)
+			: Tracker(InTrackingValue)
 		{
 		}
 		FChangeScope(const FChangeScope&) = delete;
@@ -776,13 +776,13 @@ namespace UE::Mass
 
 	/**
 	 * Waits for an atomic variable to change or become a specific value.
-	 * Used with FChangeScope for efficient syncronizaion of workers waiting for resources.
+	 * Used with FChangeScope for efficient synchronization of workers waiting for resources.
 	 */
 	template <typename T>
 	struct FChangeListener
 	{
 		explicit FChangeListener(std::atomic<T>& InValueToTrack)
-			:ValueToTrack(InValueToTrack)
+			: ValueToTrack(InValueToTrack)
 			// Set the initial value at construction time to capture changes that happen between declaration and subsequent Wait calls.
 			, InitialValue(ValueToTrack.load(std::memory_order_acquire))
 		{
diff --git a/Engine/Source/Runtime/MassEntity/Public/MassProcessor.h b/Engine/Source/Runtime/MassEntity/Public/MassProcessor.h
--- a/Engine/Source/Runtime/MassEntity/Public/MassProcessor.h
+++ b/Engine/Source/Runtime/MassEntity/Public/MassProcessor.h
@@ -189,7 +189,7 @@ public:
 	int32 GetOwnedQueriesNum() const;
 
 	/**
-	 * Indicates weather or not this processor has declared in its requirements that it creates entities.
+	 * Indicates whether or not this processor has declared in its requirements that it creates entities.
 	 */
 	UE_API bool HasEntityCreationRequirements();
 
@@ -351,7 +351,9 @@ public:
 	UE_DEPRECATED(5.6, "Initialize is deprecated. Override InitializeInternal(UObject&, const TSharedRef<FMassEntityManager>&) instead. If you want to call the function, use CallInitialize.")
 	UE_API virtual void Initialize(UObject& Owner) final;
 	UE_DEPRECATED(5.6, "This flavor of ConfigureQueries is deprecated. Override ConfigureQueries(const TSharedRef<FMassEntityManager>&) instead.")
-	virtual void ConfigureQueries() final {};
+	virtual void ConfigureQueries() final
+	{
+	};
 };
 
 
diff --git a/Engine/Source/Runtime/MassEntity/Public/MassProcessorDependencySolver.h b/Engine/Source/Runtime/MassEntity/Public/MassProcessorDependencySolver.h
--- a/Engine/Source/Runtime/MassEntity/Public/MassProcessorDependencySolver.h
+++ b/Engine/Source/Runtime/MassEntity/Public/MassProcessorDependencySolver.h
@@ -41,9 +41,15 @@ struct TMassExecutionAccess
 		return OpIndex == EMassAccessOperation::Read ? Read : Write;
 	}
 
-	TConstArrayView<T> AsArrayView() const { return MakeArrayView(&Read, 2); }
+	TConstArrayView<T> AsArrayView() const
+	{
+		return MakeArrayView(&Read, 2);
+	}
 
-	bool IsEmpty() const { return Read.IsEmpty() && Write.IsEmpty(); }
+	bool IsEmpty() const
+	{
+		return Read.IsEmpty() && Write.IsEmpty();
+	}
 };
 
 /** 
@@ -54,13 +60,13 @@ template<>
 struct TMassExecutionAccess<FMassConstSharedFragmentBitSet>
 {
 	FMassConstSharedFragmentBitSet Read;
-	TConstArrayView<FMassConstSharedFragmentBitSet> AsArrayView() const 
-	{ 
-		return MakeArrayView(&Read, 1); 
+	TConstArrayView<FMassConstSharedFragmentBitSet> AsArrayView() const
+	{
+		return MakeArrayView(&Read, 1);
 	}
-	bool IsEmpty() const 
-	{ 
-		return Read.IsEmpty(); 
+	bool IsEmpty() const
+	{
+		return Read.IsEmpty();
 	}
 };
 
@@ -90,11 +96,15 @@ struct FMassProcessorDependencySolver
 {
 	struct FNode
 	{
-		FNode(const FName InName, UMassProcessor* InProcessor, const int32 InNodeIndex = INDEX_NONE) 
+		FNode(const FName InName, UMassProcessor* InProcessor, const int32 InNodeIndex = INDEX_NONE)
 			: Name(InName), Processor(InProcessor), NodeIndex(InNodeIndex)
-		{}
+		{
+		}
 
-		bool IsGroup() const { return Processor == nullptr; }
+		bool IsGroup() const
+		{
+			return Processor == nullptr;
+		}
 		/**
 		 * @return `true` when everything's fine, `false` when cycles have been encountered. If that
 		 *	happens `OutCycleIndices` gets filled with the relevant node indices.
@@ -205,7 +215,10 @@ public:
 	/** Determines whether the dependency solving that produced InResult will produce different results if run with a given EntityManager */
 	static bool IsResultUpToDate(const FResult& InResult, TSharedPtr<FMassEntityManager> EntityManager);
 
-	bool IsSolvingForSingleThread() const { return bSingleThreadTarget; }
+	bool IsSolvingForSingleThread() const
+	{
+		return bSingleThreadTarget;
+	}
 
 protected:
 	// note that internals are protected rather than private to support unit testing
@@ -235,7 +248,7 @@ protected:
 	 * @Note currently the value depends on MASS_DO_PARALLEL and there's no way to configure it otherwise, but there's 
 	 * nothing inherently stopping us from letting users configure it.
 	 */
-	const bool bSingleThreadTarget = bool(!MASS_DO_PARALLEL);
+	const bool bSingleThreadTarget = !MASS_DO_PARALLEL;
 	const bool bGameRuntime = true;
 	FString DependencyGraphFileName;
 	TArray<FNode> AllNodes;
diff --git a/Engine/Source/Runtime/MassEntity/Public/MassRequirementAccessDetector.h b/Engine/Source/Runtime/MassEntity/Public/MassRequirementAccessDetector.h
--- a/Engine/Source/Runtime/MassEntity/Public/MassRequirementAccessDetector.h
+++ b/Engine/Source/Runtime/MassEntity/Public/MassRequirementAccessDetector.h
@@ -34,7 +34,7 @@ private:
 		}
 	}
 
-	void Aquire(TConstArrayView<FMassFragmentRequirementDescription> Requirements)
+	void Acquire(TConstArrayView<FMassFragmentRequirementDescription> Requirements)
 	{
 		for (const FMassFragmentRequirementDescription& Req : Requirements)
 		{
diff --git a/Engine/Source/Runtime/MassEntity/Public/MassRequirements.h b/Engine/Source/Runtime/MassEntity/Public/MassRequirements.h
--- a/Engine/Source/Runtime/MassEntity/Public/MassRequirements.h
+++ b/Engine/Source/Runtime/MassEntity/Public/MassRequirements.h
@@ -284,7 +284,7 @@ struct FMassFragmentRequirements
 	FMassFragmentRequirements& AddTagRequirement(TNotNull<const UScriptStruct*> TagType, const EMassFragmentPresence Presence)
 	{
 		checkf(bInitialized, TEXT("Modifying requirements before initialization is not supported."));
-		checkf(int(Presence) != int(EMassFragmentPresence::MAX), TEXT("MAX presence is not a valid value for AddTagRequirement"));
+		checkf(static_cast<int32>(Presence) != static_cast<int32>(EMassFragmentPresence::MAX), TEXT("MAX presence is not a valid value for AddTagRequirement"));
 
 		if (UE::Mass::IsSparse(TagType))
 		{
@@ -321,7 +321,7 @@ struct FMassFragmentRequirements
 	FMassFragmentRequirements& AddTagRequirement(const EMassFragmentPresence Presence)
 	{
 		checkf(bInitialized, TEXT("Modifying requirements before initialization is not supported."));
-		checkf(int(Presence) != int(EMassFragmentPresence::MAX), TEXT("MAX presence is not a valid value for AddTagRequirement"));
+		checkf(static_cast<int32>(Presence) != static_cast<int32>(EMassFragmentPresence::MAX), TEXT("MAX presence is not a valid value for AddTagRequirement"));
 
 		static_assert(UE::Mass::CTag<T>, "Given struct doesn't represent a valid tag type. Make sure to inherit from FMassTag or one of its child-types.");
 
diff --git a/Engine/Source/Runtime/MassEntity/Public/TypeBitSetBuilder.h b/Engine/Source/Runtime/MassEntity/Public/TypeBitSetBuilder.h
--- a/Engine/Source/Runtime/MassEntity/Public/TypeBitSetBuilder.h
+++ b/Engine/Source/Runtime/MassEntity/Public/TypeBitSetBuilder.h
@@ -77,7 +77,7 @@ public:
 
 	/**
 	 * Constructor that initializes the builder with a struct tracker and a source bitset.
-	 * @param InStructTracker he struct tracker to use.
+	 * @param InStructTracker the struct tracker to use.
 	 * @param Source the source bitset to initialize from.
 	 */
 	TTypeBitSetBuilder(FStructTracker& InStructTracker, FConstBitSet& Source)
