Fix copy-paste bug in FQueryExecutor::ValidateAccessors ??? was comparing `this` against itself instead of AccessorsPtr, making the member-variable range check a no-op.

#rnx

#Synced-CL 52057779

diff --git a/Engine/Source/Runtime/MassEntity/Private/MassQueryExecutor.cpp b/Engine/Source/Runtime/MassEntity/Private/MassQueryExecutor.cpp
--- a/Engine/Source/Runtime/MassEntity/Private/MassQueryExecutor.cpp
+++ b/Engine/Source/Runtime/MassEntity/Private/MassQueryExecutor.cpp
@@ -43,7 +43,7 @@ void FQueryExecutor::ValidateAccessors()
 
 	if (AccessorsPtr)
 	{
-		const UPTRINT AccessorsStart = reinterpret_cast<UPTRINT>(this);
+		const UPTRINT AccessorsStart = reinterpret_cast<UPTRINT>(AccessorsPtr);
 		checkf(ExecutorStart <= AccessorsStart && AccessorsStart <= ExecutorEnd, TEXT("Accessors assigned to a FQueryExecutor must be member variables of that struct."));
 	}
 }
