[Mass] Rename RunSparse overloads in MassExecutor for clarity

Renamed the two RunSparse overloads to better reflect their intent:
- RunSparse(Archetype, Entities) -> RunSparseEntities - processes a sparse set of entities from a single archetype
- RunSparse(EntityCollection) -> RunWithCollection - runs the pipeline on a pre-built FMassArchetypeEntityCollection

Added UE_DEPRECATED(5.8) inline wrappers preserving the old RunSparse signatures for backward compatibility. Updated all internal callers (executor implementation and tests) to the new names. Resolves a [REDACTED_HANDLE] in the original header comment.

#Synced-CL 52096882

diff --git a/Engine/Source/Developer/MassEntityTestSuite/Private/MassExecutionTest.cpp b/Engine/Source/Developer/MassEntityTestSuite/Private/MassExecutionTest.cpp
--- a/Engine/Source/Developer/MassEntityTestSuite/Private/MassExecutionTest.cpp
+++ b/Engine/Source/Developer/MassEntityTestSuite/Private/MassExecutionTest.cpp
@@ -144,7 +144,7 @@ struct FExecution_Sparse : FEntityTestBase
 		FMassArchetypeEntityCollection EntityCollection(FloatsArchetype);
 		// nothing should break. The actual result of processing is getting tested in MassProcessorTest.cpp
 		
-		UE::Mass::Executor::RunSparse(Pipeline, ProcessingContext, EntityCollection);
+		UE::Mass::Executor::RunWithCollection(Pipeline, ProcessingContext, EntityCollection);
 
 		return true;
 	}
diff --git a/Engine/Source/Programs/MassTests/Private/MassExecutionTest.cpp b/Engine/Source/Programs/MassTests/Private/MassExecutionTest.cpp
--- a/Engine/Source/Programs/MassTests/Private/MassExecutionTest.cpp
+++ b/Engine/Source/Programs/MassTests/Private/MassExecutionTest.cpp
@@ -105,7 +105,7 @@ TEST_CASE_METHOD(FMassLLTEntityFixture, "Mass::Execution::Sparse", "[Mass][Execu
 	FMassArchetypeEntityCollection EntityCollection(FloatsArchetype);
 	// nothing should break. The actual result of processing is getting tested in MassProcessorTest.cpp
 
-	UE::Mass::Executor::RunSparse(Pipeline, ProcessingContext, EntityCollection);
+	UE::Mass::Executor::RunWithCollection(Pipeline, ProcessingContext, EntityCollection);
 }
 
 } // namespace UE::Mass::LLT
diff --git a/Engine/Source/Runtime/MassEntity/Private/MassExecutor.cpp b/Engine/Source/Runtime/MassEntity/Private/MassExecutor.cpp
--- a/Engine/Source/Runtime/MassEntity/Private/MassExecutor.cpp
+++ b/Engine/Source/Runtime/MassEntity/Private/MassExecutor.cpp
@@ -34,11 +34,11 @@ void Run(FMassRuntimePipeline& RuntimePipeline, FProcessingContext& ProcessingCo
 	RunProcessorsView(RuntimePipeline.GetMutableProcessors(), ProcessingContext);
 }
 
-void RunSparse(FMassRuntimePipeline& RuntimePipeline, FProcessingContext& ProcessingContext, FMassArchetypeHandle Archetype, TConstArrayView<FMassEntityHandle> Entities)
+void RunSparseEntities(FMassRuntimePipeline& RuntimePipeline, FProcessingContext& ProcessingContext, FMassArchetypeHandle Archetype, TConstArrayView<FMassEntityHandle> Entities)
 {
-	if (!ensure(RuntimePipeline.GetProcessors().Find(nullptr) == INDEX_NONE) 
-		|| RuntimePipeline.Num() == 0 
-		|| !ensureMsgf(Archetype.IsValid(), TEXT("The Archetype passed in to UE::Mass::Executor::RunSparse is invalid")))
+	if (!ensure(RuntimePipeline.GetProcessors().Find(nullptr) == INDEX_NONE)
+		|| RuntimePipeline.Num() == 0
+		|| !ensureMsgf(Archetype.IsValid(), TEXT("The Archetype passed in to UE::Mass::Executor::RunSparseEntities is invalid")))
 	{
 		return;
 	}
@@ -49,16 +49,16 @@ void RunSparse(FMassRuntimePipeline& RuntimePipeline, FProcessingContext& Proces
 	RunProcessorsView(RuntimePipeline.GetMutableProcessors(), ProcessingContext, MakeArrayView(&EntityCollection, 1));
 }
 
-void RunSparse(FMassRuntimePipeline& RuntimePipeline, FProcessingContext& ProcessingContext, const FMassArchetypeEntityCollection& EntityCollection)
+void RunWithCollection(FMassRuntimePipeline& RuntimePipeline, FProcessingContext& ProcessingContext, const FMassArchetypeEntityCollection& EntityCollection)
 {
-	if (!ensure(RuntimePipeline.GetProcessors().Find(nullptr) == INDEX_NONE) 
+	if (!ensure(RuntimePipeline.GetProcessors().Find(nullptr) == INDEX_NONE)
 		|| RuntimePipeline.Num() == 0
-		|| !ensureMsgf(EntityCollection.GetArchetype().IsValid(), TEXT("The Archetype of EntityCollection passed in to UE::Mass::Executor::RunSparse is invalid")))
+		|| !ensureMsgf(EntityCollection.GetArchetype().IsValid(), TEXT("The Archetype of EntityCollection passed in to UE::Mass::Executor::RunWithCollection is invalid")))
 	{
 		return;
 	}
 
-	TRACE_CPUPROFILER_EVENT_SCOPE_STR("MassExecutor RunSparse");
+	TRACE_CPUPROFILER_EVENT_SCOPE_STR("MassExecutor RunWithCollection");
 
 	RunProcessorsView(RuntimePipeline.GetMutableProcessors(), ProcessingContext, MakeArrayView(&EntityCollection, 1));
 }
diff --git a/Engine/Source/Runtime/MassEntity/Public/MassExecutor.h b/Engine/Source/Runtime/MassEntity/Public/MassExecutor.h
--- a/Engine/Source/Runtime/MassEntity/Public/MassExecutor.h
+++ b/Engine/Source/Runtime/MassEntity/Public/MassExecutor.h
@@ -24,18 +24,29 @@ namespace UE::Mass::Executor
 	/** Executes given Processor. Used mainly for triggering calculations via MassCompositeProcessors, e.g processing phases */
 	MASSENTITY_API void Run(UMassProcessor& Processor, FProcessingContext& ProcessingContext);
 
-	/** Similar to the Run function, but instead of using all the entities hosted by ProcessingContext.EntitySubsystem 
-	 *  it is processing only the entities given by EntityID via the Entities input parameter. 
-	 *  Note that all the entities need to be of Archetype archetype. 
-	 *  Under the hood the function converts Archetype-Entities pair to FMassArchetypeEntityCollection and calls the other flavor of RunSparse
+	/** Similar to the Run function, but instead of using all the entities hosted by ProcessingContext.EntitySubsystem
+	 *  it is processing only the entities given by EntityID via the Entities input parameter.
+	 *  Note that all the entities need to be of Archetype archetype.
+	 *  Under the hood the function converts Archetype-Entities pair to FMassArchetypeEntityCollection and calls RunWithCollection.
 	 */
-	MASSENTITY_API void RunSparse(FMassRuntimePipeline& RuntimePipeline, FProcessingContext& ProcessingContext, FMassArchetypeHandle Archetype, TConstArrayView<FMassEntityHandle> Entities);
+	MASSENTITY_API void RunSparseEntities(FMassRuntimePipeline& RuntimePipeline, FProcessingContext& ProcessingContext, FMassArchetypeHandle Archetype, TConstArrayView<FMassEntityHandle> Entities);
 
-	/** Similar to the Run function, but instead of using all the entities hosted by ProcessingContext.EntitySubsystem 
-	 *  it is processing only the entities given by SparseEntities input parameter.
-	 *  @todo rename
+	UE_DEPRECATED(5.8, "Use RunSparseEntities instead.")
+	inline void RunSparse(FMassRuntimePipeline& RuntimePipeline, FProcessingContext& ProcessingContext, FMassArchetypeHandle Archetype, TConstArrayView<FMassEntityHandle> Entities)
+	{
+		RunSparseEntities(RuntimePipeline, ProcessingContext, Archetype, Entities);
+	}
+
+	/** Similar to the Run function, but instead of using all the entities hosted by ProcessingContext.EntitySubsystem
+	 *  it is processing only the entities given via the EntityCollection input parameter.
 	 */
-	MASSENTITY_API void RunSparse(FMassRuntimePipeline& RuntimePipeline, FProcessingContext& ProcessingContext, const FMassArchetypeEntityCollection& EntityCollection);
+	MASSENTITY_API void RunWithCollection(FMassRuntimePipeline& RuntimePipeline, FProcessingContext& ProcessingContext, const FMassArchetypeEntityCollection& EntityCollection);
+
+	UE_DEPRECATED(5.8, "Use RunWithCollection instead.")
+	inline void RunSparse(FMassRuntimePipeline& RuntimePipeline, FProcessingContext& ProcessingContext, const FMassArchetypeEntityCollection& EntityCollection)
+	{
+		RunWithCollection(RuntimePipeline, ProcessingContext, EntityCollection);
+	}
 
 	/** Executes given Processors array view. This function gets called under the hood by the rest of Run* functions */
 	MASSENTITY_API void RunProcessorsView(TArrayView<UMassProcessor* const> Processors, FProcessingContext& ProcessingContext, TConstArrayView<FMassArchetypeEntityCollection> EntityCollections = {});
