[Mass] Print padding info in archetype debug output per [REDACTED_HANDLE] in MassArchetypeData.cpp

#Synced-CL 52096882

diff --git a/Engine/Source/Runtime/MassEntity/Private/MassArchetypeData.cpp b/Engine/Source/Runtime/MassEntity/Private/MassArchetypeData.cpp
--- a/Engine/Source/Runtime/MassEntity/Private/MassArchetypeData.cpp
+++ b/Engine/Source/Runtime/MassEntity/Private/MassArchetypeData.cpp
@@ -1364,15 +1364,22 @@ void FMassArchetypeData::DebugPrintArchetype(FOutputDevice& Ar)
 	Ar.Logf(ELogVerbosity::Log, TEXT("\tEntities / Chunk: %d"), NumEntitiesPerChunk);
 
 	Ar.Logf(ELogVerbosity::Log, TEXT("\tOffset 0x%04X: Entity[] (%llu bytes each)"), EntityListOffsetWithinChunk, sizeof(FMassEntityHandle));
+	int32 NextExpectedOffset = EntityListOffsetWithinChunk + sizeof(FMassEntityHandle) * NumEntitiesPerChunk;
+	int32 TotalPaddingBytes = 0;
 	int32 TotalBytesOfValidData = sizeof(FMassEntityHandle) * NumEntitiesPerChunk;
 	for (const FMassArchetypeFragmentConfig& FragmentConfig : FragmentConfigs)
 	{
+		const int32 PaddingBeforeFragment = FragmentConfig.ArrayOffsetWithinChunk - NextExpectedOffset;
+		if (PaddingBeforeFragment > 0)
+		{
+			Ar.Logf(ELogVerbosity::Log, TEXT("\tOffset 0x%04X: <Padding> (%d bytes)"), NextExpectedOffset, PaddingBeforeFragment);
+			TotalPaddingBytes += PaddingBeforeFragment;
+		}
 		TotalBytesOfValidData += FragmentConfig.FragmentType->GetStructureSize() * NumEntitiesPerChunk;
 		Ar.Logf(ELogVerbosity::Log, TEXT("\tOffset 0x%04X: %s[] (%d bytes each)"), FragmentConfig.ArrayOffsetWithinChunk, *FragmentConfig.FragmentType->GetName(), FragmentConfig.FragmentType->GetStructureSize());
+		NextExpectedOffset = FragmentConfig.ArrayOffsetWithinChunk + FragmentConfig.FragmentType->GetStructureSize() * NumEntitiesPerChunk;
 	}
 
-	//@TODO: Print out padding in between things?
-
 	const SIZE_T UnusablePaddingOffset = TotalBytesPerEntity * NumEntitiesPerChunk;
 	const SIZE_T UnusablePaddingAmount = GetChunkAllocSize() - UnusablePaddingOffset;
 	if (UnusablePaddingAmount > 0)
@@ -1380,9 +1387,9 @@ void FMassArchetypeData::DebugPrintArchetype(FOutputDevice& Ar)
 		Ar.Logf(ELogVerbosity::Log, TEXT("\tOffset 0x%04llX: WastePadding[] (%llu bytes total)"), UnusablePaddingOffset, UnusablePaddingAmount);
 	}
 
-	if (GetChunkAllocSize() != TotalBytesOfValidData + UnusablePaddingAmount)
+	if (TotalPaddingBytes > 0)
 	{
-		Ar.Logf(ELogVerbosity::Log, TEXT("\t@TODO: EXTRA PADDING HERE:  TotalBytesOfValidData: %d (%llu missing)"), TotalBytesOfValidData, GetChunkAllocSize() - TotalBytesOfValidData);
+		Ar.Logf(ELogVerbosity::Log, TEXT("\tTotal alignment padding between fragments: %d bytes"), TotalPaddingBytes);
 	}
 }
 
