[Mass]
Added logging to verify Mass processing queue usage in MassProcessingPhaseManager

#Synced-CL 52159388

diff --git a/Engine/Source/Runtime/MassEntity/Private/MassProcessingPhaseManager.cpp b/Engine/Source/Runtime/MassEntity/Private/MassProcessingPhaseManager.cpp
--- a/Engine/Source/Runtime/MassEntity/Private/MassProcessingPhaseManager.cpp
+++ b/Engine/Source/Runtime/MassEntity/Private/MassProcessingPhaseManager.cpp
@@ -323,6 +323,8 @@ void FMassProcessingPhaseManager::Initialize(UObject& InOwner, TConstArrayView<F
 	}
 
 	bIsAllowedToTick = true;
+
+	UE_LOG(LogMass, Warning, TEXT("FMassProcessingPhaseManager::Initialize complete. bUseProcessingQueue=%s"), UE::Mass::Tweakables::bUseProcessingQueue ? TEXT("true") : TEXT("false"));
 }
 
 void FMassProcessingPhaseManager::Deinitialize()
@@ -594,7 +596,13 @@ void FMassProcessingPhaseManager::OnPhaseStart(FMassProcessingPhase& Phase)
 			if (UE::Mass::Tweakables::bUseProcessingQueue)
 			{
 				UE::Mass::FProcessingQueue& PhaseQueue = ProcessingPhases[PhaseAsInt].ProcessingQueue;
-				PhaseQueue.SetProcessors(PhaseProcessor->GetChildProcessorsView());
+				TConstArrayView<UMassProcessor*> Processors = PhaseProcessor->GetChildProcessorsView();
+				UE_LOG(LogMass, Warning, TEXT("FMassProcessingPhaseManager::OnPhaseStart SetProcessors: Phase=%d ProcessorCount=%d"), PhaseAsInt, Processors.Num());
+				for (int32 i = 0; i < Processors.Num(); ++i)
+				{
+					UE_LOG(LogMass, Warning, TEXT("  [%d] %s"), i, Processors[i] ? *Processors[i]->GetName() : TEXT("null"));
+				}
+				PhaseQueue.SetProcessors(Processors);
 			}
 
 #if WITH_MASSENTITY_DEBUG
