[Mass]
Changed ProcessingPhase init log level from Warning to Log

#Synced-CL 52159388

diff --git a/Engine/Source/Runtime/MassEntity/Private/MassProcessingPhaseManager.cpp b/Engine/Source/Runtime/MassEntity/Private/MassProcessingPhaseManager.cpp
--- a/Engine/Source/Runtime/MassEntity/Private/MassProcessingPhaseManager.cpp
+++ b/Engine/Source/Runtime/MassEntity/Private/MassProcessingPhaseManager.cpp
@@ -324,7 +324,7 @@ void FMassProcessingPhaseManager::Initialize(UObject& InOwner, TConstArrayView<F
 
 	bIsAllowedToTick = true;
 
-	UE_LOG(LogMass, Warning, TEXT("FMassProcessingPhaseManager::Initialize complete. bUseProcessingQueue=%s"), UE::Mass::Tweakables::bUseProcessingQueue ? TEXT("true") : TEXT("false"));
+	UE_LOG(LogMass, Log, TEXT("FMassProcessingPhaseManager::Initialize complete. bUseProcessingQueue=%s"), UE::Mass::Tweakables::bUseProcessingQueue ? TEXT("true") : TEXT("false"));
 }
 
 void FMassProcessingPhaseManager::Deinitialize()
@@ -597,10 +597,10 @@ void FMassProcessingPhaseManager::OnPhaseStart(FMassProcessingPhase& Phase)
 			{
 				UE::Mass::FProcessingQueue& PhaseQueue = ProcessingPhases[PhaseAsInt].ProcessingQueue;
 				TConstArrayView<UMassProcessor*> Processors = PhaseProcessor->GetChildProcessorsView();
-				UE_LOG(LogMass, Warning, TEXT("FMassProcessingPhaseManager::OnPhaseStart SetProcessors: Phase=%d ProcessorCount=%d"), PhaseAsInt, Processors.Num());
+				UE_LOG(LogMass, Log, TEXT("FMassProcessingPhaseManager::OnPhaseStart SetProcessors: Phase=%d ProcessorCount=%d"), PhaseAsInt, Processors.Num());
 				for (int32 i = 0; i < Processors.Num(); ++i)
 				{
-					UE_LOG(LogMass, Warning, TEXT("  [%d] %s"), i, Processors[i] ? *Processors[i]->GetName() : TEXT("null"));
+					UE_LOG(LogMass, Log, TEXT("  [%d] %s"), i, Processors[i] ? *Processors[i]->GetName() : TEXT("null"));
 				}
 				PhaseQueue.SetProcessors(Processors);
 			}
