[MassTrace]
Add Sparse, ConstShared, and SparseTag fragment types to MassInsights and MassEntity tracing
Include shared and sparse elements in archetype registration trace output
Check derived types before base types in fragment type classification to avoid incorrect matches

#Synced-CL 52215450

diff --git a/Engine/Plugins/MassInsights/Source/MassInsightsAnalysis/Public/MassInsightsAnalysis/Model/MassInsights.h b/Engine/Plugins/MassInsights/Source/MassInsightsAnalysis/Public/MassInsightsAnalysis/Model/MassInsights.h
--- a/Engine/Plugins/MassInsights/Source/MassInsightsAnalysis/Public/MassInsightsAnalysis/Model/MassInsights.h
+++ b/Engine/Plugins/MassInsights/Source/MassInsightsAnalysis/Public/MassInsightsAnalysis/Model/MassInsights.h
@@ -49,7 +49,10 @@ enum class EFragmentType : uint8
 	Unknown,
 	Fragment,
 	Tag,
-	Shared
+	Shared,
+	Sparse,
+	ConstShared,
+	SparseTag
 };
 
 struct FMassFragmentInfo
diff --git a/Engine/Plugins/MassInsights/Source/MassInsightsUI/Private/MassInsightsUI/Widgets/SArchetypeDetails.cpp b/Engine/Plugins/MassInsights/Source/MassInsightsUI/Private/MassInsightsUI/Widgets/SArchetypeDetails.cpp
--- a/Engine/Plugins/MassInsights/Source/MassInsightsUI/Private/MassInsightsUI/Widgets/SArchetypeDetails.cpp
+++ b/Engine/Plugins/MassInsights/Source/MassInsightsUI/Private/MassInsightsUI/Widgets/SArchetypeDetails.cpp
@@ -58,6 +58,9 @@ namespace MassInsightsUI
 						case MassInsightsAnalysis::EFragmentType::Fragment: return LOCTEXT("Fragment", "Fragment");
 						case MassInsightsAnalysis::EFragmentType::Tag: return LOCTEXT("Tag", "Tag");
 						case MassInsightsAnalysis::EFragmentType::Shared: return LOCTEXT("Shared", "Shared");
+						case MassInsightsAnalysis::EFragmentType::Sparse: return LOCTEXT("Sparse", "Sparse");
+						case MassInsightsAnalysis::EFragmentType::ConstShared: return LOCTEXT("ConstShared", "Const Shared");
+						case MassInsightsAnalysis::EFragmentType::SparseTag: return LOCTEXT("SparseTag", "Sparse Tag");
 						case MassInsightsAnalysis::EFragmentType::Unknown: // Fallthrough
 						default:
 							return LOCTEXT("Unknown", "Unknown");
@@ -69,6 +72,9 @@ namespace MassInsightsUI
 					switch (FragmentInfo->Type)
 					{
 						case MassInsightsAnalysis::EFragmentType::Fragment: return FColor(10, 120, 180);
+						case MassInsightsAnalysis::EFragmentType::Sparse: return FColor(10, 160, 140);
+						case MassInsightsAnalysis::EFragmentType::ConstShared: return FColor(140, 80, 180);
+						case MassInsightsAnalysis::EFragmentType::SparseTag: return FColor(120, 140, 10);
 						case MassInsightsAnalysis::EFragmentType::Tag: // Fallthrough
 						case MassInsightsAnalysis::EFragmentType::Shared: // Fallthrough
 						case MassInsightsAnalysis::EFragmentType::Unknown: // Fallthrough
diff --git a/Engine/Source/Runtime/MassEntity/Private/MassEntityTrace.cpp b/Engine/Source/Runtime/MassEntity/Private/MassEntityTrace.cpp
--- a/Engine/Source/Runtime/MassEntity/Private/MassEntityTrace.cpp
+++ b/Engine/Source/Runtime/MassEntity/Private/MassEntityTrace.cpp
@@ -120,7 +120,10 @@ enum class EFragmentType : uint8
 	Unknown = 0,
 	Fragment,
 	Tag,
-	SharedFragment
+	SharedFragment,
+	Sparse,
+	ConstSharedFragment,
+	SparseTag
 };
 
 //----------------------------------------------------------------//
@@ -165,9 +168,6 @@ void FMassTrace::QueryArchetypeAdded(const FMassEntityQuery* Query, const FMassA
 //----------------------------------------------------------------//
 FMassTrace::FScopedQueryForEachTrace::FScopedQueryForEachTrace(const FMassEntityQuery* InQuery)
 	: Query(InQuery)
-	, ArchetypeCount(0)
-	, ChunkCount(0)
-	, EntityCount(0)
 {
 	UE_TRACE_LOG(MassTrace, QueryForEachStarted, MassChannel)
 		<< QueryForEachStarted.Cycle(FPlatformTime::Cycles64())
@@ -196,21 +196,34 @@ void FMassTrace::FScopedQueryForEachTrace::ReportArchetype(const FMassArchetypeD
 
 void FMassTrace::OutputRegisterFragment(const UScriptStruct* Struct)
 {
-	const uint64 FragmentId = reinterpret_cast<uint64>(Struct);
-	const EFragmentType FragmentType = [](const UScriptStruct* Struct)
+	// Determine the element type for tracing. Checks derived types before base types
+	// since FMassSparseFragment/FMassSparseTag inherit from FMassFragment/FMassTag.
+	const EFragmentType FragmentType = [](const UScriptStruct* Struct) -> EFragmentType
 		{
-			if (Struct->IsChildOf<FMassFragment>())
+			if (Struct->IsChildOf<FMassSparseFragment>())
 			{
-				return EFragmentType::Fragment;
+				return EFragmentType::Sparse;
 			}
-			else if (Struct->IsChildOf<FMassTag>())
+			if (Struct->IsChildOf<FMassConstSharedFragment>())
 			{
-				return EFragmentType::Tag;
+				return EFragmentType::ConstSharedFragment;
 			}
 			else if (Struct->IsChildOf<FMassSharedFragment>())
 			{
 				return EFragmentType::SharedFragment;
 			}
+			else if (Struct->IsChildOf<FMassFragment>())
+			{
+				return EFragmentType::Fragment;
+			}
+			else if (Struct->IsChildOf<FMassSparseTag>())
+			{
+				return EFragmentType::SparseTag;
+			}
+			else if (Struct->IsChildOf<FMassTag>())
+			{
+				return EFragmentType::Tag;
+			}
 			else
 			{
 				return EFragmentType::Unknown;
@@ -218,7 +231,7 @@ void FMassTrace::OutputRegisterFragment(const UScriptStruct* Struct)
 		}(Struct);
 
 	UE_TRACE_LOG(MassTrace, RegisterMassFragment, MassChannel)
-		<< RegisterMassFragment.FragmentId(FragmentId)
+		<< RegisterMassFragment.FragmentId(reinterpret_cast<uint64>(Struct))
 		<< RegisterMassFragment.FragmentName(*Struct->GetName())
 		<< RegisterMassFragment.FragmentSize(Struct->GetStructureSize())
 		<< RegisterMassFragment.FragmentType(static_cast<uint8>(FragmentType));
@@ -245,20 +258,28 @@ void FMassTrace::OnPhaseEnd(uint64 PhaseId)
 uint64 FMassTrace::OutputRegisterArchetype(uint64 ArchetypeID, const FMassArchetypeCompositionDescriptor& CompositionDescriptor)
 {
 	TArray<uint64> FragmentsScratch;
-	FMassElementBitSet FragmentsAndTags = CompositionDescriptor.GetElementsBitSet().GetFragmentsAndTags();
-	FragmentsScratch.Reserve(FragmentsAndTags.CountStoredTypes());
+	const FMassElementBitSet& AllElements = CompositionDescriptor.GetElementsBitSet();
 
-	auto FragmentIterator = FragmentsAndTags.GetIndexIterator();
-	while (FragmentIterator)
-	{
-		const UScriptStruct* FragmentStruct = FragmentsAndTags.GetTypeAtIndex(*FragmentIterator);
+	// Include fragments, tags, shared fragments, AND sparse elements
+	const FMassElementBitSet FragmentsAndTags = AllElements.GetFragmentsAndTags();
+	const FMassElementBitSet SharedFragments = AllElements.GetSharedFragments();
+	const FMassElementBitSet SparseElements = AllElements.GetSparseElements();
 
-		OutputRegisterFragment(FragmentStruct);
+	FragmentsScratch.Reserve(FragmentsAndTags.CountStoredTypes() + SharedFragments.CountStoredTypes() + SparseElements.CountStoredTypes());
 
-		// TODO Should have utility function to go from UScriptStruct to the fragment ID
-		FragmentsScratch.Add(reinterpret_cast<uint64>(FragmentStruct));
-		++FragmentIterator;
-	}
+	auto RegisterElements = [&FragmentsScratch](const FMassElementBitSet& Elements)
+	{
+		for (auto It = Elements.GetIndexIterator(); It; ++It)
+		{
+			const UScriptStruct* FragmentStruct = Elements.GetTypeAtIndex(*It);
+			OutputRegisterFragment(FragmentStruct);
+			FragmentsScratch.Add(reinterpret_cast<uint64>(FragmentStruct));
+		}
+	};
+
+	RegisterElements(FragmentsAndTags);
+	RegisterElements(SharedFragments);
+	RegisterElements(SparseElements);
 
 	UE_TRACE_LOG(MassTrace, RegisterMassArchetype, MassChannel)
 		<< RegisterMassArchetype.ArchetypeID(ArchetypeID)
