[Mass]
Force inline processing queue execution for non-parallel mass configurations

#Synced-CL 52226149

diff --git a/Engine/Source/Runtime/MassEntity/Private/MassProcessingQueue.cpp b/Engine/Source/Runtime/MassEntity/Private/MassProcessingQueue.cpp
--- a/Engine/Source/Runtime/MassEntity/Private/MassProcessingQueue.cpp
+++ b/Engine/Source/Runtime/MassEntity/Private/MassProcessingQueue.cpp
@@ -20,13 +20,22 @@
 
 namespace UE::Mass::Tweakables
 {
+#if MASS_DO_PARALLEL
 	int32 MaxProcessorWorkerCount = 8;
 	int32 MaxChunkWorkerCount = 16;
-
+#else
+	// Non-parallel builds (e.g. dedicated server): all processors run inline on the calling thread.
+	int32 MaxProcessorWorkerCount = 0;
+	int32 MaxChunkWorkerCount = 1;
+#endif
+
+#if MASS_DO_PARALLEL
+	// Configuration is only applicable to parallel builds
 	FAutoConsoleVariableRef ProcessingQueueCVars[] = {
 		{TEXT("mass.ProcessingQueue.MaxProcessorWorkers"), MaxProcessorWorkerCount, TEXT("Max number of concurrent Processor Workers")},
 		{TEXT("mass.ProcessingQueue.MaxChunkWorkers"), MaxChunkWorkerCount, TEXT("Max number of concurrent Chunk Workers")},
 	};
+#endif
 } // namespace UE::Mass::Tweakables
 
 namespace UE::Mass
@@ -101,6 +110,7 @@ namespace UE::Mass
 		{
 			FMassEntityManager::FScopedProcessing ProcessingScope = EntityManager.NewProcessingScope();
 
+#if MASS_DO_PARALLEL
 			UE::Tasks::FTask QueueRunner = (UE::Tasks::Launch(TEXT("MassProcessingQueue Runner Task"),
 				[this, &ExecutionContext]()
 				{
@@ -133,6 +143,11 @@ namespace UE::Mass
 			{
 				checkf(false, TEXT("BlockMainThread = false is not yet supported"));
 			}
+#else
+			TRACE_CPUPROFILER_EVENT_SCOPE_TEXT("MassProcessingQueue Inline Runner");
+			RunMainThread(ExecutionContext);
+			check(CompletedProcessorCount.load() == ProcessorQueue.Num());
+#endif
 		}
 
 		// There is a very small chance that the entity manager bookkeeping for the ProcessingScope count d
