[Mass]
Prevent mass processor config from stomping on hard-coded execution ordering directives, merge them instead

#Synced-CL 52277523

diff --git a/Engine/Source/Runtime/MassEntity/Private/MassProcessor.cpp b/Engine/Source/Runtime/MassEntity/Private/MassProcessor.cpp
--- a/Engine/Source/Runtime/MassEntity/Private/MassProcessor.cpp
+++ b/Engine/Source/Runtime/MassEntity/Private/MassProcessor.cpp
@@ -302,8 +302,22 @@ bool UMassProcessor::DoesAnyArchetypeMatchOwnedQueries(const FMassEntityManager&
 
 void UMassProcessor::PostInitProperties()
 {
+	// Capture constructor-set hard deps before Super::PostInitProperties calls LoadConfig
+	TArray<FName> HardcodedExecuteAfter = ExecutionOrder.ExecuteAfter;
+	TArray<FName> HardcodedExecuteBefore = ExecutionOrder.ExecuteBefore;
+
 	Super::PostInitProperties();
 
+	// Re-merge hardcoded deps
+	for (const FName& Name : HardcodedExecuteAfter)
+	{
+		ExecutionOrder.ExecuteAfter.AddUnique(Name);
+	}
+	for (const FName& Name : HardcodedExecuteBefore)
+	{
+		ExecutionOrder.ExecuteBefore.AddUnique(Name);
+	}
+
 #if CPUPROFILERTRACE_ENABLED
 	StatId = GetProcessorName();
 #endif
