[MassCore] added deprecated forwarding stubs for headers moved out of MassEntity

adding stub headers at the old MassEntity/Public/ paths so licensee code that still includes e.g. "MassEntityHandle.h" gets a deprecation warning pointing them to the new "Mass/EntityHandle.h" in MassCore, instead of a hard compile error.

#rnx
#Synced-CL 52342139

diff --git a/Engine/Source/Runtime/MassEntity/Public/MassArchetypeGroup.h b/Engine/Source/Runtime/MassEntity/Public/MassArchetypeGroup.h
new file mode 100644
--- /dev/null
+++ b/Engine/Source/Runtime/MassEntity/Public/MassArchetypeGroup.h
@@ -0,0 +1,7 @@
+// Copyright Epic Games, Inc. All Rights Reserved.
+#pragma once
+// HEADER_UNIT_SKIP - Deprecated
+#include "Misc/CoreMiscDefines.h"
+UE_DEPRECATED_HEADER(5.8, "MassArchetypeGroup.h has moved to the MassCore module and been renamed to Mass/ArchetypeGroup.h. "
+    "Include Mass/ArchetypeGroup.h from MassCore and add a MassCore dependency to your .Build.cs.")
+#include "Mass/ArchetypeGroup.h"
diff --git a/Engine/Source/Runtime/MassEntity/Public/MassEntityElementTypes.h b/Engine/Source/Runtime/MassEntity/Public/MassEntityElementTypes.h
new file mode 100644
--- /dev/null
+++ b/Engine/Source/Runtime/MassEntity/Public/MassEntityElementTypes.h
@@ -0,0 +1,7 @@
+// Copyright Epic Games, Inc. All Rights Reserved.
+#pragma once
+// HEADER_UNIT_SKIP - Deprecated
+#include "Misc/CoreMiscDefines.h"
+UE_DEPRECATED_HEADER(5.8, "MassEntityElementTypes.h has moved to the MassCore module and been renamed to Mass/EntityElementTypes.h. "
+    "Include Mass/EntityElementTypes.h from MassCore and add a MassCore dependency to your .Build.cs.")
+#include "Mass/EntityElementTypes.h"
diff --git a/Engine/Source/Runtime/MassEntity/Public/MassEntityFragments.h b/Engine/Source/Runtime/MassEntity/Public/MassEntityFragments.h
new file mode 100644
--- /dev/null
+++ b/Engine/Source/Runtime/MassEntity/Public/MassEntityFragments.h
@@ -0,0 +1,7 @@
+// Copyright Epic Games, Inc. All Rights Reserved.
+#pragma once
+// HEADER_UNIT_SKIP - Deprecated
+#include "Misc/CoreMiscDefines.h"
+UE_DEPRECATED_HEADER(5.8, "MassEntityFragments.h has moved to the MassCore module and been renamed to Mass/EntityFragments.h. "
+    "Include Mass/EntityFragments.h from MassCore and add a MassCore dependency to your .Build.cs.")
+#include "Mass/EntityFragments.h"
diff --git a/Engine/Source/Runtime/MassEntity/Public/MassEntityHandle.h b/Engine/Source/Runtime/MassEntity/Public/MassEntityHandle.h
new file mode 100644
--- /dev/null
+++ b/Engine/Source/Runtime/MassEntity/Public/MassEntityHandle.h
@@ -0,0 +1,7 @@
+// Copyright Epic Games, Inc. All Rights Reserved.
+#pragma once
+// HEADER_UNIT_SKIP - Deprecated
+#include "Misc/CoreMiscDefines.h"
+UE_DEPRECATED_HEADER(5.8, "MassEntityHandle.h has moved to the MassCore module and been renamed to Mass/EntityHandle.h. "
+    "Include Mass/EntityHandle.h from MassCore and add a MassCore dependency to your .Build.cs.")
+#include "Mass/EntityHandle.h"
diff --git a/Engine/Source/Runtime/MassEntity/Public/MassEntityMacros.h b/Engine/Source/Runtime/MassEntity/Public/MassEntityMacros.h
new file mode 100644
--- /dev/null
+++ b/Engine/Source/Runtime/MassEntity/Public/MassEntityMacros.h
@@ -0,0 +1,7 @@
+// Copyright Epic Games, Inc. All Rights Reserved.
+#pragma once
+// HEADER_UNIT_SKIP - Deprecated
+#include "Misc/CoreMiscDefines.h"
+UE_DEPRECATED_HEADER(5.8, "MassEntityMacros.h has moved to the MassCore module and been renamed to Mass/EntityMacros.h. "
+    "Include Mass/EntityMacros.h from MassCore and add a MassCore dependency to your .Build.cs.")
+#include "Mass/EntityMacros.h"
diff --git a/Engine/Source/Runtime/MassEntity/Public/MassExternalSubsystemTraits.h b/Engine/Source/Runtime/MassEntity/Public/MassExternalSubsystemTraits.h
new file mode 100644
--- /dev/null
+++ b/Engine/Source/Runtime/MassEntity/Public/MassExternalSubsystemTraits.h
@@ -0,0 +1,7 @@
+// Copyright Epic Games, Inc. All Rights Reserved.
+#pragma once
+// HEADER_UNIT_SKIP - Deprecated
+#include "Misc/CoreMiscDefines.h"
+UE_DEPRECATED_HEADER(5.8, "MassExternalSubsystemTraits.h has moved to the MassCore module and been renamed to Mass/ExternalSubsystemTraits.h. "
+    "Include Mass/ExternalSubsystemTraits.h from MassCore and add a MassCore dependency to your .Build.cs.")
+#include "Mass/ExternalSubsystemTraits.h"
diff --git a/Engine/Source/Runtime/MassEntity/Public/MassTestableEnsures.h b/Engine/Source/Runtime/MassEntity/Public/MassTestableEnsures.h
new file mode 100644
--- /dev/null
+++ b/Engine/Source/Runtime/MassEntity/Public/MassTestableEnsures.h
@@ -0,0 +1,7 @@
+// Copyright Epic Games, Inc. All Rights Reserved.
+#pragma once
+// HEADER_UNIT_SKIP - Deprecated
+#include "Misc/CoreMiscDefines.h"
+UE_DEPRECATED_HEADER(5.8, "MassTestableEnsures.h has moved to the MassCore module and been renamed to Mass/TestableEnsures.h. "
+    "Include Mass/TestableEnsures.h from MassCore and add a MassCore dependency to your .Build.cs.")
+#include "Mass/TestableEnsures.h"
