[Mass]
Enable mass.UseProcessingQueue by default

#Synced-CL 52277523

diff --git a/Engine/Source/Runtime/MassEntity/Private/MassProcessingPhaseManager.cpp b/Engine/Source/Runtime/MassEntity/Private/MassProcessingPhaseManager.cpp
--- a/Engine/Source/Runtime/MassEntity/Private/MassProcessingPhaseManager.cpp
+++ b/Engine/Source/Runtime/MassEntity/Private/MassProcessingPhaseManager.cpp
@@ -26,7 +26,7 @@ namespace UE::Mass::Tweakables
 {
 	bool bFullyParallel = MASS_DO_PARALLEL;
 	bool bMakePrePhysicsTickFunctionHighPriority = true;
-	bool bUseProcessingQueue = false;
+	bool bUseProcessingQueue = true;
 
 	FAutoConsoleVariableRef CVars[] = {
 		{TEXT("mass.FullyParallel"), bFullyParallel, TEXT("Enables mass processing distribution to all available thread (via the task graph)")},
