[MassTrace]
Add missing trace events for archetype creation and entity moves during archetype splitting in MassEntity
Introduce a batch UE_TRACE_MASS_ENTITIES_MOVED macro to trace multiple entity handles in a single call

#Synced-CL 52361779

diff --git a/Engine/Source/Runtime/MassEntity/Private/MassArchetypeData.cpp b/Engine/Source/Runtime/MassEntity/Private/MassArchetypeData.cpp
--- a/Engine/Source/Runtime/MassEntity/Private/MassArchetypeData.cpp
+++ b/Engine/Source/Runtime/MassEntity/Private/MassArchetypeData.cpp
@@ -1576,9 +1576,12 @@ void FMassArchetypeData::BatchMoveEntitiesToAnotherArchetype(const FMassArchetyp
 
 	}
 
-	// Sorting the subchunks info so that subchunks of a given chunk are processed "from the back". Otherwise removing 
+	UE_TRACE_MASS_ARCHETYPE_CREATED(NewArchetype)
+	UE_TRACE_MASS_ENTITIES_MOVED(TConstArrayView<FMassEntityHandle>(OutEntitiesBeingMoved).Mid(InitialOutEntitiesCount), NewArchetype)
+
+	// Sorting the subchunks info so that subchunks of a given chunk are processed "from the back". Otherwise removing
 	// a subchunk from the front of the chunk would inevitably invalidate following subchunks' information.
-	// Note that we do this after already having added the entities to the new archetype to preserve the order of entities 
+	// Note that we do this after already having added the entities to the new archetype to preserve the order of entities
 	// as given by the input data.
 	Subchunks.Sort([](const FMassArchetypeEntityCollection::FArchetypeEntityRange& A, const FMassArchetypeEntityCollection::FArchetypeEntityRange& B)
 	{
diff --git a/Engine/Source/Runtime/MassEntity/Public/MassEntityTrace.h b/Engine/Source/Runtime/MassEntity/Public/MassEntityTrace.h
--- a/Engine/Source/Runtime/MassEntity/Public/MassEntityTrace.h
+++ b/Engine/Source/Runtime/MassEntity/Public/MassEntityTrace.h
@@ -85,6 +85,12 @@ public:
 #define UE_TRACE_MASS_ENTITY_MOVED(Entity, NewArchetype) \
 	FMassTrace::EntityMoved(Entity, NewArchetype);
 
+#define UE_TRACE_MASS_ENTITIES_MOVED(EntityHandles, NewArchetype) \
+	for (const FMassEntityHandle& Entity : EntityHandles) \
+	{ \
+		FMassTrace::EntityMoved(Entity, NewArchetype); \
+	}
+
 #define UE_TRACE_MASS_ENTITY_DESTROYED(Entity) \
 	FMassTrace::EntityDestroyed(Entity);
 
@@ -118,6 +124,7 @@ public:
 #define UE_TRACE_MASS_ENTITY_CREATED(Entity, Archetype)
 #define UE_TRACE_MASS_ENTITIES_CREATED(EntityHandles, Archetype)
 #define UE_TRACE_MASS_ENTITY_MOVED(Entity, NewArchetype)
+#define UE_TRACE_MASS_ENTITIES_MOVED(EntityHandles, NewArchetype)
 #define UE_TRACE_MASS_ENTITY_DESTROYED(Entity)
 #define UE_TRACE_MASS_ENTITIES_DESTROYED(EntityHandles)
 #define UE_TRACE_MASS_PHASE_BEGIN(PhaseID)
