Add static_assert checks to MassEntity command templates to enforce that type parameters derive from the expected base types (FMassFragment or FMassTag)
Clean up trailing whitespace in tag command constructors

#Synced-CL 52392538

diff --git a/Engine/Source/Runtime/MassEntity/Public/MassCommands.h b/Engine/Source/Runtime/MassEntity/Public/MassCommands.h
--- a/Engine/Source/Runtime/MassEntity/Public/MassCommands.h
+++ b/Engine/Source/Runtime/MassEntity/Public/MassCommands.h
@@ -252,6 +252,8 @@ protected:
 template<EMassCommandCheckTime CheckTime, typename... TTypes>
 struct FMassCommandAddFragmentsInternal : public FMassBatchedEntityCommand
 {
+	static_assert((TIsDerivedFrom<UE::Mass::Clean<TTypes>, FMassFragment>::Value && ...), "All types passed to FMassCommandAddFragments must derive from FMassFragment.");
+
 	using Super = FMassBatchedEntityCommand;
 	FMassCommandAddFragmentsInternal()
 		: Super(EMassCommandOperationType::Add DEBUG_NAME("AddFragments"))
@@ -283,6 +285,8 @@ using FMassCommandAddFragments = FMassCommandAddFragmentsInternal<EMassCommandCh
 template<EMassCommandCheckTime CheckTime, typename... TTypes>
 struct FMassCommandRemoveFragmentsInternal : public FMassBatchedEntityCommand
 {
+	static_assert((TIsDerivedFrom<UE::Mass::Clean<TTypes>, FMassFragment>::Value && ...), "All types passed to FMassCommandRemoveFragments must derive from FMassFragment.");
+
 	using Super = FMassBatchedEntityCommand;
 	FMassCommandRemoveFragmentsInternal()
 		: Super(EMassCommandOperationType::Remove DEBUG_NAME("RemoveFragments"))
@@ -349,12 +353,14 @@ protected:
 template<EMassCommandCheckTime CheckTime, typename... TTypes>
 struct FMassCommandAddTagsInternal : public FMassCommandChangeTags
 {
+	static_assert((TIsDerivedFrom<UE::Mass::Clean<TTypes>, FMassTag>::Value && ...), "All types passed to FMassCommandAddTags must derive from FMassTag.");
+
 	using Super = FMassCommandChangeTags;
 	FMassCommandAddTagsInternal()
 		: Super(
-			EMassCommandOperationType::Add, 
+			EMassCommandOperationType::Add,
 			UE::Mass::Utils::ConstructTagBitSet<CheckTime, TTypes...>(),
-			{} 
+			{}
 			DEBUG_NAME("AddTags"))
 	{}
 };
@@ -368,11 +374,13 @@ using FMassCommandAddTags = FMassCommandAddTagsInternal<EMassCommandCheckTime::C
 template<EMassCommandCheckTime CheckTime, typename... TTypes>
 struct FMassCommandRemoveTagsInternal : public FMassCommandChangeTags
 {
+	static_assert((TIsDerivedFrom<UE::Mass::Clean<TTypes>, FMassTag>::Value && ...), "All types passed to FMassCommandRemoveTags must derive from FMassTag.");
+
 	using Super = FMassCommandChangeTags;
 	FMassCommandRemoveTagsInternal()
 		: Super(
-			EMassCommandOperationType::Remove, 
-			{}, 
+			EMassCommandOperationType::Remove,
+			{},
 			UE::Mass::Utils::ConstructTagBitSet<CheckTime, TTypes...>()
 			DEBUG_NAME("RemoveTags"))
 	{}
@@ -395,6 +403,8 @@ using FMassCommandRemoveTags = FMassCommandRemoveTagsInternal<EMassCommandCheckT
 template<EMassCommandCheckTime CheckTime, typename TOld, typename TNew>
 struct FMassCommandSwapTagsInternal : public FMassCommandChangeTags
 {
+	static_assert(TIsDerivedFrom<UE::Mass::Clean<TOld>, FMassTag>::Value && TIsDerivedFrom<UE::Mass::Clean<TNew>, FMassTag>::Value, "Both types passed to FMassCommandSwapTags must derive from FMassTag.");
+
 	using Super = FMassCommandChangeTags;
 	FMassCommandSwapTagsInternal()
 		: Super(
@@ -414,6 +424,8 @@ using FMassCommandSwapTags = FMassCommandSwapTagsInternal<EMassCommandCheckTime:
 template<typename... TOthers>
 struct FMassCommandAddFragmentInstances : public FMassBatchedEntityCommand
 {
+	static_assert((TIsDerivedFrom<UE::Mass::Clean<TOthers>, FMassFragment>::Value && ...), "All types passed to FMassCommandAddFragmentInstances must derive from FMassFragment.");
+
 	using Super = FMassBatchedEntityCommand;
 
 	FMassCommandAddFragmentInstances(EMassCommandOperationType OperationType = EMassCommandOperationType::Set DEBUG_NAME_PARAM("AddFragmentInstanceList"))
@@ -577,6 +589,8 @@ protected:
 template<typename TSharedFragmentValues, typename... TOthers>
 struct FMassCommandBuildEntityWithSharedFragments : public FMassBatchedCommand
 {
+	static_assert((TIsDerivedFrom<UE::Mass::Clean<TOthers>, FMassFragment>::Value && ...), "All fragment types passed to FMassCommandBuildEntityWithSharedFragments must derive from FMassFragment.");
+
 	using Super = FMassBatchedCommand;
 
 	FMassCommandBuildEntityWithSharedFragments()
