[MassObserver]
Add FObjectInitializer constructor overload to UMassObserverProcessor and extract shared initialization into a private InitObserverDefaults method

#Synced-CL 52443815

diff --git a/Engine/Source/Runtime/MassEntity/Private/MassObserverProcessor.cpp b/Engine/Source/Runtime/MassEntity/Private/MassObserverProcessor.cpp
--- a/Engine/Source/Runtime/MassEntity/Private/MassObserverProcessor.cpp
+++ b/Engine/Source/Runtime/MassEntity/Private/MassObserverProcessor.cpp
@@ -9,6 +9,17 @@
 // UMassObserverProcessor
 //----------------------------------------------------------------------//
 UMassObserverProcessor::UMassObserverProcessor()
+{
+	InitObserverDefaults();
+}
+
+UMassObserverProcessor::UMassObserverProcessor(const FObjectInitializer& ObjectInitializer)
+	: Super(ObjectInitializer)
+{
+	InitObserverDefaults();
+}
+
+void UMassObserverProcessor::InitObserverDefaults()
 {
 	bAutoRegisterWithProcessingPhases = false;
 #if WITH_EDITORONLY_DATA
diff --git a/Engine/Source/Runtime/MassEntity/Public/MassObserverProcessor.h b/Engine/Source/Runtime/MassEntity/Public/MassObserverProcessor.h
--- a/Engine/Source/Runtime/MassEntity/Public/MassObserverProcessor.h
+++ b/Engine/Source/Runtime/MassEntity/Public/MassObserverProcessor.h
@@ -85,6 +85,7 @@ class UMassObserverProcessor : public UMassProcessor
 
 public:
 	UE_API UMassObserverProcessor();
+	UE_API explicit UMassObserverProcessor(const FObjectInitializer& ObjectInitializer);
 
 	EMassObservedOperationFlags GetObservedOperations() const;
 	/**
@@ -106,6 +107,9 @@ protected:
 	 */
 	UE_API virtual void Register();
 
+private:
+	void InitObserverDefaults();
+
 protected:
 	/** Determines which Fragment or Tag type this given UMassObserverProcessor will be observing (single-type legacy). */
 	UE_DEPRECATED(5.8, "UMassObserverProcessor::ObservedType is deprecated. Use ObservedTypes instead")
