[Mass] Mass Relations Manager ForEachParticipant

Added callback-based enumeration for relation participants to elide TArray allocation when not necessary.
GetParticipants is now implemented in terms of ForEachParticipant.

#Synced-CL 52439714

diff --git a/Engine/Source/Runtime/MassEntity/Private/MassRelationManager.cpp b/Engine/Source/Runtime/MassEntity/Private/MassRelationManager.cpp
--- a/Engine/Source/Runtime/MassEntity/Private/MassRelationManager.cpp
+++ b/Engine/Source/Runtime/MassEntity/Private/MassRelationManager.cpp
@@ -79,27 +79,39 @@ namespace UE::Mass
 	}
 
 	TArray<FMassEntityHandle> FRelationData::GetParticipants(const FMassEntityManager& EntityManager, const FMassEntityHandle RoleEntity, const ERelationRole QueriedRole) const
+	{
+		TArray<FMassEntityHandle> ReturnSubjects;
+		if (const FRelationInstanceData* InstanceData = RoleMap.Find(RoleEntity))
+		{
+			ReturnSubjects.Reserve((*InstanceData)[QueriedRole].Num());
+		}
+		ForEachParticipant(EntityManager, RoleEntity, QueriedRole, [&ReturnSubjects](FMassEntityHandle Handle)
+		{
+			ReturnSubjects.Add(Handle);
+		});
+		return ReturnSubjects;
+	}
+
+	void FRelationData::ForEachParticipant(const FMassEntityManager& EntityManager, const FMassEntityHandle RoleEntity,
+		const ERelationRole QueriedRole, TFunctionRef<void(FMassEntityHandle)> Callback) const
 	{
 		const int32 QueriedRoleIndex = static_cast<int32>(QueriedRole);
 
-		ensureMsgf(Traits.RoleTraits[QueriedRoleIndex].RequiresExternalMapping == EExternalMappingRequired::Yes
-			, TEXT("Fetching relation participants for role %s, while role traits explicitly prevent generic mapping")
-			, *LexToString(QueriedRole));
+		ensureMsgf(Traits.RoleTraits[QueriedRoleIndex].RequiresExternalMapping == EExternalMappingRequired::Yes,
+			TEXT("Fetching relation participants for role %s, while role traits explicitly prevent generic mapping"),
+			*LexToString(QueriedRole));
 
-		TArray<FMassEntityHandle> ReturnSubjects;
 		if (const FRelationInstanceData* InstanceData = RoleMap.Find(RoleEntity))
 		{
-			ReturnSubjects.Reserve((*InstanceData)[QueriedRoleIndex].Num());
 			for (const FMassRelationRoleInstanceHandle RelationInstanceHandle : (*InstanceData)[QueriedRoleIndex])
 			{
 				const FMassEntityHandle RoleEntityHandle = RelationInstanceHandle.GetRoleEntityHandle(EntityManager);
 				if (RoleEntityHandle.IsValid())
 				{
-					ReturnSubjects.Add(RoleEntityHandle);
+					Callback(RoleEntityHandle);
 				}
 			}
 		}
-		return ReturnSubjects;
 	}
 
 	//-----------------------------------------------------------------------------
@@ -372,6 +384,20 @@ namespace UE::Mass
 		return RelationData.GetParticipants(EntityManager, SubjectEntity, ERelationRole::Object);
 	}
 
+	void FRelationManager::ForEachRelationSubject(const FTypeHandle RelationTypeHandle, const FMassEntityHandle ObjectEntity,
+		TFunctionRef<void(FMassEntityHandle)> Callback) const
+	{
+		const FRelationData& RelationData = GetRelationDataChecked(RelationTypeHandle);
+		RelationData.ForEachParticipant(EntityManager, ObjectEntity, ERelationRole::Subject, Callback);
+	}
+
+	void FRelationManager::ForEachRelationObject(const FTypeHandle RelationTypeHandle, const FMassEntityHandle SubjectEntity,
+		TFunctionRef<void(FMassEntityHandle)> Callback) const
+	{
+		const FRelationData& RelationData = GetRelationDataChecked(RelationTypeHandle);
+		RelationData.ForEachParticipant(EntityManager, SubjectEntity, ERelationRole::Object, Callback);
+	}
+
 	void FRelationManager::GatherHierarchy(const FRelationData& RelationData, const FMassEntityHandle SubjectHandle, FHierarchyEntitiesContainer& SubSubjects, const uint32 Depth) const
 	{
 		// getting subject's subjects
diff --git a/Engine/Source/Runtime/MassEntity/Public/MassRelationManager.h b/Engine/Source/Runtime/MassEntity/Public/MassRelationManager.h
--- a/Engine/Source/Runtime/MassEntity/Public/MassRelationManager.h
+++ b/Engine/Source/Runtime/MassEntity/Public/MassRelationManager.h
@@ -61,6 +61,10 @@ namespace UE::Mass
 		TMap<FMassEntityHandle, FRelationInstanceData> RoleMap;
 
 		TArray<FMassEntityHandle> GetParticipants(const FMassEntityManager& EntityManager, const FMassEntityHandle RoleEntity, ERelationRole QueriedRole) const;
+
+		/** Callback-based enumeration of participants. Same iteration as GetParticipants but invokes a callback per valid entity instead of building a TArray. */
+		void ForEachParticipant(const FMassEntityManager& EntityManager, FMassEntityHandle RoleEntity,
+			ERelationRole QueriedRole, TFunctionRef<void(FMassEntityHandle)> Callback) const;
 	};
 
 
@@ -100,6 +104,14 @@ namespace UE::Mass
 
 		FTypeHandle GetRelationTypeHandle(TNotNull<const UScriptStruct*> RelationType) const;
 
+		/** Callback-based enumeration of subjects for a given object in a relation. Zero allocation. */
+		UE_API void ForEachRelationSubject(const FTypeHandle RelationTypeHandle, FMassEntityHandle ObjectEntity,
+			TFunctionRef<void(FMassEntityHandle)> Callback) const;
+
+		/** Callback-based enumeration of objects for a given subject in a relation. Zero allocation. */
+		UE_API void ForEachRelationObject(const FTypeHandle RelationTypeHandle, FMassEntityHandle SubjectEntity,
+			TFunctionRef<void(FMassEntityHandle)> Callback) const;
+
 		UE_API bool IsSubjectOfRelation(const FTypeHandle RelationTypeHandle, const FMassEntityHandle Subject, const FMassEntityHandle Object) const;
 		UE_API bool IsSubjectOfRelation(const FRelationData& RelationDataInstance, const FMassEntityHandle Subject, const FMassEntityHandle Object) const;
 		UE_API bool IsSubjectOfRelationRecursive(const FTypeHandle RelationTypeHandle, const FMassEntityHandle Subject, const FMassEntityHandle Object) const;
@@ -225,6 +237,7 @@ namespace UE::Mass
 		const FRelationData& RelationData = GetRelationDataChecked(RelationTypeHandle);
 		return GetRelationObjects(RelationData, SubjectEntity);
 	}
+
 } // namespace UE::Mass
 
 #undef UE_API
