[Mass]
Disable mass.UseProcessingQueue temporarily

#Synced-CL 52277523

diff --git a/Engine/Source/Runtime/MassEntity/Private/MassProcessingPhaseManager.cpp b/Engine/Source/Runtime/MassEntity/Private/MassProcessingPhaseManager.cpp
--- a/Engine/Source/Runtime/MassEntity/Private/MassProcessingPhaseManager.cpp
+++ b/Engine/Source/Runtime/MassEntity/Private/MassProcessingPhaseManager.cpp
@@ -30,7 +30,7 @@ namespace UE::Mass::Tweakables
 {
 	bool bFullyParallel = MASS_DO_PARALLEL;
 	bool bMakePrePhysicsTickFunctionHighPriority = true;
-	bool bUseProcessingQueue = true;
+	bool bUseProcessingQueue = false;
 
 	FAutoConsoleVariableRef CVars[] = {
 		{TEXT("mass.FullyParallel"), bFullyParallel, TEXT("Enables mass processing distribution to all available thread (via the task graph)")},
