Remove tag-and-fragment-only validation checks from BatchChangeComposition in MassEntityManager
Allow shared and sparse element types to be passed through the batch composition change API

#Synced-CL 52486799

diff --git a/Engine/Source/Runtime/MassEntity/Private/MassEntityManager.cpp b/Engine/Source/Runtime/MassEntity/Private/MassEntityManager.cpp
--- a/Engine/Source/Runtime/MassEntity/Private/MassEntityManager.cpp
+++ b/Engine/Source/Runtime/MassEntity/Private/MassEntityManager.cpp
@@ -1886,15 +1886,6 @@ void FMassEntityManager::BatchChangeCompositionForEntities(TConstArrayView<FMass
 {
 	CHECK_SYNC_API();
 
-	testableCheckf(ElementsToAdd.ContainsOnlyTagsAndFragments()
-		, TEXT("%hs: ElementsToAdd contains shared or sparse element types."
-			" Only tags and fragments are supported. Use the dedicated shared fragment APIs instead.")
-		, __FUNCTION__);
-	testableCheckf(ElementsToRemove.ContainsOnlyTagsAndFragments()
-		, TEXT("%hs: ElementsToRemove contains shared or sparse element types."
-			" Only tags and fragments are supported. Use the dedicated shared fragment APIs instead.")
-		, __FUNCTION__);
-
 	TRACE_CPUPROFILER_EVENT_SCOPE(Mass_BatchChangeCompositionForEntities);
 
 	const bool bElementsAddedAreObserved = ObserverManager.HasObserversForBitSet(ElementsToAdd, EMassObservedOperation::AddElement);
