Mass: Add a new ArchetypeVersion that can be used to track when an archetype changes (entity addition, removal, compaction). This is separate from the EntityOrderVersion which is only incremented on specific operations that modify entity order.
Add a function to has the ArchetypeVersion for all archetypes associated with a query.

Teds: Update CalculateQueryTablesTopologyHash to use the ArchetypeVersion instead of EntityOrderVersion since we want to track all changes.

#Synced-CL 52470086

diff --git a/Engine/Plugins/Experimental/EditorDataStorage/Source/TedsCore/Private/Queries/TypedElementExtendedQueryStore.cpp b/Engine/Plugins/Experimental/EditorDataStorage/Source/TedsCore/Private/Queries/TypedElementExtendedQueryStore.cpp
--- a/Engine/Plugins/Experimental/EditorDataStorage/Source/TedsCore/Private/Queries/TypedElementExtendedQueryStore.cpp
+++ b/Engine/Plugins/Experimental/EditorDataStorage/Source/TedsCore/Private/Queries/TypedElementExtendedQueryStore.cpp
@@ -277,11 +277,11 @@ namespace UE::Editor::DataStorage
 	uint64 FExtendedQueryStore::CalculateQueryTablesTopologyHash(Handle Query) const
 	{
 		const FExtendedQuery* QueryData = Get(Query);
-		// Const cast used because CalculateArchetypeEntityOrderHash is not const as it may need to redetermine what
+		// Const cast used because CalculateArchetypeVersionHash is not const as it may need to redetermine what
 		// archetypes match with the native query. This is however conceptually const from TEDS' point of view.
 		return
 			QueryData
-			? const_cast<FExtendedQuery*>(QueryData)->NativeQuery.CalculateArchetypeEntityOrderHash()
+			? const_cast<FExtendedQuery*>(QueryData)->NativeQuery.CalculateArchetypeVersionHash()
 			: 0;
 	}
 
diff --git a/Engine/Source/Runtime/MassEntity/Internal/MassArchetypeData.h b/Engine/Source/Runtime/MassEntity/Internal/MassArchetypeData.h
--- a/Engine/Source/Runtime/MassEntity/Internal/MassArchetypeData.h
+++ b/Engine/Source/Runtime/MassEntity/Internal/MassArchetypeData.h
@@ -382,9 +382,18 @@ private:
 	 * Note that it's set once and never changed afterward.
 	 */
 	uint32 CreatedArchetypeDataVersion = 0;
+	
+	/**
+	 * The current version of this archetype, this value is incremented whenever an entity is added to or removed from this archetype and when
+	 * any operation modifies the order of hosted entities (e.g compaction). This value can be checked to see if the archetype has changed in
+	 * any way after an operation.
+	 */
+	uint32 ArchetypeVersion = 0;
 
 	/**
 	 * Incremented whenever an operation modifies the order of hosted entities, for example entity removal and compaction.
+	 * Unlike ArchetypeVersion - this value is only incremented when the order changes. E.g adding an entity to the archetype will not
+	 * increment this value since it does not modify entity order.
 	 * This value is used to validate stored entity ranges, including FMassArchetypeEntityCollection.
 	 */
 	uint32 EntityOrderVersion = 0;
@@ -576,6 +585,7 @@ public:
 
 	uint32 GetCreatedArchetypeDataVersion() const;
 	uint32 GetEntityOrderVersion() const;
+	uint32 GetArchetypeVersion() const;
 
 	void ExecuteFunction(FMassExecutionContext& RunContext, const FMassExecuteFunction& Function, const FMassQueryRequirementIndicesMapping& RequirementMapping
 		, FMassArchetypeEntityCollection::FConstEntityRangeArrayView EntityRangeContainer, const FMassChunkConditionFunction& ChunkCondition);
@@ -904,6 +914,11 @@ inline uint32 FMassArchetypeData::GetEntityOrderVersion() const
 	return EntityOrderVersion;
 }
 
+inline uint32 FMassArchetypeData::GetArchetypeVersion() const
+{
+	return ArchetypeVersion;
+}
+
 inline const FMassElementBitSet& FMassArchetypeData::GetCompositionBitSet() const
 {
 	return CompositionBitSet;
diff --git a/Engine/Source/Runtime/MassEntity/Private/MassArchetypeData.cpp b/Engine/Source/Runtime/MassEntity/Private/MassArchetypeData.cpp
--- a/Engine/Source/Runtime/MassEntity/Private/MassArchetypeData.cpp
+++ b/Engine/Source/Runtime/MassEntity/Private/MassArchetypeData.cpp
@@ -369,6 +369,8 @@ int32 FMassArchetypeData::AddEntityInternal(FMassEntityHandle Entity, const FMas
 	checkfSlow(SharedFragmentValues.DoesMatchComposition(CompositionBitSet)
 		, TEXT("Expecting values for every specified shared fragment in the archetype and only those"));
 
+	++ArchetypeVersion;
+	
 	int32 IndexWithinChunk = 0;
 	int32 AbsoluteIndex = 0;
 
@@ -546,6 +548,7 @@ void FMassArchetypeData::RemoveEntity(FMassEntityHandle Entity)
 void FMassArchetypeData::RemoveEntityInternal(const int32 AbsoluteIndex)
 {
 	++EntityOrderVersion;
+	++ArchetypeVersion;
 
 	const int32 ChunkIndex = AbsoluteIndex / NumEntitiesPerChunk;
 	const int32 IndexWithinChunk = AbsoluteIndex - (NumEntitiesPerChunk * ChunkIndex);
@@ -993,6 +996,7 @@ int32 FMassArchetypeData::CompactEntities(const double TimeAllowed)
 	if (TotalEntitiesMoved > 0)
 	{
 		++EntityOrderVersion;
+		++ArchetypeVersion;
 	}
 
 	return TotalEntitiesMoved;
@@ -1457,6 +1461,8 @@ void FMassArchetypeData::BatchAddEntities(TConstArrayView<FMassEntityHandle> Ent
 
 	checkfSlow(SharedFragmentValues.DoesMatchComposition(CompositionBitSet)
 		, TEXT("%hs parameter SharedFragmentValues doesn't match archetype's composition"), __FUNCTION__);
+	
+	++ArchetypeVersion;
 
 	FMassArchetypeEntityCollection::FArchetypeEntityRange ResultSubchunk;
 	ResultSubchunk.ChunkIndex = 0;
@@ -1669,6 +1675,7 @@ void FMassArchetypeData::BatchRemoveEntitiesInternal(const int32 ChunkIndex, con
 	}
 
 	++EntityOrderVersion;
+	++ArchetypeVersion;
 
 	FMassArchetypeChunk& Chunk = Chunks[ChunkIndex];
 	
@@ -1717,6 +1724,8 @@ void FMassArchetypeData::MoveFragmentsToAnotherArchetypeInternal(FMassArchetypeD
 {
 	FTransientChunkLocation Target(TargetChunk.GetRawMemory(), NewIndexWithinChunk);
 	FTransientChunkLocation Source(SourceChunk.GetRawMemory(), OriginalIndexWithinChunk);
+	
+	++TargetArchetype.ArchetypeVersion;
 
 	// for every TargetArchetype's fragment see if it was in the old archetype as well and if so copy it's value. 
 	// If not then initialize the fragment.
diff --git a/Engine/Source/Runtime/MassEntity/Private/MassEntityQuery.cpp b/Engine/Source/Runtime/MassEntity/Private/MassEntityQuery.cpp
--- a/Engine/Source/Runtime/MassEntity/Private/MassEntityQuery.cpp
+++ b/Engine/Source/Runtime/MassEntity/Private/MassEntityQuery.cpp
@@ -704,23 +704,18 @@ void FMassEntityQuery::ParallelForEachEntityChunk(FMassExecutionContext& Executi
 
 uint64 FMassEntityQuery::CalculateArchetypeEntityOrderHash()
 {
-	CacheArchetypes();
+	return CalculateArchetypeHash_Internal([](const FMassArchetypeData* ArchetypeData)
+		{
+			return ArchetypeData->GetEntityOrderVersion();
+		});
+}
 
-	// The hash is created as a cumulative hash to make sure that the order of archetypes doesn't matter.
-	uint64 Hash = 0;
-	for (const FMassArchetypeHandle& ArchetypeHandle : ValidArchetypes)
-	{
-		if (const FMassArchetypeData* Archetype = FMassArchetypeHelper::ArchetypeDataFromHandle(ArchetypeHandle))
+uint64 FMassEntityQuery::CalculateArchetypeVersionHash()
+{
+	return CalculateArchetypeHash_Internal([](const FMassArchetypeData* ArchetypeData)
 		{
-			int32 EntityOrderRevision = Archetype->GetEntityOrderVersion();
-			// Accept the possible overflow as a desired side effect.
-			Hash += FXxHash64::HashBuffer(&EntityOrderRevision, sizeof(int32)).Hash;
-		}
-	}
-	// Further reduce the risk of a hash collision between two different sized archetype lists.
-	int32 ArchetypeCount = ValidArchetypes.Num();
-	Hash += FXxHash64::HashBuffer(&ArchetypeCount, sizeof(int32)).Hash;
-	return Hash;
+			return ArchetypeData->GetArchetypeVersion();
+		});
 }
 
 int32 FMassEntityQuery::GetNumMatchingEntities()
@@ -825,6 +820,27 @@ void FMassEntityQuery::BuildOrderedArchetypeIndices(const int32 FirstNewArchetyp
 	}
 }
 
+uint64 FMassEntityQuery::CalculateArchetypeHash_Internal(TFunctionRef<uint32(const FMassArchetypeData*)> HashValueGetter)
+{
+	CacheArchetypes();
+
+	// The hash is created as a cumulative hash to make sure that the order of archetypes doesn't matter.
+	uint64 Hash = 0;
+	for (const FMassArchetypeHandle& ArchetypeHandle : ValidArchetypes)
+	{
+		if (const FMassArchetypeData* Archetype = FMassArchetypeHelper::ArchetypeDataFromHandle(ArchetypeHandle))
+		{
+			int32 HashValue = HashValueGetter(Archetype);
+			// Accept the possible overflow as a desired side effect.
+			Hash += FXxHash64::HashBuffer(&HashValue, sizeof(int32)).Hash;
+		}
+	}
+	// Further reduce the risk of a hash collision between two different sized archetype lists.
+	int32 ArchetypeCount = ValidArchetypes.Num();
+	Hash += FXxHash64::HashBuffer(&ArchetypeCount, sizeof(int32)).Hash;
+	return Hash;
+}
+
 void FMassEntityQuery::SortArchetypes(const int32 FirstNewArchetypeIndex)
 {
 	if (GroupSortingSteps.Num() == 0)
diff --git a/Engine/Source/Runtime/MassEntity/Public/MassEntityQuery.h b/Engine/Source/Runtime/MassEntity/Public/MassEntityQuery.h
--- a/Engine/Source/Runtime/MassEntity/Public/MassEntityQuery.h
+++ b/Engine/Source/Runtime/MassEntity/Public/MassEntityQuery.h
@@ -172,9 +172,16 @@ public:
 	/**
 	 * Calculates a hash of the entity order revisions in all archetypes. If returned value ever differs from
 	 * a previous call it indicates that one or more archetypes associated with this query have been updated
-	 * by having entities added, removed or reordered.
+	 * by having the entity order changed.
 	 */
 	UE_API uint64 CalculateArchetypeEntityOrderHash();
+	
+	/**
+	 * Calculates a hash of the version of all archetypes. If returned value ever differs from
+	 * a previous call it indicates that one or more archetypes associated with this query have been updated
+	 * by having entities added, removed or re-ordered.
+	 */
+	UE_API uint64 CalculateArchetypeVersionHash();
 
 	/** 
 	 * Goes through ValidArchetypes and sums up the number of entities contained in them.
@@ -291,6 +298,11 @@ private:
 	 * and have it reflect the order of ValidArchetypes.
 	 */
 	void BuildOrderedArchetypeIndices(const int32 FirstNewArchetypeIndex = 0);
+	
+	/**
+	 * Internal function to hash any value for all archetypes associated with this query.
+	 */
+	uint64 CalculateArchetypeHash_Internal(TFunctionRef<uint32(const FMassArchetypeData*)> HashValueGetter);
 
 	/** 
 	 * This function represents a condition that will be called for every chunk to be processed before the actual 
