[MassCommands]
Add FMassCommandRemoveElements for batch removal of mixed element types

- Add templated FMassCommandRemoveElements<TTypes...> that accepts any mix of fragments, tags, shared and const shared fragments, with PushCommand batching across chunks
- Add runtime FMassCommandRemoveElementList for dynamic element type lists via PushUniqueCommand
- Add RemoveElements convenience methods on FMassCommandBuffer (templated and runtime variants)
- Fix some call sites to use single batched command instead of fragment remove + per-entity deferred lambda
- Fix/update some unit tests broken by bad merge with local changes

#tests mass unit tests
#Synced-CL 52505535

diff --git a/Engine/Source/Developer/MassEntityTestSuite/Private/MassCommandsTest.cpp b/Engine/Source/Developer/MassEntityTestSuite/Private/MassCommandsTest.cpp
--- a/Engine/Source/Developer/MassEntityTestSuite/Private/MassCommandsTest.cpp
+++ b/Engine/Source/Developer/MassEntityTestSuite/Private/MassCommandsTest.cpp
@@ -720,53 +720,193 @@ struct FCommands_AddElement_RejectsConstSharedFragment : FEntityTestBase
 };
 IMPLEMENT_AI_INSTANT_TEST(FCommands_AddElement_RejectsConstSharedFragment, "System.Mass.Commands.AddElement.RejectsConstSharedFragment");
 
-struct FCommands_BatchChangeComposition_RejectsSharedAdd : FEntityTestBase
+struct FCommands_BatchChangeComposition_RemovesSharedFragments : FEntityTestBase
 {
 	virtual bool InstantTest() override
 	{
-		constexpr int32 Count = 1;
+		constexpr int32 Count = 5;
+		const FMassArchetypeHandle IntsTagAArchetype = EntityManager->CreateArchetype(
+			{ FTestFragment_Int::StaticStruct(), FTestTag_A::StaticStruct() });
+
 		TArray<FMassEntityHandle> Entities;
-		EntityManager->BatchCreateEntities(IntsArchetype, Count, Entities);
+		EntityManager->BatchCreateEntities(IntsTagAArchetype, Count, Entities);
+
+		// Add a const shared fragment to all entities
+		FTestConstSharedFragment_Int ConstSharedInstance(42);
+		FConstSharedStruct ConstSharedStruct = EntityManager->GetOrCreateConstSharedFragment(ConstSharedInstance);
+		for (const FMassEntityHandle& Entity : Entities)
+		{
+			EntityManager->AddConstSharedFragmentToEntity(Entity, ConstSharedStruct);
+		}
+
+		for (const FMassEntityHandle& Entity : Entities)
+		{
+			AITEST_TRUE(TEXT("Entity should have const shared before removal"), EntityManager->DoesEntityHaveElement(Entity, FTestConstSharedFragment_Int::StaticStruct()));
+			AITEST_TRUE(TEXT("Entity should have TagA before removal"), EntityManager->DoesEntityHaveElement(Entity, FTestTag_A::StaticStruct()));
+		}
+
+		// Remove tag + const shared via BatchChangeCompositionForEntities directly
+		FMassElementBitSet ElementsToRemove;
+		ElementsToRemove.Add(FTestTag_A::StaticStruct());
+		ElementsToRemove.Add(FTestConstSharedFragment_Int::StaticStruct());
+
+		TArray<FMassArchetypeEntityCollection> EntityCollections;
+		UE::Mass::Utils::CreateEntityCollections(*EntityManager.Get(), Entities, FMassArchetypeEntityCollection::NoDuplicates, EntityCollections);
+		EntityManager->BatchChangeCompositionForEntities(EntityCollections, FMassElementBitSet(), ElementsToRemove);
+
+		for (const FMassEntityHandle& Entity : Entities)
+		{
+			AITEST_TRUE(TEXT("Entity should still have FTestFragment_Int"), EntityManager->DoesEntityHaveElement(Entity, FTestFragment_Int::StaticStruct()));
+			AITEST_FALSE(TEXT("Entity should no longer have TagA"), EntityManager->DoesEntityHaveElement(Entity, FTestTag_A::StaticStruct()));
+			AITEST_FALSE(TEXT("Entity should no longer have const shared"), EntityManager->DoesEntityHaveElement(Entity, FTestConstSharedFragment_Int::StaticStruct()));
+		}
+
+		return true;
+	}
+};
+IMPLEMENT_AI_INSTANT_TEST(FCommands_BatchChangeComposition_RemovesSharedFragments, "System.Mass.Commands.BatchChangeComposition.RemovesSharedFragments");
+
+//----------------------------------------------------------------------//
+// FMassCommandRemoveElements tests
+//----------------------------------------------------------------------//
+struct FCommands_RemoveElements_TemplatedMixedTypes : FEntityTestBase
+{
+	virtual bool InstantTest() override
+	{
+		constexpr int32 Count = 5;
+		const FMassArchetypeHandle IntsFloatsTagAArchetype = EntityManager->CreateArchetype(
+			{ FTestFragment_Int::StaticStruct(), FTestFragment_Float::StaticStruct(), FTestTag_A::StaticStruct() });
+
+		TArray<FMassEntityHandle> Entities;
+		EntityManager->BatchCreateEntities(IntsFloatsTagAArchetype, Count, Entities);
+
+		// Add a const shared fragment
+		FTestConstSharedFragment_Int ConstSharedInstance(42);
+		FConstSharedStruct ConstSharedStruct = EntityManager->GetOrCreateConstSharedFragment(ConstSharedInstance);
+		for (const FMassEntityHandle& Entity : Entities)
+		{
+			EntityManager->AddConstSharedFragmentToEntity(Entity, ConstSharedStruct);
+		}
+
+		// Add a sparse fragment
+		TArray<FMassArchetypeEntityCollection> EntityCollections;
+		UE::Mass::Utils::CreateEntityCollections(*EntityManager.Get(), Entities, FMassArchetypeEntityCollection::NoDuplicates, EntityCollections);
+		EntityManager->BatchAddSparseElementToEntities(EntityCollections, FTestFragment_SparseInt::StaticStruct());
+
+		for (const FMassEntityHandle& Entity : Entities)
+		{
+			AITEST_TRUE(TEXT("Entity should have sparse before removal"), EntityManager->DoesEntityHaveElement(Entity, FTestFragment_SparseInt::StaticStruct()));
+		}
 
+		// Remove fragment + tag + const shared + sparse in one command
+		EntityManager->Defer().PushCommand<FMassCommandRemoveElements<FTestFragment_Float, FTestTag_A, FTestConstSharedFragment_Int, FTestFragment_SparseInt>>(Entities);
+
+		EntityManager->FlushCommands();
+
+		for (const FMassEntityHandle& Entity : Entities)
+		{
+			AITEST_TRUE(TEXT("Entity should still have FTestFragment_Int"), EntityManager->DoesEntityHaveElement(Entity, FTestFragment_Int::StaticStruct()));
+			AITEST_FALSE(TEXT("Entity should no longer have FTestFragment_Float"), EntityManager->DoesEntityHaveElement(Entity, FTestFragment_Float::StaticStruct()));
+			AITEST_FALSE(TEXT("Entity should no longer have TagA"), EntityManager->DoesEntityHaveElement(Entity, FTestTag_A::StaticStruct()));
+			AITEST_FALSE(TEXT("Entity should no longer have const shared"), EntityManager->DoesEntityHaveElement(Entity, FTestConstSharedFragment_Int::StaticStruct()));
+			AITEST_FALSE(TEXT("Entity should no longer have sparse"), EntityManager->DoesEntityHaveElement(Entity, FTestFragment_SparseInt::StaticStruct()));
+		}
+
+		return true;
+	}
+};
+IMPLEMENT_AI_INSTANT_TEST(FCommands_RemoveElements_TemplatedMixedTypes, "System.Mass.Commands.RemoveElements.TemplatedMixedTypes");
+
+struct FCommands_RemoveElements_RuntimeMixedTypes : FEntityTestBase
+{
+	virtual bool InstantTest() override
+	{
+		constexpr int32 Count = 5;
+		const FMassArchetypeHandle IntsFloatsTagAArchetype = EntityManager->CreateArchetype(
+			{ FTestFragment_Int::StaticStruct(), FTestFragment_Float::StaticStruct(), FTestTag_A::StaticStruct() });
+
+		TArray<FMassEntityHandle> Entities;
+		EntityManager->BatchCreateEntities(IntsFloatsTagAArchetype, Count, Entities);
+
+		// Add a const shared fragment
+		FTestConstSharedFragment_Int ConstSharedInstance(42);
+		FConstSharedStruct ConstSharedStruct = EntityManager->GetOrCreateConstSharedFragment(ConstSharedInstance);
+		for (const FMassEntityHandle& Entity : Entities)
+		{
+			EntityManager->AddConstSharedFragmentToEntity(Entity, ConstSharedStruct);
+		}
+
+		// Add a sparse fragment
 		TArray<FMassArchetypeEntityCollection> EntityCollections;
 		UE::Mass::Utils::CreateEntityCollections(*EntityManager.Get(), Entities, FMassArchetypeEntityCollection::NoDuplicates, EntityCollections);
+		EntityManager->BatchAddSparseElementToEntities(EntityCollections, FTestFragment_SparseInt::StaticStruct());
+
+		// Remove fragment + tag + const shared + sparse via runtime command
+		const UScriptStruct* ElementTypes[] = {
+			FTestFragment_Float::StaticStruct(), FTestTag_A::StaticStruct(),
+			FTestConstSharedFragment_Int::StaticStruct(), FTestFragment_SparseInt::StaticStruct()
+		};
+		FMassCommandRemoveElementList& Cmd = EntityManager->Defer().PushUniqueCommand<FMassCommandRemoveElementList>(ElementTypes);
+		Cmd.Add(Entities);
 
-		FMassConstSharedFragmentBitSet ConstSharedBits(FTestConstSharedFragment_Int::StaticStruct());
-		const FMassElementBitSet SharedBits(ConstSharedBits);
+		EntityManager->FlushCommands();
 
+		for (const FMassEntityHandle& Entity : Entities)
 		{
-			MASS_SCOPED_ENSURE_TEST(TEXT("ElementsToAdd contains shared or sparse element types"), Count);
-			EntityManager->BatchChangeCompositionForEntities(EntityCollections, SharedBits, FMassElementBitSet());
+			AITEST_TRUE(TEXT("Entity should still have FTestFragment_Int"), EntityManager->DoesEntityHaveElement(Entity, FTestFragment_Int::StaticStruct()));
+			AITEST_FALSE(TEXT("Entity should no longer have FTestFragment_Float"), EntityManager->DoesEntityHaveElement(Entity, FTestFragment_Float::StaticStruct()));
+			AITEST_FALSE(TEXT("Entity should no longer have TagA"), EntityManager->DoesEntityHaveElement(Entity, FTestTag_A::StaticStruct()));
+			AITEST_FALSE(TEXT("Entity should no longer have const shared"), EntityManager->DoesEntityHaveElement(Entity, FTestConstSharedFragment_Int::StaticStruct()));
+			AITEST_FALSE(TEXT("Entity should no longer have sparse"), EntityManager->DoesEntityHaveElement(Entity, FTestFragment_SparseInt::StaticStruct()));
 		}
 
 		return true;
 	}
 };
-IMPLEMENT_AI_INSTANT_TEST(FCommands_BatchChangeComposition_RejectsSharedAdd, "System.Mass.Commands.BatchChangeComposition.RejectsSharedAdd");
+IMPLEMENT_AI_INSTANT_TEST(FCommands_RemoveElements_RuntimeMixedTypes, "System.Mass.Commands.RemoveElements.RuntimeMixedTypes");
 
-struct FCommands_BatchChangeComposition_RejectsSharedRemove : FEntityTestBase
+struct FCommands_RemoveElements_ConvenienceMethod : FEntityTestBase
 {
 	virtual bool InstantTest() override
 	{
-		constexpr int32 Count = 1;
+		constexpr int32 Count = 5;
+		const FMassArchetypeHandle IntsFloatsTagAArchetype = EntityManager->CreateArchetype(
+			{ FTestFragment_Int::StaticStruct(), FTestFragment_Float::StaticStruct(), FTestTag_A::StaticStruct() });
+
 		TArray<FMassEntityHandle> Entities;
-		EntityManager->BatchCreateEntities(IntsArchetype, Count, Entities);
+		EntityManager->BatchCreateEntities(IntsFloatsTagAArchetype, Count, Entities);
 
+		// Add a const shared fragment
+		FTestConstSharedFragment_Int ConstSharedInstance(42);
+		FConstSharedStruct ConstSharedStruct = EntityManager->GetOrCreateConstSharedFragment(ConstSharedInstance);
+		for (const FMassEntityHandle& Entity : Entities)
+		{
+			EntityManager->AddConstSharedFragmentToEntity(Entity, ConstSharedStruct);
+		}
+
+		// Add a sparse fragment
 		TArray<FMassArchetypeEntityCollection> EntityCollections;
 		UE::Mass::Utils::CreateEntityCollections(*EntityManager.Get(), Entities, FMassArchetypeEntityCollection::NoDuplicates, EntityCollections);
+		EntityManager->BatchAddSparseElementToEntities(EntityCollections, FTestFragment_SparseInt::StaticStruct());
 
-		FMassSharedFragmentBitSet MutableSharedBits(FTestSharedFragment_Int::StaticStruct());
-		const FMassElementBitSet SharedBits(MutableSharedBits);
+		// Use convenience method with all element categories
+		EntityManager->Defer().RemoveElements<FTestFragment_Float, FTestTag_A, FTestConstSharedFragment_Int, FTestFragment_SparseInt>(Entities);
 
+		EntityManager->FlushCommands();
+
+		for (const FMassEntityHandle& Entity : Entities)
 		{
-			MASS_SCOPED_ENSURE_TEST(TEXT("ElementsToRemove contains shared or sparse element types"), Count);
-			EntityManager->BatchChangeCompositionForEntities(EntityCollections, FMassElementBitSet(), SharedBits);
+			AITEST_TRUE(TEXT("Entity should still have FTestFragment_Int"), EntityManager->DoesEntityHaveElement(Entity, FTestFragment_Int::StaticStruct()));
+			AITEST_FALSE(TEXT("Entity should no longer have FTestFragment_Float"), EntityManager->DoesEntityHaveElement(Entity, FTestFragment_Float::StaticStruct()));
+			AITEST_FALSE(TEXT("Entity should no longer have TagA"), EntityManager->DoesEntityHaveElement(Entity, FTestTag_A::StaticStruct()));
+			AITEST_FALSE(TEXT("Entity should no longer have const shared"), EntityManager->DoesEntityHaveElement(Entity, FTestConstSharedFragment_Int::StaticStruct()));
+			AITEST_FALSE(TEXT("Entity should no longer have sparse"), EntityManager->DoesEntityHaveElement(Entity, FTestFragment_SparseInt::StaticStruct()));
 		}
 
 		return true;
 	}
 };
-IMPLEMENT_AI_INSTANT_TEST(FCommands_BatchChangeComposition_RejectsSharedRemove, "System.Mass.Commands.BatchChangeComposition.RejectsSharedRemove");
+IMPLEMENT_AI_INSTANT_TEST(FCommands_RemoveElements_ConvenienceMethod, "System.Mass.Commands.RemoveElements.ConvenienceMethod");
 
 #endif // WITH_MASSENTITY_DEBUG
 
diff --git a/Engine/Source/Runtime/Mass/MassEngine/Private/Mesh/MassEngineRenderISMProcessors.cpp b/Engine/Source/Runtime/Mass/MassEngine/Private/Mesh/MassEngineRenderISMProcessors.cpp
--- a/Engine/Source/Runtime/Mass/MassEngine/Private/Mesh/MassEngineRenderISMProcessors.cpp
+++ b/Engine/Source/Runtime/Mass/MassEngine/Private/Mesh/MassEngineRenderISMProcessors.cpp
@@ -479,19 +479,7 @@ void UMassISMDestroyRenderStateProcessor::Execute(FMassEntityManager& EntityMana
 				FMassEntityHandle EntityHandle = Context.GetEntity(*EntityIt);
 
 				ISMInfoFragments[*EntityIt].InstanceIndex = INDEX_NONE;
-				// @Todo to reenable once we allow to remove shared fragments through this command
-				Context.Defer().PushCommand<FMassCommandRemoveFragments<FMassRenderISMInstanceFragment/*, FMassRenderISMLinkFragment*/>>(EntityHandle);
-				Context.Defer().PushCommand<FMassDeferredRemoveCommand>([EntityHandle](FMassEntityManager& InEntityManager)
-					{
-						if (InEntityManager.IsEntityValid(EntityHandle))
-						{
-							InEntityManager.RemoveConstSharedFragmentFromEntity<FMassRenderISMLinkFragment>(EntityHandle);
-						}
-						else
-						{
-							UE_LOGF(LogMassEngine, Log, "FMassRenderISMLinkFragment not removed: Entity is no longer valid");
-						}
-					});
+				Context.Defer().PushCommand<FMassCommandRemoveElements<FMassRenderISMInstanceFragment, FMassRenderISMLinkFragment>>(EntityHandle);
 			}
 		});
 
diff --git a/Engine/Source/Runtime/MassEntity/Private/MassArchetypeData.cpp b/Engine/Source/Runtime/MassEntity/Private/MassArchetypeData.cpp
--- a/Engine/Source/Runtime/MassEntity/Private/MassArchetypeData.cpp
+++ b/Engine/Source/Runtime/MassEntity/Private/MassArchetypeData.cpp
@@ -265,11 +265,9 @@ void FMassArchetypeData::InitializeWithSimilar(const FMassEntityManager& EntityM
 {
 	checkf(IsInitialized() == false, TEXT("Trying to InitializeWithSimilar but this archetype has already been initialized"));
 
-	// InCompositionBitSet - BaseArchetype.GetCompositionBitSet() -> should contain only tags and fragments
-	// BaseArchetype.GetCompositionBitSet() - InCompositionBitSet -> should contain only tags and fragments
-	bool bAddsOnlyTagsAndElements = FMassElementBitSet(InCompositionBitSet - BaseArchetype.GetCompositionBitSet()).ContainsOnlyTagsAndFragments();
-	bool bNothingElse = FMassElementBitSet(BaseArchetype.GetCompositionBitSet() - InCompositionBitSet).ContainsOnlyTagsAndFragments();
-	checkf(bAddsOnlyTagsAndElements && bNothingElse, TEXT("NewComposition is expected to differ from the BaseArchetype's composition just by tags and fragments"))
+	// The composition delta between InCompositionBitSet and BaseArchetype may include tags, fragments,
+	// and shared fragments. Shared fragment changes don't affect fragment configs or chunk data layout,
+	// so the similar-archetype optimization (copying fragment configs from base) is still valid.
 
 	CreatedArchetypeDataVersion = ArchetypeDataVersion;
 	SparseElementsStorage = &InSparseElementsStorage;
diff --git a/Engine/Source/Runtime/MassEntity/Private/MassEntityManager.cpp b/Engine/Source/Runtime/MassEntity/Private/MassEntityManager.cpp
--- a/Engine/Source/Runtime/MassEntity/Private/MassEntityManager.cpp
+++ b/Engine/Source/Runtime/MassEntity/Private/MassEntityManager.cpp
@@ -1888,6 +1888,11 @@ void FMassEntityManager::BatchChangeCompositionForEntities(TConstArrayView<FMass
 
 	TRACE_CPUPROFILER_EVENT_SCOPE(Mass_BatchChangeCompositionForEntities);
 
+	// Extract shared fragment removal bitsets for BatchMoveEntitiesToAnotherArchetype
+	const FMassSharedFragmentBitSet SharedToRemove = ElementsToRemove.Get<FMassSharedFragmentBitSet>();
+	const FMassConstSharedFragmentBitSet ConstSharedToRemove = ElementsToRemove.Get<FMassConstSharedFragmentBitSet>();
+	const bool bRemoveSharedFragments = !SharedToRemove.IsEmpty() || !ConstSharedToRemove.IsEmpty();
+
 	const bool bElementsAddedAreObserved = ObserverManager.HasObserversForBitSet(ElementsToAdd, EMassObservedOperation::AddElement);
 	const bool bElementsRemovedAreObserved = ObserverManager.HasObserversForBitSet(ElementsToRemove, EMassObservedOperation::RemoveElement);
 
@@ -1916,20 +1921,23 @@ void FMassEntityManager::BatchChangeCompositionForEntities(TConstArrayView<FMass
 					ObserverManager.OnCompositionChanged(Collection, MoveTemp(ElementsRemoved), EMassObservedOperation::RemoveElement);
 				}
 			}
-			
+
 			FMassArchetypeHandle NewArchetypeHandle = InternalCreateSimilarArchetype(Collection.GetArchetype().DataPtr, MoveTemp(NewComposition));
 			checkSlow(NewArchetypeHandle.IsValid());
 
-			// Move the entity over
+			// Move the entity over; single move handles fragment/tag changes + shared fragment removal
 			FMassArchetypeEntityCollection::FEntityRangeArray NewArchetypeEntityRanges;
 			TArray<FMassEntityHandle> EntitiesBeingMoved;
 			CurrentArchetype->BatchMoveEntitiesToAnotherArchetype(Collection, *NewArchetypeHandle.DataPtr.Get(), EntitiesBeingMoved
-				, bElementsAddedAreObserved && (ElementsAdded.IsEmpty() == false) ? &NewArchetypeEntityRanges : nullptr);
+				, bElementsAddedAreObserved && (ElementsAdded.IsEmpty() == false) ? &NewArchetypeEntityRanges : nullptr
+				, /*SharedFragmentValuesToAdd=*/nullptr
+				, bRemoveSharedFragments ? &SharedToRemove : nullptr
+				, bRemoveSharedFragments ? &ConstSharedToRemove : nullptr);
 
 			for (const FMassEntityHandle& Entity : EntitiesBeingMoved)
 			{
 				check(GetEntityStorageInterface().IsValidIndex(Entity.Index));
-				
+
 				GetEntityStorageInterface().SetArchetypeFromShared(Entity.Index, NewArchetypeHandle.DataPtr);
 			}
 
diff --git a/Engine/Source/Runtime/MassEntity/Public/MassCommandBuffer.h b/Engine/Source/Runtime/MassEntity/Public/MassCommandBuffer.h
--- a/Engine/Source/Runtime/MassEntity/Public/MassCommandBuffer.h
+++ b/Engine/Source/Runtime/MassEntity/Public/MassCommandBuffer.h
@@ -187,6 +187,7 @@ struct FMassCommandBuffer
 	}
 
 	template<typename T>
+	UE_DEPRECATED(5.8, "Use AddFragment instead")
 	void AddFragment_RuntimeCheck(FMassEntityHandle Entity)
 	{
 		checkf(UE::Mass::IsA<FMassFragment>(T::StaticStruct()), TEXT(MASS_INVALID_FRAGMENT_MSG_F), *T::StaticStruct()->GetName());
@@ -197,14 +198,15 @@ struct FMassCommandBuffer
 	void RemoveFragment(FMassEntityHandle Entity)
 	{
 		MASS_STATIC_CHECK_FRAGMENT(T);
-		PushCommand<FMassCommandRemoveFragmentsInternal<EMassCommandCheckTime::CompileTimeCheck, T>>(Entity);
+		PushCommand<FMassCommandRemoveElements<T>>(Entity);
 	}
 
 	template<typename T>
+	UE_DEPRECATED(5.8, "Use RemoveFragment instead")
 	void RemoveFragment_RuntimeCheck(FMassEntityHandle Entity)
 	{
 		checkf(UE::Mass::IsA<FMassFragment>(T::StaticStruct()), TEXT(MASS_INVALID_FRAGMENT_MSG_F), *T::StaticStruct()->GetName());
-		PushCommand<FMassCommandRemoveFragmentsInternal<EMassCommandCheckTime::RuntimeCheck, T>>(Entity);
+		PushCommand<FMassCommandRemoveElements<T>>(Entity);
 	}
 
 	/** the convenience function equivalent to calling PushCommand<FMassCommandAddTag<T>>(Entity) */
@@ -223,6 +225,7 @@ struct FMassCommandBuffer
 	}
 
 	template<typename T>
+	UE_DEPRECATED(5.8, "Use AddTag instead")
 	void AddTag_RuntimeCheck(FMassEntityHandle Entity)
 	{
 		checkf(UE::Mass::IsA<FMassTag>(T::StaticStruct()), TEXT("Given struct type is not a valid tag type."));
@@ -234,21 +237,22 @@ struct FMassCommandBuffer
 	void RemoveTag(FMassEntityHandle Entity)
 	{
 		static_assert(UE::Mass::CTag<T>, "Given struct type is not a valid tag type.");
-		PushCommand<FMassCommandRemoveTagsInternal<EMassCommandCheckTime::CompileTimeCheck, T>>(Entity);
+		PushCommand<FMassCommandRemoveElements<T>>(Entity);
 	}
 
 	template<typename T>
 	void RemoveTag(TConstArrayView<FMassEntityHandle> Entities)
 	{
 		static_assert(UE::Mass::CTag<T>, "Given struct type is not a valid tag type.");
-		PushCommand<FMassCommandRemoveTagsInternal<EMassCommandCheckTime::CompileTimeCheck, T>>(Entities);
+		PushCommand<FMassCommandRemoveElements<T>>(Entities);
 	}
 
 	template<typename T>
+	UE_DEPRECATED(5.8, "Use RemoveTag instead")
 	void RemoveTag_RuntimeCheck(FMassEntityHandle Entity)
 	{
 		checkf(UE::Mass::IsA<FMassTag>(T::StaticStruct()), TEXT("Given struct type is not a valid tag type."));
-		PushCommand<FMassCommandRemoveTagsInternal<EMassCommandCheckTime::RuntimeCheck, T>>(Entity);
+		PushCommand<FMassCommandRemoveElements<T>>(Entity);
 	}
 
 	/** the convenience function equivalent to calling PushCommand<FMassCommandSwapTags<TOld, TNew>>(Entity)  */
@@ -261,6 +265,7 @@ struct FMassCommandBuffer
 	}
 
 	template<typename TOld, typename TNew>
+	UE_DEPRECATED(5.8, "Use SwapTags instead")
 	void SwapTags_RuntimeCheck(FMassEntityHandle Entity)
 	{
 		checkf(UE::Mass::IsA<FMassTag>(TOld::StaticStruct()), TEXT("Given struct type is not a valid tag type."));
@@ -268,6 +273,32 @@ struct FMassCommandBuffer
 		PushCommand<FMassCommandSwapTagsInternal<EMassCommandCheckTime::RuntimeCheck, TOld, TNew>>(Entity);
 	}
 
+	/** Removes any mix of element types (fragments, tags, shared, const shared, sparse) from an entity. Batches via PushCommand. */
+	template<typename... TTypes>
+	void RemoveElements(FMassEntityHandle Entity)
+	{
+		PushCommand<FMassCommandRemoveElements<TTypes...>>(Entity);
+	}
+
+	/** Removes any mix of element types (fragments, tags, shared, const shared, sparse) from entities. Batches via PushCommand. */
+	template<typename... TTypes>
+	void RemoveElements(TConstArrayView<FMassEntityHandle> Entities)
+	{
+		PushCommand<FMassCommandRemoveElements<TTypes...>>(Entities);
+	}
+
+	/** Removes any mix of element types from an entity. Types provided at runtime. Uses PushUniqueCommand. */
+	void RemoveElements(FMassEntityHandle Entity, TConstArrayView<const UScriptStruct*> ElementTypes)
+	{
+		PushUniqueCommand<FMassCommandRemoveElementList>(ElementTypes).Add(Entity);
+	}
+
+	/** Removes any mix of element types from entities. Types provided at runtime. Uses PushUniqueCommand. */
+	void RemoveElements(TConstArrayView<FMassEntityHandle> Entities, TConstArrayView<const UScriptStruct*> ElementTypes)
+	{
+		PushUniqueCommand<FMassCommandRemoveElementList>(ElementTypes).Add(Entities);
+	}
+
 	void DestroyEntity(FMassEntityHandle Entity)
 	{
 		PushCommand<FMassCommandDestroyEntities>(Entity);
diff --git a/Engine/Source/Runtime/MassEntity/Public/MassCommands.h b/Engine/Source/Runtime/MassEntity/Public/MassCommands.h
--- a/Engine/Source/Runtime/MassEntity/Public/MassCommands.h
+++ b/Engine/Source/Runtime/MassEntity/Public/MassCommands.h
@@ -19,7 +19,7 @@
 UENUM()
 enum class EMassCommandOperationType : uint8
 {
-	None,				// default value. Commands marked this way will be always executed last. Programmers are encouraged to instead use one of the meaningful values below.
+	None,				// default value. Commands marked this way will always be executed last. Programmers are encouraged to instead use one of the meaningful values below.
 	Create,				// signifies commands performing entity creation
 	Add,				// signifies commands adding fragments or tags to entities
 	Remove,				// signifies commands removing fragments or tags from entities
@@ -247,5 +247,214 @@
 };
 
 //-----------------------------------------------------------------------------
+// Mixed element types
+//-----------------------------------------------------------------------------
+/**
+ * Command adds any element type to entities, be it a fragment or a tag, including sparse element types
+ * Note that this type of command can only be added via PushUniqueCommand
+ */
+struct FMassCommandAddElement : public FMassBatchedEntityCommand
+{
+	using Super = FMassBatchedEntityCommand;
+
+	FMassCommandAddElement(const TNotNull<const UScriptStruct*> InElementType)
+		: Super(EMassCommandOperationType::Add DEBUG_NAME("AddElement"))
+		, ElementType(InElementType)
+	{}
+
+protected:
+	virtual void Run(FMassEntityManager& EntityManager) override
+	{
+		TRACE_CPUPROFILER_EVENT_SCOPE(MassCommandAddElement_Execute);
+		EntityManager.BatchAddElementToEntities(TargetEntities, ElementType);
+	}
+
+	TNotNull<const UScriptStruct*> ElementType;
+};
+
+template<>
+struct UE::Mass::Command::TCommandTraits<FMassCommandAddElement> final
+{
+	enum
+	{
+		RequiresUniqueHandling = true
+	};
+};
+
+/**
+ * Command removes any element type from entities, be it a fragment or a tag, including sparse element types
+ * Note that this type of command can only be added via PushUniqueCommand
+ */
+struct FMassCommandRemoveElement : public FMassBatchedEntityCommand
+{
+	using Super = FMassBatchedEntityCommand;
+
+	FMassCommandRemoveElement(const TNotNull<const UScriptStruct*> InElementType)
+		: Super(EMassCommandOperationType::Remove DEBUG_NAME("RemoveElement"))
+		, ElementType(InElementType)
+	{}
+
+protected:
+	virtual void Run(FMassEntityManager& EntityManager) override
+	{
+		TRACE_CPUPROFILER_EVENT_SCOPE(MassCommandRemoveElement_Execute);
+		EntityManager.BatchRemoveElementFromEntities(TargetEntities, ElementType);
+	}
+
+	TNotNull<const UScriptStruct*> ElementType;
+};
+
+template<>
+struct UE::Mass::Command::TCommandTraits<FMassCommandRemoveElement> final
+{
+	enum
+	{
+		RequiresUniqueHandling = true
+	};
+};
+
+/**
+ * Runtime command that removes any mix of element types from entities.
+ * Element types are provided as UScriptStruct* at construction time.
+ * Uses PushUniqueCommand (no cross-call batching, but entities batch via Add()).
+ */
+struct FMassCommandRemoveElementList : public FMassBatchedEntityCommand
+{
+	using Super = FMassBatchedEntityCommand;
+
+	FMassCommandRemoveElementList(TConstArrayView<const UScriptStruct*> InElementTypes)
+		: Super(EMassCommandOperationType::Remove DEBUG_NAME("RemoveElementList"))
+	{
+		for (const UScriptStruct* ElementType : InElementTypes)
+		{
+			if (UE::Mass::IsSparse(ElementType))
+			{
+				SparseTypes.Add(ElementType);
+			}
+			else
+			{
+				NonSparseTypes.Add(ElementType);
+			}
+		}
+	}
+
+protected:
+	virtual void Run(FMassEntityManager& EntityManager) override
+	{
+		TRACE_CPUPROFILER_EVENT_SCOPE(MassCommandRemoveElementList_Execute);
+
+		TArray<FMassArchetypeEntityCollection> EntityCollections;
+		UE::Mass::Utils::CreateEntityCollections(EntityManager, TargetEntities, FMassArchetypeEntityCollection::FoldDuplicates, EntityCollections);
+
+		// Sparse removal is always in-place (no entity move)
+		// Process first before changing the composition of the entities that will invalidate
+		// the entity collection
+		for (const UScriptStruct* SparseType : SparseTypes)
+		{
+			EntityManager.BatchRemoveSparseElementFromEntities(EntityCollections, SparseType);
+		}
+
+		if (!NonSparseTypes.IsEmpty())
+		{
+			EntityManager.BatchChangeCompositionForEntities(EntityCollections, FMassElementBitSet(), NonSparseTypes);
+		}
+	}
+
+	virtual SIZE_T GetAllocatedSize() const override
+	{
+		return Super::GetAllocatedSize() + NonSparseTypes.GetAllocatedSize() + SparseTypes.GetAllocatedSize();
+	}
+
+#if CSV_PROFILER_STATS || WITH_MASSENTITY_DEBUG
+	virtual int32 GetNumOperationsStat() const override
+	{
+		return TargetEntities.Num();
+	}
+#endif // CSV_PROFILER_STATS || WITH_MASSENTITY_DEBUG
+
+	FMassElementBitSet NonSparseTypes;
+	TArray<const UScriptStruct*> SparseTypes;
+};
+
+template<>
+struct UE::Mass::Command::TCommandTraits<FMassCommandRemoveElementList> final
+{
+	enum
+	{
+		RequiresUniqueHandling = true
+	};
+};
+
+/**
+ * Templated command that removes any mix of element types (fragments, tags, shared, const shared, sparse)
+ * from entities. Uses PushCommand for efficient batching across chunks.
+ */
+template<typename... TTypes>
+struct FMassCommandRemoveElements : public FMassBatchedEntityCommand
+{
+	using Super = FMassBatchedEntityCommand;
+
+	FMassCommandRemoveElements()
+		: Super(EMassCommandOperationType::Remove DEBUG_NAME("RemoveElements"))
+	{
+		(CategorizeElementType<TTypes>(), ...);
+	}
+
+#if WITH_MASSENTITY_DEBUG
+	template<typename T>
+	static bool CheckBreakpoints(T Entity, TTypes... InFragments)
+	{
+		return UE::Mass::Debug::FBreakpoint::CheckFragmentRemoveBreakpoints(Entity, Forward<TTypes>(InFragments)...);
+	}
+#endif // WITH_MASSENTITY_DEBUG
+
+private:
+	template<typename T>
+	void CategorizeElementType()
+	{
+		using FClean = UE::Mass::Clean<T>;
+		if constexpr (UE::Mass::CSparse<FClean>)
+		{
+			SparseTypes.Add(FClean::StaticStruct());
+		}
+		else
+		{
+			NonSparseTypes.Add(FClean::StaticStruct());
+		}
+	}
+
+protected:
+	virtual void Run(FMassEntityManager& EntityManager) override
+	{
+		TRACE_CPUPROFILER_EVENT_SCOPE(MassCommandRemoveElements_Execute);
+
+		TArray<FMassArchetypeEntityCollection> EntityCollections;
+		UE::Mass::Utils::CreateEntityCollections(EntityManager, TargetEntities, FMassArchetypeEntityCollection::FoldDuplicates, EntityCollections);
+
+		// Sparse removal is always in-place (no entity move)
+		// Process first before changing the composition of the entities that will invalidate
+		// the entity collection
+		for (const UScriptStruct* SparseType : SparseTypes)
+		{
+			EntityManager.BatchRemoveSparseElementFromEntities(EntityCollections, SparseType);
+		}
+
+		// Single entity move for all non-sparse element types
+		if (!NonSparseTypes.IsEmpty())
+		{
+			EntityManager.BatchChangeCompositionForEntities(EntityCollections, FMassElementBitSet(), NonSparseTypes);
+		}
+	}
+
+	virtual SIZE_T GetAllocatedSize() const override
+	{
+		return Super::GetAllocatedSize() + NonSparseTypes.GetAllocatedSize() + SparseTypes.GetAllocatedSize();
+	}
+
+	FMassElementBitSet NonSparseTypes;
+	TArray<const UScriptStruct*> SparseTypes;
+};
+
+//-----------------------------------------------------------------------------
 // Simple fragment composition change
 //-----------------------------------------------------------------------------
@@ -283,7 +492,7 @@
 using FMassCommandAddFragments = FMassCommandAddFragmentsInternal<EMassCommandCheckTime::CompileTimeCheck, TTypes...>;
 
 template<EMassCommandCheckTime CheckTime, typename... TTypes>
-struct FMassCommandRemoveFragmentsInternal : public FMassBatchedEntityCommand
+struct UE_DEPRECATED(5.8, "Use FMassCommandRemoveElements instead — supports any mix of element types and performs a single entity move.") FMassCommandRemoveFragmentsInternal : public FMassBatchedEntityCommand
 {
 	static_assert((TIsDerivedFrom<UE::Mass::Clean<TTypes>, FMassFragment>::Value && ...), "All types passed to FMassCommandRemoveFragments must derive from FMassFragment.");
 
@@ -313,7 +522,7 @@
 };
 
 template<typename... TTypes>
-using FMassCommandRemoveFragments = FMassCommandRemoveFragmentsInternal<EMassCommandCheckTime::CompileTimeCheck, TTypes...>;
+using FMassCommandRemoveFragments = FMassCommandRemoveElements<TTypes...>;
 
 //-----------------------------------------------------------------------------
 // Simple tag composition change
@@ -372,7 +581,7 @@
 using FMassCommandAddTags = FMassCommandAddTagsInternal<EMassCommandCheckTime::CompileTimeCheck, TTypes...>;
 
 template<EMassCommandCheckTime CheckTime, typename... TTypes>
-struct FMassCommandRemoveTagsInternal : public FMassCommandChangeTags
+struct UE_DEPRECATED(5.8, "Use FMassCommandRemoveElements instead — supports any mix of element types and performs a single entity move.") FMassCommandRemoveTagsInternal : public FMassCommandChangeTags
 {
 	static_assert((TIsDerivedFrom<UE::Mass::Clean<TTypes>, FMassTag>::Value && ...), "All types passed to FMassCommandRemoveTags must derive from FMassTag.");
 
@@ -389,16 +598,16 @@
 	template<typename T>
 	static bool CheckBreakpoints(T Entity, TTypes... InFragments)
 	{
-		return UE::Mass::Debug::FBreakpoint::CheckFragmentAddBreakpoints(Entity, Forward<TTypes>(InFragments)...);
+		return UE::Mass::Debug::FBreakpoint::CheckFragmentRemoveBreakpoints(Entity, Forward<TTypes>(InFragments)...);
 	}
 #endif // WITH_MASSENTITY_DEBUG
 };
 
 template<typename T>
-using FMassCommandRemoveTag = FMassCommandRemoveTagsInternal<EMassCommandCheckTime::CompileTimeCheck, T>;
+using FMassCommandRemoveTag = FMassCommandRemoveElements<T>;
 
 template<typename... TTypes>
-using FMassCommandRemoveTags = FMassCommandRemoveTagsInternal<EMassCommandCheckTime::CompileTimeCheck, TTypes...>;
+using FMassCommandRemoveTags = FMassCommandRemoveElements<TTypes...>;
 
 template<EMassCommandCheckTime CheckTime, typename TOld, typename TNew>
 struct FMassCommandSwapTagsInternal : public FMassCommandChangeTags
@@ -478,70 +687,6 @@
 	const FMassFragmentBitSet FragmentsAffected;
 };
 
-/**
- * Command adds any element type to entities, be it a fragment or a tag, including sparse element types
- * Note that this type of command can only be added via PushUniqueCommand
- */
-struct FMassCommandAddElement : public FMassBatchedEntityCommand
-{
-	using Super = FMassBatchedEntityCommand;
-
-	FMassCommandAddElement(const TNotNull<const UScriptStruct*> InElementType)
-		: Super(EMassCommandOperationType::Add DEBUG_NAME("AddElement"))
-		, ElementType(InElementType)
-	{}
-
-protected:
-	virtual void Run(FMassEntityManager& EntityManager) override
-	{
-		TRACE_CPUPROFILER_EVENT_SCOPE(MassCommandAddElement_Execute);
-		EntityManager.BatchAddElementToEntities(TargetEntities, ElementType);
-	}
-
-	TNotNull<const UScriptStruct*> ElementType;
-};
-
-template<>
-struct UE::Mass::Command::TCommandTraits<FMassCommandAddElement> final
-{
-	enum
-	{
-		RequiresUniqueHandling = true
-	};
-};
-
-/**
- * Command removes any element type from entities, be it a fragment or a tag, including sparse element types
- * Note that this type of command can only be added via PushUniqueCommand
- */
-struct FMassCommandRemoveElement : public FMassBatchedEntityCommand
-{
-	using Super = FMassBatchedEntityCommand;
-
-	FMassCommandRemoveElement(const TNotNull<const UScriptStruct*> InElementType)
-		: Super(EMassCommandOperationType::Remove DEBUG_NAME("RemoveElement"))
-		, ElementType(InElementType)
-	{}
-
-protected:
-	virtual void Run(FMassEntityManager& EntityManager) override
-	{
-		TRACE_CPUPROFILER_EVENT_SCOPE(MassCommandAddElement_Execute);
-		EntityManager.BatchRemoveElementFromEntities(TargetEntities, ElementType);
-	}
-
-	TNotNull<const UScriptStruct*> ElementType;
-};
-
-template<>
-struct UE::Mass::Command::TCommandTraits<FMassCommandRemoveElement> final
-{
-	enum
-	{
-		RequiresUniqueHandling = true
-	};
-};
-
 template<typename... TOthers>
 struct FMassCommandBuildEntity : public FMassCommandAddFragmentInstances<TOthers...>
 {
