[MassEntity] fix off-by-one in CacheSubsystem, minor cleanup

Subsystems.AddZeroed was computing the wrong delta when resizing - `Num() - SystemIndex` instead of `SystemIndex - Num()`. Hasnt blown up in practice because the array is pre-sized in the ctor, but its still wrong.
Also: redundant FTypeHandle re-wrapping in RegisterTypeInternal, duplicate include in MassRequirements.h.

#Synced-CL 52534678

diff --git a/Engine/Source/Runtime/MassEntity/Private/MassSubsystemAccess.cpp b/Engine/Source/Runtime/MassEntity/Private/MassSubsystemAccess.cpp
--- a/Engine/Source/Runtime/MassEntity/Private/MassSubsystemAccess.cpp
+++ b/Engine/Source/Runtime/MassEntity/Private/MassSubsystemAccess.cpp
@@ -88,7 +88,7 @@ bool FMassSubsystemAccess::CacheSubsystem(const uint32 SystemIndex)
 {
 	if (UNLIKELY(Subsystems.IsValidIndex(SystemIndex) == false))
 	{
-		Subsystems.AddZeroed(Subsystems.Num() - SystemIndex + 1);
+		Subsystems.AddZeroed(SystemIndex - Subsystems.Num() + 1);
 	}
 
 	if (Subsystems[SystemIndex])
diff --git a/Engine/Source/Runtime/MassEntity/Private/MassTypeManager.cpp b/Engine/Source/Runtime/MassEntity/Private/MassTypeManager.cpp
--- a/Engine/Source/Runtime/MassEntity/Private/MassTypeManager.cpp
+++ b/Engine/Source/Runtime/MassEntity/Private/MassTypeManager.cpp
@@ -112,15 +112,15 @@ namespace UE::Mass
 	FTypeHandle FTypeManager::RegisterTypeInternal(TNotNull<const UStruct*> InType, FTypeInfo&& TypeInfo)
 	{
 		FTypeHandle TypeHandle(InType);
-		FTypeInfo* ExistingData = TypeDataMap.Find(FTypeHandle(TypeHandle.TypeKey));
+		FTypeInfo* ExistingData = TypeDataMap.Find(TypeHandle);
 		if (LIKELY(ExistingData == nullptr))
 		{
 			if (TypeInfo.Traits.IsType<FSubsystemTypeTraits>())
 			{
-				SubsystemTypes.Add(FTypeHandle(TypeHandle.TypeKey));
+				SubsystemTypes.Add(TypeHandle);
 			}
 
-			TypeDataMap.Add(FTypeHandle(TypeHandle.TypeKey), MoveTemp(TypeInfo));
+			TypeDataMap.Add(TypeHandle, MoveTemp(TypeInfo));
 		}
 		else
 		{
diff --git a/Engine/Source/Runtime/MassEntity/Public/MassRequirements.h b/Engine/Source/Runtime/MassEntity/Public/MassRequirements.h
--- a/Engine/Source/Runtime/MassEntity/Public/MassRequirements.h
+++ b/Engine/Source/Runtime/MassEntity/Public/MassRequirements.h
@@ -5,7 +5,6 @@
 #include "MassEntityTypes.h"
 #include "Mass/ExternalSubsystemTraits.h"
 #include "MassEntityLinkFragments.h"
-#include "MassEntityTypes.h"
 #include "MassArchetypeTypes.h"
 #include "Templates/SubclassOf.h"
 #include "MassRequirements.generated.h"
