[MassCommands]
Unify Add/Remove commands with FMassCommandAddElements and FMassCommandRemoveElements

- Add FMassCommandAddElements<TTypes...> (PushCommand batching) and FMassCommandAddElementList (PushUniqueCommand) for adding any mix of fragments, tags, and sparse elements in a single entity move
- Add AddElements convenience methods on FMassCommandBuffer (templated and runtime)
- Redirect AddFragment, AddTag convenience methods to FMassCommandAddElements
- Deprecate FMassCommandAddFragmentsInternal and FMassCommandAddTagsInternal
- Add sparse element support to FMassCommandAddFragmentsInternal and FMassCommandAddTagsInternal by splitting sparse types to BatchAddSparseElementToEntities
- Add check in BatchChangeCompositionForEntities rejecting shared/chunk/sparse on Add side
- Add chunk fragment rejection in BatchAddElementToEntities and AddElementToEntity
- Fix CheckBreakpoints signature on AddElements and RemoveElements to match composition-only usage (entity-only, no fragment instances)
- Add unit tests for chunk/shared/sparse Add-side gaps

#Synced-CL 52558571

diff --git a/Engine/Source/Developer/MassEntityTestSuite/Private/MassCommandsTest.cpp b/Engine/Source/Developer/MassEntityTestSuite/Private/MassCommandsTest.cpp
--- a/Engine/Source/Developer/MassEntityTestSuite/Private/MassCommandsTest.cpp
+++ b/Engine/Source/Developer/MassEntityTestSuite/Private/MassCommandsTest.cpp
@@ -908,6 +908,78 @@ struct FCommands_RemoveElements_ConvenienceMethod : FEntityTestBase
 };
 IMPLEMENT_AI_INSTANT_TEST(FCommands_RemoveElements_ConvenienceMethod, "System.Mass.Commands.RemoveElements.ConvenienceMethod");
 
+//----------------------------------------------------------------------//
+// Add command gap tests
+//----------------------------------------------------------------------//
+struct FCommands_AddElement_RejectsChunkFragment : FEntityTestBase
+{
+	virtual bool InstantTest() override
+	{
+		constexpr int32 Count = 1;
+		TArray<FMassEntityHandle> Entities;
+		EntityManager->BatchCreateEntities(IntsArchetype, Count, Entities);
+
+		FMassCommandAddElement& Cmd = EntityManager->Defer().PushUniqueCommand<FMassCommandAddElement>(FTestChunkFragment_Int::StaticStruct());
+		Cmd.Add(Entities[0]);
+
+		{
+			MASS_SCOPED_ENSURE_TEST(TEXT("chunk fragments require dedicated APIs"), Count);
+			EntityManager->FlushCommands();
+		}
+
+		return true;
+	}
+};
+IMPLEMENT_AI_INSTANT_TEST(FCommands_AddElement_RejectsChunkFragment, "System.Mass.Commands.AddElement.RejectsChunkFragment");
+
+struct FCommands_BatchChangeComposition_RejectsSharedAdd : FEntityTestBase
+{
+	virtual bool InstantTest() override
+	{
+		constexpr int32 Count = 1;
+		TArray<FMassEntityHandle> Entities;
+		EntityManager->BatchCreateEntities(IntsArchetype, Count, Entities);
+
+		TArray<FMassArchetypeEntityCollection> EntityCollections;
+		UE::Mass::Utils::CreateEntityCollections(*EntityManager.Get(), Entities, FMassArchetypeEntityCollection::NoDuplicates, EntityCollections);
+
+		FMassElementBitSet SharedBits;
+		SharedBits.Add(FTestConstSharedFragment_Int::StaticStruct());
+
+		{
+			MASS_SCOPED_ENSURE_TEST(TEXT("ElementsToAdd contains shared, chunk, or sparse element types"), Count);
+			EntityManager->BatchChangeCompositionForEntities(EntityCollections, SharedBits, FMassElementBitSet());
+		}
+
+		return true;
+	}
+};
+IMPLEMENT_AI_INSTANT_TEST(FCommands_BatchChangeComposition_RejectsSharedAdd, "System.Mass.Commands.BatchChangeComposition.RejectsSharedAdd");
+
+struct FCommands_BatchChangeComposition_RejectsSparseAdd : FEntityTestBase
+{
+	virtual bool InstantTest() override
+	{
+		constexpr int32 Count = 1;
+		TArray<FMassEntityHandle> Entities;
+		EntityManager->BatchCreateEntities(IntsArchetype, Count, Entities);
+
+		TArray<FMassArchetypeEntityCollection> EntityCollections;
+		UE::Mass::Utils::CreateEntityCollections(*EntityManager.Get(), Entities, FMassArchetypeEntityCollection::NoDuplicates, EntityCollections);
+
+		FMassElementBitSet SparseBits;
+		SparseBits.Add(FTestFragment_SparseInt::StaticStruct());
+
+		{
+			MASS_SCOPED_ENSURE_TEST(TEXT("ElementsToAdd contains shared, chunk, or sparse element types"), Count);
+			EntityManager->BatchChangeCompositionForEntities(EntityCollections, SparseBits, FMassElementBitSet());
+		}
+
+		return true;
+	}
+};
+IMPLEMENT_AI_INSTANT_TEST(FCommands_BatchChangeComposition_RejectsSparseAdd, "System.Mass.Commands.BatchChangeComposition.RejectsSparseAdd");
+
 #endif // WITH_MASSENTITY_DEBUG
 
 //----------------------------------------------------------------------//
diff --git a/Engine/Source/Runtime/MassEntity/Private/MassDebuggerBreakpoints.cpp b/Engine/Source/Runtime/MassEntity/Private/MassDebuggerBreakpoints.cpp
--- a/Engine/Source/Runtime/MassEntity/Private/MassDebuggerBreakpoints.cpp
+++ b/Engine/Source/Runtime/MassEntity/Private/MassDebuggerBreakpoints.cpp
@@ -485,6 +485,23 @@ namespace UE::Mass::Debug
 		return false;
 	}
 
+	bool FBreakpoint::CheckFragmentRemoveBreakpoints(TConstArrayView<FMassEntityHandle> Entities, const UScriptStruct* FragmentType)
+	{
+		if (LIKELY(!bHasBreakpoint))
+		{
+			return false;
+		}
+
+		for (const FMassEntityHandle& Handle : Entities)
+		{
+			if (CheckFragmentRemoveBreakpoints(Handle, FragmentType))
+			{
+				return true;
+			}
+		}
+		return false;
+	}
+
 	bool FBreakpoint::CheckCreateEntityBreakpoints(TArray<const UScriptStruct*> Fragments)
 	{
 		if (LIKELY(!bHasBreakpoint))
diff --git a/Engine/Source/Runtime/MassEntity/Private/MassEntityManager.cpp b/Engine/Source/Runtime/MassEntity/Private/MassEntityManager.cpp
--- a/Engine/Source/Runtime/MassEntity/Private/MassEntityManager.cpp
+++ b/Engine/Source/Runtime/MassEntity/Private/MassEntityManager.cpp
@@ -1515,18 +1515,15 @@ bool FMassEntityManager::DoesEntityHaveElement(FMassEntityHandle Entity, TNotNul
 
 void FMassEntityManager::BatchAddElementToEntities(TConstArrayView<FMassEntityHandle> Entities, TNotNull<const UScriptStruct*> ElementType)
 {
-	const bool bIsAddableElement = UE::Mass::Private::IsElement(ElementType)
+	testableCheckfReturn(UE::Mass::Private::IsElement(ElementType)
 		&& !UE::Mass::IsA<FMassSharedFragment>(ElementType)
-		&& !UE::Mass::IsA<FMassConstSharedFragment>(ElementType);
-	testableCheckf(bIsAddableElement
+		&& !UE::Mass::IsA<FMassConstSharedFragment>(ElementType)
+		&& !UE::Mass::IsA<FMassChunkFragment>(ElementType)
+		, return
 		, TEXT("%hs: Only tags, fragments, and sparse elements can be added via this API."
-			" Shared fragments require a value - use AddSharedFragmentToEntity"
-			" or AddConstSharedFragmentToEntity instead. Type provided: %s")
+			" Shared fragments require a value, chunk fragments require dedicated APIs."
+			" Type provided: %s")
 		, __FUNCTION__, *ElementType->GetName());
-	if (!bIsAddableElement)
-	{
-		return;
-	}
 
 	TArray<FMassArchetypeEntityCollection> EntityCollections;
 	UE::Mass::Utils::CreateEntityCollections(*this, Entities, FMassArchetypeEntityCollection::EDuplicatesHandling::FoldDuplicates, EntityCollections);
@@ -1547,18 +1544,15 @@ void FMassEntityManager::BatchAddElementToEntities(TConstArrayView<FMassEntityHa
 
 void FMassEntityManager::AddElementToEntity(FMassEntityHandle EntityHandle, TNotNull<const UScriptStruct*> ElementType)
 {
-	const bool bIsAddableElement = UE::Mass::Private::IsElement(ElementType)
+	testableCheckfReturn(UE::Mass::Private::IsElement(ElementType)
 		&& !UE::Mass::IsA<FMassSharedFragment>(ElementType)
-		&& !UE::Mass::IsA<FMassConstSharedFragment>(ElementType);
-	testableCheckf(bIsAddableElement
+		&& !UE::Mass::IsA<FMassConstSharedFragment>(ElementType)
+		&& !UE::Mass::IsA<FMassChunkFragment>(ElementType)
+		, return
 		, TEXT("%hs: Only tags, fragments, and sparse elements can be added via this API."
-			" Shared fragments require a value - use AddSharedFragmentToEntity"
-			" or AddConstSharedFragmentToEntity instead. Type provided: %s")
+			" Shared fragments require a value, chunk fragments require dedicated APIs."
+			" Type provided: %s")
 		, __FUNCTION__, *ElementType->GetName());
-	if (!bIsAddableElement)
-	{
-		return;
-	}
 
 	if (UE::Mass::IsSparse(ElementType))
 	{
@@ -1886,6 +1880,16 @@ void FMassEntityManager::BatchChangeCompositionForEntities(TConstArrayView<FMass
 {
 	CHECK_SYNC_API();
 
+	// Adding shared or sparse element types via this API is not supported — shared fragments require
+	// values (use AddSharedFragmentToEntity), sparse elements have separate storage,
+	// and chunk fragments require dedicated APIs.
+	// Removing shared fragments is supported (extracted below and passed to BatchMoveEntitiesToAnotherArchetype).
+	testableCheckfReturn(ElementsToAdd.ContainsOnlyTagsAndFragments() && ElementsToAdd.GetSparseElements().IsEmpty()
+		, return
+		, TEXT("%hs: ElementsToAdd contains shared, chunk, or sparse element types."
+			" Only non-sparse tags and fragments can be added via this API.")
+		, __FUNCTION__);
+
 	TRACE_CPUPROFILER_EVENT_SCOPE(Mass_BatchChangeCompositionForEntities);
 
 	// Extract shared fragment removal bitsets for BatchMoveEntitiesToAnotherArchetype
diff --git a/Engine/Source/Runtime/MassEntity/Public/MassCommandBuffer.h b/Engine/Source/Runtime/MassEntity/Public/MassCommandBuffer.h
--- a/Engine/Source/Runtime/MassEntity/Public/MassCommandBuffer.h
+++ b/Engine/Source/Runtime/MassEntity/Public/MassCommandBuffer.h
@@ -183,7 +183,7 @@ struct FMassCommandBuffer
 	void AddFragment(FMassEntityHandle Entity)
 	{
 		MASS_STATIC_CHECK_FRAGMENT(T);
-		PushCommand<FMassCommandAddFragmentsInternal<EMassCommandCheckTime::CompileTimeCheck, T>>(Entity);
+		PushCommand<FMassCommandAddElements<T>>(Entity);
 	}
 
 	template<typename T>
@@ -191,7 +191,7 @@ struct FMassCommandBuffer
 	void AddFragment_RuntimeCheck(FMassEntityHandle Entity)
 	{
 		checkf(UE::Mass::IsA<FMassFragment>(T::StaticStruct()), TEXT(MASS_INVALID_FRAGMENT_MSG_F), *T::StaticStruct()->GetName());
-		PushCommand<FMassCommandAddFragmentsInternal<EMassCommandCheckTime::RuntimeCheck, T>>(Entity);
+		PushCommand<FMassCommandAddElements<T>>(Entity);
 	}
 
 	template<typename T>
@@ -214,14 +214,14 @@ struct FMassCommandBuffer
 	void AddTag(FMassEntityHandle Entity)
 	{
 		static_assert(UE::Mass::CTag<T>, "Given struct type is not a valid tag type.");
-		PushCommand<FMassCommandAddTagsInternal<EMassCommandCheckTime::CompileTimeCheck, T>>(Entity);
+		PushCommand<FMassCommandAddElements<T>>(Entity);
 	}
 
 	template<typename T>
 	void AddTag(TConstArrayView<FMassEntityHandle> Entities)
 	{
 		static_assert(UE::Mass::CTag<T>, "Given struct type is not a valid tag type.");
-		PushCommand<FMassCommandAddTagsInternal<EMassCommandCheckTime::CompileTimeCheck, T>>(Entities);
+		PushCommand<FMassCommandAddElements<T>>(Entities);
 	}
 
 	template<typename T>
@@ -229,7 +229,7 @@ struct FMassCommandBuffer
 	void AddTag_RuntimeCheck(FMassEntityHandle Entity)
 	{
 		checkf(UE::Mass::IsA<FMassTag>(T::StaticStruct()), TEXT("Given struct type is not a valid tag type."));
-		PushCommand<FMassCommandAddTagsInternal<EMassCommandCheckTime::RuntimeCheck, T>>(Entity);
+		PushCommand<FMassCommandAddElements<T>>(Entity);
 	}
 
 	/** the convenience function equivalent to calling PushCommand<FMassCommandRemoveTag<T>>(Entity) */
@@ -273,6 +273,32 @@ struct FMassCommandBuffer
 		PushCommand<FMassCommandSwapTagsInternal<EMassCommandCheckTime::RuntimeCheck, TOld, TNew>>(Entity);
 	}
 
+	/** Adds any mix of non-shared element types (fragments, tags, sparse) to an entity. Batches via PushCommand. */
+	template<typename... TTypes>
+	void AddElements(FMassEntityHandle Entity)
+	{
+		PushCommand<FMassCommandAddElements<TTypes...>>(Entity);
+	}
+
+	/** Adds any mix of non-shared element types (fragments, tags, sparse) to entities. Batches via PushCommand. */
+	template<typename... TTypes>
+	void AddElements(TConstArrayView<FMassEntityHandle> Entities)
+	{
+		PushCommand<FMassCommandAddElements<TTypes...>>(Entities);
+	}
+
+	/** Adds any mix of non-shared element types to an entity. Types provided at runtime. Uses PushUniqueCommand. */
+	void AddElements(FMassEntityHandle Entity, TConstArrayView<const UScriptStruct*> ElementTypes)
+	{
+		PushUniqueCommand<FMassCommandAddElementList>(ElementTypes).Add(Entity);
+	}
+
+	/** Adds any mix of non-shared element types to entities. Types provided at runtime. Uses PushUniqueCommand. */
+	void AddElements(TConstArrayView<FMassEntityHandle> Entities, TConstArrayView<const UScriptStruct*> ElementTypes)
+	{
+		PushUniqueCommand<FMassCommandAddElementList>(ElementTypes).Add(Entities);
+	}
+
 	/** Removes any mix of element types (fragments, tags, shared, const shared, sparse) from an entity. Batches via PushCommand. */
 	template<typename... TTypes>
 	void RemoveElements(FMassEntityHandle Entity)
diff --git a/Engine/Source/Runtime/MassEntity/Public/MassCommands.h b/Engine/Source/Runtime/MassEntity/Public/MassCommands.h
--- a/Engine/Source/Runtime/MassEntity/Public/MassCommands.h
+++ b/Engine/Source/Runtime/MassEntity/Public/MassCommands.h
@@ -12,8 +12,88 @@
 #include "MassCommands.generated.h"
 
 /**
- * Enum used by MassBatchCommands to declare their "type". This data is later used to group commands so that command 
- * effects are applied in a controllable fashion 
+ * ========================================================================================================
+ * Mass Commands Quick Reference
+ * ========================================================================================================
+ *
+ * ADDING ELEMENTS (composition-only, no initial values):
+ * -------------------------------------------------------
+ * | Need                           | Command                                    | Batching     | BP  |
+ * |--------------------------------|--------------------------------------------|--------------|-----|
+ * | Add fragments/tags/sparse      | FMassCommandAddElements<T...>              | Cross-chunk  | Yes |
+ * | (compile-time types)           | Convenience: AddElements<T...>()           |              |     |
+ * |--------------------------------|--------------------------------------------|--------------|-----|
+ * | Add fragments/tags/sparse      | FMassCommandAddElementList                 | Per-instance | No  |
+ * | (runtime types)                | Convenience: AddElements(E, T)             |              |     |
+ * |--------------------------------|--------------------------------------------|--------------|-----|
+ * | Add single type (runtime)      | FMassCommandAddElement                     | Per-instance | No  |
+ *
+ * ADDING FRAGMENTS WITH VALUES:
+ * ------------------------------
+ * | Need                           | Command                                    | Batching     | BP  |
+ * |--------------------------------|--------------------------------------------|--------------|-----|
+ * | Add fragments with values      | FMassCommandAddFragmentInstances<T...>     | Cross-chunk  | Yes |
+ * | (sets initial data per entity) |                                            |              |     |
+ *
+ * REMOVING ELEMENTS:
+ * -------------------
+ * | Need                           | Command                                    | Batching     | BP  |
+ * |--------------------------------|--------------------------------------------|--------------|-----|
+ * | Remove any element types       | FMassCommandRemoveElements<T...>           | Cross-chunk  | Yes |
+ * | (compile-time types)           | Convenience: RemoveElements<T...>()        |              |     |
+ * |--------------------------------|--------------------------------------------|--------------|-----|
+ * | Remove any element types       | FMassCommandRemoveElementList              | Per-instance | No  |
+ * | (runtime types)                | Convenience: RemoveElements(E, T)          |              |     |
+ * |--------------------------------|--------------------------------------------|--------------|-----|
+ * | Remove single type (runtime)   | FMassCommandRemoveElement                  | Per-instance | No  |
+ *
+ * CHANGING COMPOSITION:
+ * ----------------------
+ * | Need                           | Command                                    | Batching     | BP  |
+ * |--------------------------------|--------------------------------------------|--------------|-----|
+ * | Swap one tag for another       | FMassCommandSwapTags<TOld, TNew>           | Cross-chunk  | No  |
+ * | (compile-time types)           | Convenience: SwapTags<TOld, TNew>()        |              |     |
+ *
+ * CREATING ENTITIES:
+ * -------------------
+ * | Need                           | Command                                    | Batching     | BP  |
+ * |--------------------------------|--------------------------------------------|--------------|-----|
+ * | Create entity with values      | FMassCommandBuildEntity<T...>              | Cross-chunk  | Yes |
+ * | Create with shared fragments   | FMassCommandBuildEntity                    | Cross-chunk  | Yes |
+ * |                                | WithSharedFragments<TShared, T...>         |              |     |
+ *
+ * DESTROYING ENTITIES:
+ * ---------------------
+ * | Need                           | Command                                    | Batching     | BP  |
+ * |--------------------------------|--------------------------------------------|--------------|-----|
+ * | Destroy entities               | FMassCommandDestroyEntities                | Cross-chunk  | Yes |
+ * |                                | Convenience: DestroyEntity/Entities()      |              |     |
+ *
+ * DEFERRED (LAMBDA):
+ * -------------------
+ * | Need                           | Command                                    | Batching     | BP  |
+ * |--------------------------------|--------------------------------------------|--------------|-----|
+ * | Arbitrary logic at flush time  | FMassDeferredCreateCommand                 | Cross-chunk  | No  |
+ * | (type depends on operation)    | FMassDeferredAddCommand                    |              |     |
+ * |                                | FMassDeferredRemoveCommand                 |              |     |
+ * |                                | FMassDeferredChangeCompositionCommand      |              |     |
+ * |                                | FMassDeferredSetCommand                    |              |     |
+ * |                                | FMassDeferredDestroyCommand                |              |     |
+ *
+ * NOTES:
+ * - Cross-chunk batching: PushCommand reuses one instance across all chunks; one Run() for all entities.
+ * - Per-instance batching: PushUniqueCommand creates a new instance per call. Entities batch via Add().
+ * - Adding shared/const shared fragments requires values; use AddSharedFragmentToEntity directly.
+ * - Removing shared/const shared fragments is supported by RemoveElements and RemoveElement.
+ * - Sparse elements are handled in-place (no archetype move). Non-sparse types are combined into a single entity move.
+ * - BP = Breakpoint support (Mass debugger breakpoints triggered at PushCommand time).
+ * - Deferred commands store a TFunction<void(FMassEntityManager&)>; no type safety, full flexibility.
+ * ========================================================================================================
+ */
+
+/**
+ * Enum used by MassBatchCommands to declare their "type". This data is later used to group commands so that command
+ * effects are applied in a controllable fashion
  * Important: if changed make sure to update FMassCommandBuffer::Flush.CommandTypeOrder as well
  */
 UENUM()
@@ -218,6 +298,17 @@ protected:
 //-----------------------------------------------------------------------------
 // Entity destruction
 //-----------------------------------------------------------------------------
+/**
+ * Destroys entities. Uses PushCommand for cross-chunk batching.
+ *
+ * Example:
+ *   // Single entity:
+ *   Context.Defer().DestroyEntity(EntityHandle);
+ *   // All entities in chunk:
+ *   Context.Defer().DestroyEntities(Context.GetEntities());
+ *
+ * Breakpoint support: Yes — CheckDestroyEntityBreakpoints called at PushCommand time.
+ */
 struct FMassCommandDestroyEntities : public FMassBatchedEntityCommand
 {
 	using Super = FMassBatchedEntityCommand;
@@ -250,8 +341,23 @@ protected:
 // Mixed element types
 //-----------------------------------------------------------------------------
 /**
- * Command adds any element type to entities, be it a fragment or a tag, including sparse element types
- * Note that this type of command can only be added via PushUniqueCommand
+ * Adds a single element type (fragment, tag, or sparse) to entities at runtime.
+ * Requires PushUniqueCommand — each call creates a separate command instance (no cross-call batching).
+ * Shared/const shared fragments are rejected at runtime (they require values).
+ *
+ * When to use: when the element type is only known at runtime and you need to add a single type.
+ * Prefer FMassCommandAddElements<T...> when types are known at compile time (better batching).
+ *
+ * Example:
+ *   // Single entity:
+ *   FMassCommandAddElement& Cmd = Context.Defer().PushUniqueCommand<FMassCommandAddElement>(FMyFragment::StaticStruct());
+ *   Cmd.Add(EntityHandle);
+ *   // All entities in chunk:
+ *   Cmd.Add(Context.GetEntities());
+ *
+ * Supported types: FMassFragment, FMassTag, FMassSparseFragment, FMassSparseTag.
+ * Rejected types: FMassSharedFragment, FMassConstSharedFragment, FMassChunkFragment.
+ * Breakpoint support: No (runtime types — use FMassCommandAddElements<T...> for breakpoints).
  */
 struct FMassCommandAddElement : public FMassBatchedEntityCommand
 {
@@ -282,8 +388,22 @@ struct UE::Mass::Command::TCommandTraits<FMassCommandAddElement> final
 };
 
 /**
- * Command removes any element type from entities, be it a fragment or a tag, including sparse element types
- * Note that this type of command can only be added via PushUniqueCommand
+ * Removes a single element type from entities at runtime.
+ * Requires PushUniqueCommand — each call creates a separate command instance (no cross-call batching).
+ * Supports all element types including shared and const shared fragments.
+ *
+ * When to use: when the element type is only known at runtime and you need to remove a single type.
+ * Prefer FMassCommandRemoveElements<T...> when types are known at compile time (better batching, single entity move).
+ *
+ * Example:
+ *   // Single entity:
+ *   FMassCommandRemoveElement& Cmd = Context.Defer().PushUniqueCommand<FMassCommandRemoveElement>(FMyFragment::StaticStruct());
+ *   Cmd.Add(EntityHandle);
+ *   // All entities in chunk:
+ *   Cmd.Add(Context.GetEntities());
+ *
+ * Supported types: all FMassElement subtypes (fragments, tags, shared, const shared, sparse).
+ * Breakpoint support: No (runtime types — use FMassCommandRemoveElements<T...> for breakpoints).
  */
 struct FMassCommandRemoveElement : public FMassBatchedEntityCommand
 {
@@ -314,9 +434,25 @@ struct UE::Mass::Command::TCommandTraits<FMassCommandRemoveElement> final
 };
 
 /**
- * Runtime command that removes any mix of element types from entities.
+ * Removes any mix of element types from entities at runtime.
  * Element types are provided as UScriptStruct* at construction time.
- * Uses PushUniqueCommand (no cross-call batching, but entities batch via Add()).
+ * Requires PushUniqueCommand — each call creates a separate command instance
+ * (no cross-call batching, but entities batch within a single instance via Add()).
+ * Performs a single entity move for all non-sparse types + in-place sparse removal.
+ *
+ * When to use: when element types are only known at runtime and you need to remove multiple types at once.
+ * Prefer FMassCommandRemoveElements<T...> when types are known at compile time (cross-chunk batching).
+ *
+ * Example:
+ *   const UScriptStruct* Types[] = { FMyFragment::StaticStruct(), FMyTag::StaticStruct() };
+ *   // Single entity:
+ *   FMassCommandRemoveElementList& Cmd = Context.Defer().PushUniqueCommand<FMassCommandRemoveElementList>(Types);
+ *   Cmd.Add(EntityHandle);
+ *   // All entities in chunk via convenience:
+ *   Context.Defer().RemoveElements(Context.GetEntities(), Types);
+ *
+ * Supported types: all FMassElement subtypes (fragments, tags, shared, const shared, sparse).
+ * Breakpoint support: No (runtime types).
  */
 struct FMassCommandRemoveElementList : public FMassBatchedEntityCommand
 {
@@ -386,8 +522,25 @@ struct UE::Mass::Command::TCommandTraits<FMassCommandRemoveElementList> final
 };
 
 /**
- * Templated command that removes any mix of element types (fragments, tags, shared, const shared, sparse)
- * from entities. Uses PushCommand for efficient batching across chunks.
+ * Removes any mix of element types from entities. This is the preferred command for element removal.
+ * Types are template parameters — uses PushCommand for efficient cross-chunk batching
+ * (one Run() for all entities across all chunks that push into this command).
+ * Performs a single entity move for all non-sparse types.
+ *
+ * When to use: whenever you need to remove fragments, tags, shared fragments, and/or sparse elements from
+ * entities and the types are known at compile time. Replaces FMassCommandRemoveFragments and FMassCommandRemoveTags.
+ *
+ * Example:
+ *   // All entities in chunk:
+ *   Context.Defer().PushCommand<FMassCommandRemoveElements<FMyFragment, FMyTag, FMyConstShared>>(Context.GetEntities());
+ *   // Single entity:
+ *   Context.Defer().PushCommand<FMassCommandRemoveElements<FMyFragment, FMyTag>>(EntityHandle);
+ *   // Via convenience:
+ *   Context.Defer().RemoveElements<FMyFragment, FMyTag, FMyConstShared>(Context.GetEntities());
+ *
+ * Supported types: all FMassElement subtypes (fragments, tags, shared, const shared, sparse).
+ * Breakpoint support: Yes — CheckFragmentRemoveBreakpoints called per entity at PushCommand time.
+ * Entity moves: 1 move for all non-sparse types combined. Sparse removal is in-place (no move).
  */
 template<typename... TTypes>
 struct FMassCommandRemoveElements : public FMassBatchedEntityCommand
@@ -401,10 +554,27 @@ struct FMassCommandRemoveElements : public FMassBatchedEntityCommand
 	}
 
 #if WITH_MASSENTITY_DEBUG
-	template<typename T>
-	static bool CheckBreakpoints(T Entity, TTypes... InFragments)
+	// Function pointer forces the non-template overload of CheckFragmentRemoveBreakpoints.
+	// Without it, MSVC deduces TFragments=UScriptStruct* from StaticStruct() return type,
+	// matching the variadic template overload which then calls UScriptStruct*::StaticStruct() and fails.
+	static bool CheckBreakpoints(const FMassEntityHandle Entity)
 	{
-		return UE::Mass::Debug::FBreakpoint::CheckFragmentRemoveBreakpoints(Entity, Forward<TTypes>(InFragments)...);
+		using FCheckFn = bool(*)(const FMassEntityHandle&, const UScriptStruct*);
+		const FCheckFn Check = &UE::Mass::Debug::FBreakpoint::CheckFragmentRemoveBreakpoints;
+		return (Check(Entity, UE::Mass::Clean<TTypes>::StaticStruct()) || ...);
+	}
+	static bool CheckBreakpoints(TConstArrayView<FMassEntityHandle> Entities)
+	{
+		using FCheckFn = bool(*)(const FMassEntityHandle&, const UScriptStruct*);
+		const FCheckFn Check = &UE::Mass::Debug::FBreakpoint::CheckFragmentRemoveBreakpoints;
+		for (const FMassEntityHandle& Entity : Entities)
+		{
+			if ((Check(Entity, UE::Mass::Clean<TTypes>::StaticStruct()) || ...))
+			{
+				return true;
+			}
+		}
+		return false;
 	}
 #endif // WITH_MASSENTITY_DEBUG
 
@@ -455,13 +625,203 @@ protected:
 	TArray<const UScriptStruct*> SparseTypes;
 };
 
+/**
+ * Adds any mix of non-shared element types to entities. This is the preferred command for element addition.
+ * Types are template parameters — uses PushCommand for efficient cross-chunk batching (one Run() for all entities
+ * across all chunks that push into this command).
+ * Performs a single entity move for all non-sparse types.
+ * Shared/const shared fragments cannot be added — they require values (use AddSharedFragmentToEntity).
+ *
+ * When to use: whenever you need to add fragments, tags, and/or sparse elements to entities and the types
+ * are known at compile time. Replaces FMassCommandAddFragments and FMassCommandAddTags.
+ *
+ * Example:
+ *   // All entities in chunk:
+ *   Context.Defer().PushCommand<FMassCommandAddElements<FMyFragment, FMyTag>>(Context.GetEntities());
+ *   // Single entity:
+ *   Context.Defer().PushCommand<FMassCommandAddElements<FMyFragment, FMyTag>>(EntityHandle);
+ *   // Via convenience:
+ *   Context.Defer().AddElements<FMyFragment, FMyTag>(Context.GetEntities());
+ *
+ * Supported types: FMassFragment, FMassTag, FMassSparseFragment, FMassSparseTag.
+ * Rejected types: FMassSharedFragment, FMassConstSharedFragment (compile-time static_assert), FMassChunkFragment (runtime).
+ * Breakpoint support: Yes — CheckFragmentAddBreakpoints called per entity at PushCommand time.
+ * Entity moves: 1 move for all non-sparse types combined. Sparse addition is in-place (no move).
+ */
+template<typename... TTypes>
+struct FMassCommandAddElements : public FMassBatchedEntityCommand
+{
+	static_assert((!TIsDerivedFrom<UE::Mass::Clean<TTypes>, FMassSharedFragment>::Value && ...), "Shared fragments cannot be added via FMassCommandAddElements — they require a value.");
+	static_assert((!TIsDerivedFrom<UE::Mass::Clean<TTypes>, FMassConstSharedFragment>::Value && ...), "Const shared fragments cannot be added via FMassCommandAddElements — they require a value.");
+
+	using Super = FMassBatchedEntityCommand;
+
+	FMassCommandAddElements()
+		: Super(EMassCommandOperationType::Add DEBUG_NAME("AddElements"))
+	{
+		(CategorizeElementType<TTypes>(), ...);
+	}
+
+#if WITH_MASSENTITY_DEBUG
+	// Function pointer forces the non-template overload of CheckFragmentAddBreakpoints.
+	// Without it, MSVC deduces TFragments=UScriptStruct* from StaticStruct() return type,
+	// matching the variadic template overload which then calls UScriptStruct*::StaticStruct() and fails.
+	static bool CheckBreakpoints(const FMassEntityHandle Entity)
+	{
+		using FCheckFn = bool(*)(const FMassEntityHandle&, const UScriptStruct*);
+		const FCheckFn Check = &UE::Mass::Debug::FBreakpoint::CheckFragmentAddBreakpoints;
+		return (Check(Entity, UE::Mass::Clean<TTypes>::StaticStruct()) || ...);
+	}
+	static bool CheckBreakpoints(TConstArrayView<FMassEntityHandle> Entities)
+	{
+		using FCheckFn = bool(*)(const FMassEntityHandle&, const UScriptStruct*);
+		const FCheckFn Check = &UE::Mass::Debug::FBreakpoint::CheckFragmentAddBreakpoints;
+		for (const FMassEntityHandle& Entity : Entities)
+		{
+			if ((Check(Entity, UE::Mass::Clean<TTypes>::StaticStruct()) || ...))
+			{
+				return true;
+			}
+		}
+		return false;
+	}
+#endif // WITH_MASSENTITY_DEBUG
+
+protected:
+	virtual void Run(FMassEntityManager& EntityManager) override
+	{
+		TRACE_CPUPROFILER_EVENT_SCOPE(MassCommandAddElements_Execute);
+
+		TArray<FMassArchetypeEntityCollection> EntityCollections;
+		UE::Mass::Utils::CreateEntityCollections(EntityManager, TargetEntities, FMassArchetypeEntityCollection::FoldDuplicates, EntityCollections);
+
+		if (!NonSparseToAdd.IsEmpty())
+		{
+			EntityManager.BatchChangeCompositionForEntities(EntityCollections, NonSparseToAdd, FMassElementBitSet());
+		}
+		for (const UScriptStruct* SparseType : SparseTypes)
+		{
+			EntityManager.BatchAddSparseElementToEntities(EntityCollections, SparseType);
+		}
+	}
+
+	virtual SIZE_T GetAllocatedSize() const override
+	{
+		return Super::GetAllocatedSize() + NonSparseToAdd.GetAllocatedSize() + SparseTypes.GetAllocatedSize();
+	}
+
+private:
+	template<typename T>
+	void CategorizeElementType()
+	{
+		using FClean = UE::Mass::Clean<T>;
+		if constexpr (UE::Mass::CSparse<FClean>)
+		{
+			SparseTypes.Add(FClean::StaticStruct());
+		}
+		else
+		{
+			NonSparseToAdd.Add(FClean::StaticStruct());
+		}
+	}
+
+	FMassElementBitSet NonSparseToAdd;
+	TArray<const UScriptStruct*> SparseTypes;
+};
+
+/**
+ * Adds any mix of non-shared element types to entities at runtime.
+ * Element types are provided as UScriptStruct* at construction time.
+ * Requires PushUniqueCommand — each call creates a separate command instance
+ * (no cross-call batching, but entities batch within a single instance via Add()).
+ * Performs a single entity move for all non-sparse types + in-place sparse addition.
+ *
+ * When to use: when element types are only known at runtime and you need to add multiple types at once.
+ * Prefer FMassCommandAddElements<T...> when types are known at compile time (cross-chunk batching).
+ *
+ * Example:
+ *   const UScriptStruct* Types[] = { FMyFragment::StaticStruct(), FMyTag::StaticStruct() };
+ *   // Single entity:
+ *   FMassCommandAddElementList& Cmd = Context.Defer().PushUniqueCommand<FMassCommandAddElementList>(Types);
+ *   Cmd.Add(EntityHandle);
+ *   // All entities in chunk via convenience:
+ *   Context.Defer().AddElements(Context.GetEntities(), Types);
+ *
+ * Supported types: FMassFragment, FMassTag, FMassSparseFragment, FMassSparseTag.
+ * Rejected types: shared/const shared and chunk fragments are rejected at runtime by BatchChangeCompositionForEntities (no compile-time check).
+ * Breakpoint support: No (runtime types).
+ */
+struct FMassCommandAddElementList : public FMassBatchedEntityCommand
+{
+	using Super = FMassBatchedEntityCommand;
+
+	FMassCommandAddElementList(TConstArrayView<const UScriptStruct*> InElementTypes)
+		: Super(EMassCommandOperationType::Add DEBUG_NAME("AddElementList"))
+	{
+		for (const UScriptStruct* ElementType : InElementTypes)
+		{
+			if (UE::Mass::IsSparse(ElementType))
+			{
+				SparseTypes.Add(ElementType);
+			}
+			else
+			{
+				NonSparseToAdd.Add(ElementType);
+			}
+		}
+	}
+
+protected:
+	virtual void Run(FMassEntityManager& EntityManager) override
+	{
+		TRACE_CPUPROFILER_EVENT_SCOPE(MassCommandAddElementList_Execute);
+
+		TArray<FMassArchetypeEntityCollection> EntityCollections;
+		UE::Mass::Utils::CreateEntityCollections(EntityManager, TargetEntities, FMassArchetypeEntityCollection::FoldDuplicates, EntityCollections);
+
+		if (!NonSparseToAdd.IsEmpty())
+		{
+			EntityManager.BatchChangeCompositionForEntities(EntityCollections, NonSparseToAdd, FMassElementBitSet());
+		}
+		for (const UScriptStruct* SparseType : SparseTypes)
+		{
+			EntityManager.BatchAddSparseElementToEntities(EntityCollections, SparseType);
+		}
+	}
+
+	virtual SIZE_T GetAllocatedSize() const override
+	{
+		return Super::GetAllocatedSize() + NonSparseToAdd.GetAllocatedSize() + SparseTypes.GetAllocatedSize();
+	}
+
+#if CSV_PROFILER_STATS || WITH_MASSENTITY_DEBUG
+	virtual int32 GetNumOperationsStat() const override
+	{
+		return TargetEntities.Num();
+	}
+#endif // CSV_PROFILER_STATS || WITH_MASSENTITY_DEBUG
+
+	FMassElementBitSet NonSparseToAdd;
+	TArray<const UScriptStruct*> SparseTypes;
+};
+
+template<>
+struct UE::Mass::Command::TCommandTraits<FMassCommandAddElementList> final
+{
+	enum
+	{
+		RequiresUniqueHandling = true
+	};
+};
+
 //-----------------------------------------------------------------------------
 // Simple fragment composition change
 //-----------------------------------------------------------------------------
 template<EMassCommandCheckTime CheckTime, typename... TTypes>
-struct FMassCommandAddFragmentsInternal : public FMassBatchedEntityCommand
+struct UE_DEPRECATED(5.8, "Use FMassCommandAddElements instead — supports any mix of element types in a single operation.") FMassCommandAddFragmentsInternal : public FMassBatchedEntityCommand
 {
 	static_assert((TIsDerivedFrom<UE::Mass::Clean<TTypes>, FMassFragment>::Value && ...), "All types passed to FMassCommandAddFragments must derive from FMassFragment.");
+	static_assert((!UE::Mass::CSparse<UE::Mass::Clean<TTypes>> && ...), "Sparse elements cannot be used with FMassCommandAddFragments. Use FMassCommandAddElements instead.");
 
 	using Super = FMassBatchedEntityCommand;
 	FMassCommandAddFragmentsInternal()
@@ -489,7 +849,7 @@ protected:
 };
 
 template<typename... TTypes>
-using FMassCommandAddFragments = FMassCommandAddFragmentsInternal<EMassCommandCheckTime::CompileTimeCheck, TTypes...>;
+using FMassCommandAddFragments = FMassCommandAddElements<TTypes...>;
 
 template<EMassCommandCheckTime CheckTime, typename... TTypes>
 struct UE_DEPRECATED(5.8, "Use FMassCommandRemoveElements instead — supports any mix of element types and performs a single entity move.") FMassCommandRemoveFragmentsInternal : public FMassBatchedEntityCommand
@@ -560,9 +920,10 @@ protected:
 };
 
 template<EMassCommandCheckTime CheckTime, typename... TTypes>
-struct FMassCommandAddTagsInternal : public FMassCommandChangeTags
+struct UE_DEPRECATED(5.8, "Use FMassCommandAddElements instead — supports any mix of element types in a single operation.") FMassCommandAddTagsInternal : public FMassCommandChangeTags
 {
 	static_assert((TIsDerivedFrom<UE::Mass::Clean<TTypes>, FMassTag>::Value && ...), "All types passed to FMassCommandAddTags must derive from FMassTag.");
+	static_assert((!UE::Mass::CSparse<UE::Mass::Clean<TTypes>> && ...), "Sparse elements cannot be used with FMassCommandAddTags. Use FMassCommandAddElements instead.");
 
 	using Super = FMassCommandChangeTags;
 	FMassCommandAddTagsInternal()
@@ -575,10 +936,10 @@ struct FMassCommandAddTagsInternal : public FMassCommandChangeTags
 };
 
 template<typename T>
-using FMassCommandAddTag = FMassCommandAddTagsInternal<EMassCommandCheckTime::CompileTimeCheck, T>;
+using FMassCommandAddTag = FMassCommandAddElements<T>;
 
 template<typename... TTypes>
-using FMassCommandAddTags = FMassCommandAddTagsInternal<EMassCommandCheckTime::CompileTimeCheck, TTypes...>;
+using FMassCommandAddTags = FMassCommandAddElements<TTypes...>;
 
 template<EMassCommandCheckTime CheckTime, typename... TTypes>
 struct UE_DEPRECATED(5.8, "Use FMassCommandRemoveElements instead — supports any mix of element types and performs a single entity move.") FMassCommandRemoveTagsInternal : public FMassCommandChangeTags
@@ -609,6 +970,21 @@ using FMassCommandRemoveTag = FMassCommandRemoveElements<T>;
 template<typename... TTypes>
 using FMassCommandRemoveTags = FMassCommandRemoveElements<TTypes...>;
 
+/**
+ * Atomically swaps one tag for another on entities.
+ * Removes TOld and adds TNew in a single composition change.
+ * Uses PushCommand for cross-chunk batching.
+ *
+ * Example:
+ *   // All entities in chunk:
+ *   Context.Defer().SwapTags<FOldTag, FNewTag>(Context.GetEntities());
+ *   // Single entity:
+ *   Context.Defer().SwapTags<FOldTag, FNewTag>(EntityHandle);
+ *
+ * Supported types: FMassTag subtypes only (compile-time static_assert).
+ * Breakpoint support: No.
+ * Entity moves: 1 move (both tag changes applied together).
+ */
 template<EMassCommandCheckTime CheckTime, typename TOld, typename TNew>
 struct FMassCommandSwapTagsInternal : public FMassCommandChangeTags
 {
@@ -630,10 +1006,26 @@ using FMassCommandSwapTags = FMassCommandSwapTagsInternal<EMassCommandCheckTime:
 //-----------------------------------------------------------------------------
 // Struct Instances adding and setting
 //-----------------------------------------------------------------------------
+/**
+ * Adds fragments with initial values to entities. Each entity receives typed fragment data.
+ * Uses PushCommand for cross-chunk batching. Sparse fragments are supported (handled internally).
+ *
+ * When to use: when you need to add fragments AND set their initial values in one operation.
+ * For composition-only changes (no values), prefer FMassCommandAddElements<T...>.
+ *
+ * Example:
+ *   // Add with values — entity + fragment instances passed together:
+ *   Context.Defer().PushCommand<FMassCommandAddFragmentInstances>(EntityHandle, FMyFragmentA(42), FMyFragmentB(3.14f));
+ *
+ * Supported types: FMassFragment subtypes (including sparse fragments).
+ * Breakpoint support: Yes — CheckFragmentAddBreakpoints called at PushCommand time.
+ */
 template<typename... TOthers>
 struct FMassCommandAddFragmentInstances : public FMassBatchedEntityCommand
 {
 	static_assert((TIsDerivedFrom<UE::Mass::Clean<TOthers>, FMassFragment>::Value && ...), "All types passed to FMassCommandAddFragmentInstances must derive from FMassFragment.");
+	// Note: sparse fragments are intentionally allowed here
+	// BatchAddFragmentInstancesForEntities handles them by splitting sparse elements out and dispatching to BatchAddSparseElementToEntities internally.
 
 	using Super = FMassBatchedEntityCommand;
 
@@ -687,6 +1079,18 @@ protected:
 	const FMassFragmentBitSet FragmentsAffected;
 };
 
+/**
+ * Builds reserved entities with fragment values.
+ * Entities must be pre-reserved via EntityManager->ReserveEntity().
+ * Uses PushCommand for cross-chunk batching.
+ *
+ * Example:
+ *   FMassEntityHandle Entity = EntityManager->ReserveEntity();
+ *   Context.Defer().PushCommand<FMassCommandBuildEntity>(Entity, FMyFragmentA(42), FMyFragmentB(3.14f));
+ *
+ * Supported types: FMassFragment subtypes for TOthers.
+ * Breakpoint support: Yes — CheckCreateEntityBreakpoints called at PushCommand time.
+ */
 template<typename... TOthers>
 struct FMassCommandBuildEntity : public FMassCommandAddFragmentInstances<TOthers...>
 {
@@ -726,10 +1130,23 @@ protected:
 	}
 };
 
-/** 
- * Note: that TSharedFragmentValues is always expected to be FMassArchetypeSharedFragmentValues, but is declared as 
- *	template's param to maintain uniform command adding interface via FMassCommandBuffer.PushCommand. 
- *	PushCommands received all input params in one `typename...` list and as such cannot be easily split up to reason about.
+/**
+ * Builds reserved entities with fragment values AND shared fragment values.
+ * Entities with the same shared fragment values are grouped for efficient archetype creation.
+ * Uses PushCommand for cross-chunk batching.
+ *
+ * Note: TSharedFragmentValues is always FMassArchetypeSharedFragmentValues but is a template param
+ * to maintain uniform PushCommand interface (all params in one typename... list).
+ *
+ * Example:
+ *   FMassEntityHandle Entity = EntityManager->ReserveEntity();
+ *   FMassArchetypeSharedFragmentValues SharedValues;
+ *   SharedValues.Add(EntityManager->GetOrCreateConstSharedFragment(FMySharedFragment(42)));
+ *   Context.Defer().PushCommand<FMassCommandBuildEntityWithSharedFragments>(
+ *       Entity, MoveTemp(SharedValues), FMyFragmentA(1), FMyFragmentB(2.0f));
+ *
+ * Supported types: FMassFragment subtypes for TOthers. Shared fragments via FMassArchetypeSharedFragmentValues.
+ * Breakpoint support: Yes — CheckCreateEntityBreakpoints called at PushCommand time.
  */
 template<typename TSharedFragmentValues, typename... TOthers>
 struct FMassCommandBuildEntityWithSharedFragments : public FMassBatchedCommand
@@ -848,6 +1265,29 @@ protected:
 //-----------------------------------------------------------------------------
 // Commands that really can't know the types at compile time
 //-----------------------------------------------------------------------------
+/**
+ * Stores a lambda to execute at flush time.
+ * Maximum flexibility, no type safety.
+ * The OpType template parameter controls execution ordering relative to other commands.
+ * Multiple lambdas of the same OpType batch into a single command instance.
+ *
+ * When to use: as a last resort when no typed command fits the use case, or for complex
+ * multi-step operations that can't be expressed as a single composition change.
+ * Prefer typed commands (AddElements, RemoveElements, etc.) when possible.
+ *
+ * Example:
+ *   Context.Defer().PushCommand<FMassDeferredSetCommand>([Entity](FMassEntityManager& Manager) {
+ *       FMassEntityView View(Manager, Entity);
+ *       View.GetFragmentData<FMyFragment>().Value = 42;
+ *   });
+ *
+ * Variants (by execution order):
+ *   FMassDeferredCreateCommand, FMassDeferredAddCommand, FMassDeferredRemoveCommand,
+ *   FMassDeferredChangeCompositionCommand, FMassDeferredSetCommand, FMassDeferredDestroyCommand.
+ *
+ * Breakpoint support: No.
+ * No type safety. No validation.
+ */
 template<EMassCommandOperationType OpType>
 struct FMassDeferredCommand : public FMassBatchedCommand
 {
