[MassEntity] FMassExecutionContext: pimpl SubsystemAccess, explicit copy/move semantics

pimpl FMassSubsystemAccess behind TUniquePtr to decouple MassExecutionContext.h from subsystem internals - prerequisite for the FMassSystem migration but independently useful for compile times and encapsulation.

FMassSubsystemAccess now stores FMassEntityManager* instead of TWeakObjectPtr<UWorld>, eliminating the World???EntityManager round-trip in CacheSubsystem/FindSubsystemIndex. raw pointer is safe because SubsystemAccess lifetime is bounded by FMassExecutionContext which holds TSharedRef<FMassEntityManager>.

FMassExecutionContext gets explicit operator= that intentionally skips QueriesStack and processing callbacks (they're managed by callers, not inherited on copy). move ctor explicitly defaulted. copy ctor simplified (no longer needs QueriesStack.Reset() since operator= doesn't copy it).

include cleanup: forward-declare UWorld/UMassProcessor instead of pulling full headers into MassExecutionContext.h.

#Synced-CL 52534678

diff --git a/Engine/Source/Developer/MassEntityTestSuite/Private/MassSystemRequirementTest.cpp b/Engine/Source/Developer/MassEntityTestSuite/Private/MassSystemRequirementTest.cpp
--- a/Engine/Source/Developer/MassEntityTestSuite/Private/MassSystemRequirementTest.cpp
+++ b/Engine/Source/Developer/MassEntityTestSuite/Private/MassSystemRequirementTest.cpp
@@ -51,8 +51,8 @@ struct FGenericSubsystemAccessAPI : FSystemRequirementTestBase
 	{
 		UWorld* World = FAITestHelpers::GetWorld();
 		check(World);
-		FMassSubsystemAccess SubsystemAccessBits(World);
-		FMassSubsystemAccess SubsystemAccessClasses(World);
+		FMassSubsystemAccess SubsystemAccessBits(EntityManager.Get());
+		FMassSubsystemAccess SubsystemAccessClasses(EntityManager.Get());
 
 		EntityQuery.AddSubsystemRequirement<UMassTestWorldSubsystem>(EMassFragmentAccess::ReadWrite);
 		EntityQuery.AddSubsystemRequirement<UMassTestEngineSubsystem>(EMassFragmentAccess::ReadWrite);
@@ -122,8 +122,8 @@ struct FGameSubsystemAccessAPI : FSystemRequirementTestBase
 			return true;
 		}
 
-		FMassSubsystemAccess SubsystemAccessBits(World);
-		FMassSubsystemAccess SubsystemAccessClasses(World);
+		FMassSubsystemAccess SubsystemAccessBits(EntityManager.Get());
+		FMassSubsystemAccess SubsystemAccessClasses(EntityManager.Get());
 
 		EntityQuery.AddSubsystemRequirement<UMassTestLocalPlayerSubsystem>(EMassFragmentAccess::ReadWrite);
 		EntityQuery.AddSubsystemRequirement<UMassTestGameInstanceSubsystem>(EMassFragmentAccess::ReadWrite);
diff --git a/Engine/Source/Runtime/MassEntity/Private/MassExecutionContext.cpp b/Engine/Source/Runtime/MassEntity/Private/MassExecutionContext.cpp
--- a/Engine/Source/Runtime/MassEntity/Private/MassExecutionContext.cpp
+++ b/Engine/Source/Runtime/MassEntity/Private/MassExecutionContext.cpp
@@ -4,6 +4,8 @@
 #include "MassArchetypeData.h"
 #include "MassEntityManager.h"
 #include "MassEntityLinkFragments.h"
+#include "MassSubsystemAccess.h"
+#include "MassProcessor.h"
 #if WITH_MASSENTITY_DEBUG
 #include "MassDebugger.h"
 #endif // WITH_MASSENTITY_DEBUG
@@ -13,7 +15,7 @@
 // FMassExecutionContext
 //-----------------------------------------------------------------------------
 FMassExecutionContext::FMassExecutionContext(FMassEntityManager& InEntityManager, const float InDeltaTimeSeconds, const bool bInFlushDeferredCommands)
-	: SubsystemAccess(InEntityManager.GetWorld())
+	: SubsystemAccessPtr(MakeUnique<FMassSubsystemAccess>(&InEntityManager))
 	, DeltaTimeSeconds(InDeltaTimeSeconds)
 	, EntityManager(InEntityManager.AsShared())
 	, bFlushDeferredCommands(bInFlushDeferredCommands)
@@ -23,10 +25,9 @@ FMassExecutionContext::FMassExecutionContext(FMassEntityManager& InEntityManager
 FMassExecutionContext::FMassExecutionContext(const FMassExecutionContext& Other)
 	: EntityManager(Other.EntityManager)
 {
-	// we're using assignment operator here as a setup helper, as per operator='s comment.
+	// Delegate to operator= so copy logic lives in one place; operator= intentionally
+	// excludes QueriesStack and processing callbacks (see comment in operator= body).
 	*this = Other;
-
-	QueriesStack.Reset();
 }
 
 FMassExecutionContext::FMassExecutionContext(const FMassExecutionContext& Other, FMassEntityQuery& Query, const TSharedPtr<FMassCommandBuffer>& InCommandBuffer)
@@ -37,6 +38,56 @@ FMassExecutionContext::FMassExecutionContext(const FMassExecutionContext& Other,
 	SetDeferredCommandBuffer(InCommandBuffer);
 }
 
+// Hand-written because TUniquePtr<FMassSubsystemAccess> is not copyable (need deep copy via
+// MakeUnique), and because we intentionally skip QueriesStack and processing callbacks — callers
+// manage those explicitly after copy. Every data member of FMassExecutionContext must appear here;
+// if you add a member to the struct, add it to this operator too.
+// @todo MassCore migration: once FMassSubsystemAccess no longer needs to be pimpl'd behind
+// TUniquePtr (i.e. once the MassSubsystemAccess.h include can be removed from this header),
+// revert to an inline member and = default this operator.
+FMassExecutionContext& FMassExecutionContext::operator=(const FMassExecutionContext& Other)
+{
+	if (this != &Other)
+	{
+		LinkedEntityView = Other.LinkedEntityView;
+		CachedIndirectEntityView = Other.CachedIndirectEntityView;
+		SubsystemAccessPtr = MakeUnique<FMassSubsystemAccess>(*Other.SubsystemAccessPtr);
+		DeferredCommandBuffer = Other.DeferredCommandBuffer;
+		EntityListView = Other.EntityListView;
+		EntityCollection = Other.EntityCollection;
+		AuxData = Other.AuxData;
+		DeltaTimeSeconds = Other.DeltaTimeSeconds;
+		CurrentArchetypeCompositionBitSet = Other.CurrentArchetypeCompositionBitSet;
+#if WITH_MASSENTITY_DEBUG
+		DebugColor = Other.DebugColor;
+#endif
+		EntityManager = Other.EntityManager;
+		// Note: QueriesStack is intentionally not copied — callers manage it explicitly.
+		// Note: PreQueryCallbackFunc, PostQueryCallbackFunc, and ChunkQueueFunc are intentionally
+		// not copied — they are set by the processing queue on the copy after construction
+		// (see FProcessorWorker::ExecuteProcessor in MassProcessingQueueTypes.cpp).
+		IteratorSerialNumberGenerator = Other.IteratorSerialNumberGenerator;
+		ExecutionType = Other.ExecutionType;
+		bFlushDeferredCommands = Other.bFlushDeferredCommands;
+		bHasSparseRequirements = Other.bHasSparseRequirements;
+		ArchetypeDataPtr = Other.ArchetypeDataPtr;
+		ChunkEntityIndex = Other.ChunkEntityIndex;
+#if WITH_MASSENTITY_DEBUG
+		DebugExecutionDescription = Other.DebugExecutionDescription;
+		DebugProcessor = Other.DebugProcessor;
+#endif
+		ActiveEntityCreationRequirements = Other.ActiveEntityCreationRequirements;
+		FragmentViews = Other.FragmentViews;
+		ChunkFragmentViews = Other.ChunkFragmentViews;
+		ConstSharedFragmentViews = Other.ConstSharedFragmentViews;
+		SharedFragmentViews = Other.SharedFragmentViews;
+		RequestedSparseElementsAccess = Other.RequestedSparseElementsAccess;
+	}
+	return *this;
+}
+
+FMassExecutionContext::FMassExecutionContext(FMassExecutionContext&& Other) = default;
+
 FMassExecutionContext::~FMassExecutionContext()
 {
 	ensureMsgf(QueriesStack.Num() == 0, TEXT("Destroying a FMassExecutionContext instance while not all queries have been popped is unexpected."));
@@ -75,7 +126,7 @@ void FMassExecutionContext::SetCurrentArchetype(FMassArchetypeData& Archetype)
 
 bool FMassExecutionContext::CacheSubsystemRequirements(const FMassSubsystemRequirements& SubsystemRequirements)
 {
-	return SubsystemAccess.CacheSubsystemRequirements(SubsystemRequirements);
+	return SubsystemAccessPtr->CacheSubsystemRequirements(SubsystemRequirements);
 }
 
 void FMassExecutionContext::SetEntityCollection(const FMassArchetypeEntityCollection& InEntityCollection)
diff --git a/Engine/Source/Runtime/MassEntity/Private/MassSubsystemAccess.cpp b/Engine/Source/Runtime/MassEntity/Private/MassSubsystemAccess.cpp
--- a/Engine/Source/Runtime/MassEntity/Private/MassSubsystemAccess.cpp
+++ b/Engine/Source/Runtime/MassEntity/Private/MassSubsystemAccess.cpp
@@ -1,6 +1,7 @@
 // Copyright Epic Games, Inc. All Rights Reserved.
 
 #include "MassSubsystemAccess.h"
+#include "MassEntityManager.h"
 
 
 namespace UE::Mass::Private
@@ -15,11 +16,26 @@ namespace UE::Mass::Private
 //-----------------------------------------------------------------------------
 // FMassSubsystemAccess
 //-----------------------------------------------------------------------------
+FMassSubsystemAccess::FMassSubsystemAccess(FMassEntityManager* InEntityManager)
+	: EntityManager(InEntityManager)
+{
+	Subsystems.AddZeroed(FMassExternalSubsystemBitSet::GetMaxNum());
+}
+
+PRAGMA_DISABLE_DEPRECATION_WARNINGS
 FMassSubsystemAccess::FMassSubsystemAccess(UWorld* InWorld)
-	: World(InWorld)
+	: EntityManager(nullptr)
 {
+	ensureMsgf(InWorld == nullptr, TEXT("FMassSubsystemAccess(UWorld*) is deprecated and no longer functional. "
+		"Use FMassSubsystemAccess(FMassEntityManager*) instead. Subsystem caching will not work."));
 	Subsystems.AddZeroed(FMassExternalSubsystemBitSet::GetMaxNum());
 }
+PRAGMA_ENABLE_DEPRECATION_WARNINGS
+
+UWorld* FMassSubsystemAccess::GetWorld() const
+{
+	return EntityManager ? EntityManager->GetWorld() : nullptr;
+}
 
 USubsystem* FMassSubsystemAccess::FetchSubsystemInstance(UWorld* World, TSubclassOf<USubsystem> SubsystemClass)
 {
@@ -28,8 +44,8 @@ USubsystem* FMassSubsystemAccess::FetchSubsystemInstance(UWorld* World, TSubclas
 	check(SubsystemClass);
 	if (SubsystemClass->IsChildOf<UWorldSubsystem>())
 	{
-		return World 
-			? World->GetSubsystemBase(UE::Mass::Private::ConvertToSubsystemClass<UWorldSubsystem>(SubsystemClass)) 
+		return World
+			? World->GetSubsystemBase(UE::Mass::Private::ConvertToSubsystemClass<UWorldSubsystem>(SubsystemClass))
 			: nullptr;
 	}
 	if (SubsystemClass->IsChildOf<UEngineSubsystem>())
@@ -104,7 +120,7 @@ bool FMassSubsystemAccess::CacheSubsystem(const uint32 SystemIndex)
 
 	if (SubsystemSubclass)
 	{
-		USubsystem* SystemInstance = FetchSubsystemInstance(World.Get(), SubsystemSubclass);
+		USubsystem* SystemInstance = FetchSubsystemInstance(GetWorld(), SubsystemSubclass);
 		Subsystems[SystemIndex] = SystemInstance;
 
 		if (SystemInstance == nullptr)
diff --git a/Engine/Source/Runtime/MassEntity/Public/MassExecutionContext.h b/Engine/Source/Runtime/MassEntity/Public/MassExecutionContext.h
--- a/Engine/Source/Runtime/MassEntity/Public/MassExecutionContext.h
+++ b/Engine/Source/Runtime/MassEntity/Public/MassExecutionContext.h
@@ -11,10 +11,15 @@
 #include "MassEntityQuery.h"
 #include "MassArchetypeTypes.h"
 #include "MassSubsystemAccess.h"
+#if UE_ENABLE_INCLUDE_ORDER_DEPRECATED_IN_5_8
 #include "MassProcessor.h"
+#endif
 #include "MassEntityCollection.h"
 #include "MassEntityView.h"
 
+#if WITH_MASSENTITY_DEBUG
+#include "MassProcessor.h"  // full definition needed for TWeakObjectPtr<UMassProcessor> debug member
+#endif
 
 #define MASS_CHECK_IF_VALID(View, Type) \
 	checkf(View \
@@ -29,6 +34,8 @@
 	checkf(RequestedSparseElementsAccess.Contains(ElementType) \
 			, TEXT("%hs: %s not requested as requirement by the query"), __FUNCTION__, *ElementType->GetName())
 
+class UWorld;
+class UMassProcessor;
 struct FMassEntityQuery;
 struct FMassEntityLinkFragment;
 struct FMassArchetypeChunk;
@@ -81,9 +88,9 @@ private:
 	/** Used for indirect fragment access */
 	FMassEntityView CachedIndirectEntityView;
 
-	FMassSubsystemAccess SubsystemAccess;
-	
-	// mz@todo make this shared ptr thread-safe and never auto-flush in MT environment. 
+	TUniquePtr<FMassSubsystemAccess> SubsystemAccessPtr;
+
+	// mz@todo make this shared ptr thread-safe and never auto-flush in MT environment.
 	TSharedPtr<FMassCommandBuffer> DeferredCommandBuffer;
 	TArrayView<FMassEntityHandle> EntityListView;
 	
@@ -175,25 +182,26 @@ private:
 
 	void GetSubsystemRequirementBits(FMassExternalSubsystemBitSet& OutConstSubsystemsBitSet, FMassExternalSubsystemBitSet& OutMutableSubsystemsBitSet)
 	{
-		SubsystemAccess.GetSubsystemRequirementBits(OutConstSubsystemsBitSet, OutMutableSubsystemsBitSet);
+		SubsystemAccessPtr->GetSubsystemRequirementBits(OutConstSubsystemsBitSet, OutMutableSubsystemsBitSet);
 	}
 
 	void SetSubsystemRequirementBits(const FMassExternalSubsystemBitSet& InConstSubsystemsBitSet, const FMassExternalSubsystemBitSet& InMutableSubsystemsBitSet)
 	{
-		SubsystemAccess.SetSubsystemRequirementBits(InConstSubsystemsBitSet, InMutableSubsystemsBitSet);
+		SubsystemAccessPtr->SetSubsystemRequirementBits(InConstSubsystemsBitSet, InMutableSubsystemsBitSet);
 	}
 
 	MASSENTITY_API FStructView GetSparseElementInternal(TNotNull<const UScriptStruct*> ElementType, const FMassEntityHandle EntityHandle) const;
 
 	/**
-	 * Note that this operator is private on purpose, used to simplify the implementation of constructors.
-	 * The context itself does not support assignment.
+	 * Explicit assignment used by copy ctor and internal helpers. Intentionally does NOT copy
+	 * QueriesStack (callers manage it explicitly) or processing callbacks (set by FProcessorWorker
+	 * on the copy after construction).
 	 */
-	FMassExecutionContext& operator=(const FMassExecutionContext& Other) = default;
+	MASSENTITY_API FMassExecutionContext& operator=(const FMassExecutionContext& Other);
 
 public:
 	MASSENTITY_API explicit FMassExecutionContext(FMassEntityManager& InEntityManager, const float InDeltaTimeSeconds = 0.f, const bool bInFlushDeferredCommands = true);
-	FMassExecutionContext(FMassExecutionContext&& Other) = default;
+	MASSENTITY_API FMassExecutionContext(FMassExecutionContext&& Other);
 	MASSENTITY_API FMassExecutionContext(const FMassExecutionContext& Other);
 	MASSENTITY_API FMassExecutionContext(const FMassExecutionContext& Other, FMassEntityQuery& Query, const TSharedPtr<FMassCommandBuffer>& InCommandBuffer = {});
 	MASSENTITY_API ~FMassExecutionContext();
@@ -723,49 +731,49 @@ public:
 	template<typename T, typename = typename TEnableIf<TIsDerivedFrom<T, USubsystem>::IsDerived>::Type>
 	T* GetMutableSubsystem()
 	{
-		return SubsystemAccess.GetMutableSubsystem<T>();
+		return SubsystemAccessPtr->GetMutableSubsystem<T>();
 	}
 
 	template<typename T, typename = typename TEnableIf<TIsDerivedFrom<T, USubsystem>::IsDerived>::Type>
 	T& GetMutableSubsystemChecked()
 	{
-		return SubsystemAccess.GetMutableSubsystemChecked<T>();
+		return SubsystemAccessPtr->GetMutableSubsystemChecked<T>();
 	}
 
 	template<typename T, typename = typename TEnableIf<TIsDerivedFrom<T, USubsystem>::IsDerived>::Type>
 	const T* GetSubsystem()
 	{
-		return SubsystemAccess.GetSubsystem<T>();
+		return SubsystemAccessPtr->GetSubsystem<T>();
 	}
 
 	template<typename T, typename = typename TEnableIf<TIsDerivedFrom<T, USubsystem>::IsDerived>::Type>
 	const T& GetSubsystemChecked()
 	{
-		return SubsystemAccess.GetSubsystemChecked<T>();
+		return SubsystemAccessPtr->GetSubsystemChecked<T>();
 	}
 
 	template<typename T, typename = typename TEnableIf<TIsDerivedFrom<T, USubsystem>::IsDerived>::Type>
 	T* GetMutableSubsystem(const TSubclassOf<USubsystem> SubsystemClass)
 	{
-		return SubsystemAccess.GetMutableSubsystem<T>(SubsystemClass);
+		return SubsystemAccessPtr->GetMutableSubsystem<T>(SubsystemClass);
 	}
 
 	template<typename T, typename = typename TEnableIf<TIsDerivedFrom<T, USubsystem>::IsDerived>::Type>
 	T& GetMutableSubsystemChecked(const TSubclassOf<USubsystem> SubsystemClass)
 	{
-		return SubsystemAccess.GetMutableSubsystemChecked<T>(SubsystemClass);
+		return SubsystemAccessPtr->GetMutableSubsystemChecked<T>(SubsystemClass);
 	}
 
 	template<typename T, typename = typename TEnableIf<TIsDerivedFrom<T, USubsystem>::IsDerived>::Type>
 	const T* GetSubsystem(const TSubclassOf<USubsystem> SubsystemClass)
 	{
-		return SubsystemAccess.GetSubsystem<T>(SubsystemClass);
+		return SubsystemAccessPtr->GetSubsystem<T>(SubsystemClass);
 	}
 
 	template<typename T, typename = typename TEnableIf<TIsDerivedFrom<T, USubsystem>::IsDerived>::Type>
 	const T& GetSubsystemChecked(const TSubclassOf<USubsystem> SubsystemClass)
 	{
-		return SubsystemAccess.GetSubsystemChecked<T>(SubsystemClass);
+		return SubsystemAccessPtr->GetSubsystemChecked<T>(SubsystemClass);
 	}
 
 	//-----------------------------------------------------------------------------
@@ -915,7 +923,7 @@ public:
 protected:
 	void SetSubsystemRequirements(const FMassSubsystemRequirements& SubsystemRequirements)
 	{
-		SubsystemAccess.SetSubsystemRequirements(SubsystemRequirements);
+		SubsystemAccessPtr->SetSubsystemRequirements(SubsystemRequirements);
 	}
 
 	void SetFragmentRequirements(const FMassFragmentRequirements& FragmentRequirements);
diff --git a/Engine/Source/Runtime/MassEntity/Public/MassSubsystemAccess.h b/Engine/Source/Runtime/MassEntity/Public/MassSubsystemAccess.h
--- a/Engine/Source/Runtime/MassEntity/Public/MassSubsystemAccess.h
+++ b/Engine/Source/Runtime/MassEntity/Public/MassSubsystemAccess.h
@@ -18,10 +18,14 @@
 #include "Engine/Engine.h"
 #endif // WITH_EDITOR
 
+struct FMassEntityManager;
 
 struct FMassSubsystemAccess
 {
-	MASSENTITY_API explicit FMassSubsystemAccess(UWorld* InWorld = nullptr);
+	MASSENTITY_API explicit FMassSubsystemAccess(FMassEntityManager* InEntityManager = nullptr);
+
+	UE_DEPRECATED(5.8, "Use FMassSubsystemAccess(FMassEntityManager*) constructor instead.")
+	MASSENTITY_API explicit FMassSubsystemAccess(UWorld* InWorld);
 
 	//-----------------------------------------------------------------------------
 	// Statically-typed subsystems
@@ -203,7 +207,7 @@ protected:
 		{
 			if constexpr (DoesRequireWorld<T>())
 			{
-				SystemInstance = FetchSubsystemInstance<std::remove_const_t<T>>(World.Get());
+				SystemInstance = FetchSubsystemInstance<std::remove_const_t<T>>(GetWorld());
 			}
 			else
 			{
@@ -226,16 +230,21 @@ protected:
 		USubsystem* SystemInstance = (T*)Subsystems[SystemIndex];
 		if (SystemInstance == nullptr)
 		{
-			SystemInstance = FetchSubsystemInstance(World.Get(), SubsystemClass);
+			SystemInstance = FetchSubsystemInstance(GetWorld(), SubsystemClass);
 			Subsystems[SystemIndex] = SystemInstance;
 		}
 		return Cast<T>(SystemInstance);
 	}
 
+	MASSENTITY_API UWorld* GetWorld() const;
 	MASSENTITY_API bool CacheSubsystem(const uint32 SystemIndex);
 
 	FMassExternalSubsystemBitSet ConstSubsystemsBitSet;
 	FMassExternalSubsystemBitSet MutableSubsystemsBitSet;
 	TArray<USubsystem*> Subsystems;
-	TWeakObjectPtr<UWorld> World;
+
+	/** Non-owning. Lifetime is guaranteed by the owning FMassExecutionContext which holds
+	 *  a TSharedRef<FMassEntityManager>. Raw pointer avoids atomic ref-count overhead on
+	 *  copy (FMassSubsystemAccess is deep-copied when FMassExecutionContext is copied). */
+	FMassEntityManager* EntityManager = nullptr;
 };
