[FMassEntityManager]
Wrap atomic increment and decrement of the MassEntity scoped processing counter in AutoRTFM::Open to ensure correctness within transactional memory contexts

#Synced-CL 52680077

diff --git a/Engine/Source/Runtime/MassEntity/Public/MassEntityManager.h b/Engine/Source/Runtime/MassEntity/Public/MassEntityManager.h
--- a/Engine/Source/Runtime/MassEntity/Public/MassEntityManager.h
+++ b/Engine/Source/Runtime/MassEntity/Public/MassEntityManager.h
@@ -111,11 +111,22 @@ public:
 	{
 		explicit FScopedProcessing(std::atomic<int32>& InProcessingScopeCount) : ScopedProcessingCount(InProcessingScopeCount)
 		{
-			++ScopedProcessingCount;
+			UE_AUTORTFM_OPEN
+			{
+				ScopedProcessingCount.fetch_add(1, std::memory_order_relaxed);
+			};
+			AutoRTFM::PushOnAbortHandler(this, [this]
+				{
+					ScopedProcessingCount.fetch_sub(1, std::memory_order_relaxed);
+				});
 		}
 		~FScopedProcessing()
 		{
-			--ScopedProcessingCount;
+			UE_AUTORTFM_OPEN
+			{
+				ScopedProcessingCount.fetch_sub(1, std::memory_order_relaxed);
+			};
+			AutoRTFM::PopOnAbortHandler(this);
 		}
 	private:
 		std::atomic<int32>& ScopedProcessingCount;
