[MassComands]
Extend fragment instance and build commands to support mixed element types in a single entity move

New commands:
- Add FMassCommandAddElementsWithSharedFragments<TTypes...> for composition-only changes + shared fragment values in one move
- Add FMassCommandAddFragmentInstancesWithSharedFragments<TShared, T...> for fragment values + shared values in one move
- Add convenience methods on FMassCommandBuffer: AddElementsWithSharedFragments, AddFragmentInstancesWithSharedFragments

Extend existing commands to accept tags and sparse types:
- FMassCommandAddFragmentInstances<T...> now accepts FMassTag, FMassSparseTag, FMassSparseFragment alongside FMassFragment
- FMassCommandBuildEntity<T...> same mixed-type support
- FMassCommandBuildEntityWithSharedFragments<TShared, T...> same mixed-type support
- Tags included in composition, excluded from value payload (GetAsNonTagGenericMultiArray)
- Sparse types handled in-place via BatchAddSparseElementToEntities before composition changes

EntityManager API additions:
- BatchChangeCompositionForEntities overload with FMassArchetypeSharedFragmentValues
- BatchAddFragmentInstancesForEntities overloads with FMassTagBitSet and/or FMassArchetypeSharedFragmentValues

Callsite migrations (2 entity moves -> 1): AddFragmentInstances + deferred shared -> AddFragmentInstancesWithSharedFragments

Tests:
- Correctness tests for all command variants with fragment + tag + sparse fragment + sparse tag + shared combinations
- Performance benchmarks: separate vs mixed commands
- Parity benchmarks: single-type vs multi-type overhead validation

#Synced-CL 52708313

diff --git a/Engine/Source/Developer/MassEntityTestSuite/Private/MassCommandsTest.cpp b/Engine/Source/Developer/MassEntityTestSuite/Private/MassCommandsTest.cpp
--- a/Engine/Source/Developer/MassEntityTestSuite/Private/MassCommandsTest.cpp
+++ b/Engine/Source/Developer/MassEntityTestSuite/Private/MassCommandsTest.cpp
@@ -6,8 +6,6 @@
 #include "MassEntityTypes.h"
 #include "MassEntityView.h"
 #include "MassExecutionContext.h"
-#include "MassProcessingTypes.h"
-#include "Algo/Sort.h"
 #include "Algo/RandomShuffle.h"
 
 #define LOCTEXT_NAMESPACE "MassTest"
@@ -980,6 +978,808 @@ struct FCommands_BatchChangeComposition_RejectsSparseAdd : FEntityTestBase
 };
 IMPLEMENT_AI_INSTANT_TEST(FCommands_BatchChangeComposition_RejectsSparseAdd, "System.Mass.Commands.BatchChangeComposition.RejectsSparseAdd");
 
+struct FCommands_AddElementsWithSharedFragments_AllTypes : FEntityTestBase
+{
+	virtual bool InstantTest() override
+	{
+		constexpr int32 Count = 5;
+		TArray<FMassEntityHandle> Entities;
+		EntityManager->BatchCreateEntities(IntsArchetype, Count, Entities);
+
+		// Prepare const shared fragment value
+		const FTestConstSharedFragment_Int ConstSharedInstance(42);
+		const FConstSharedStruct ConstSharedStruct = EntityManager->GetOrCreateConstSharedFragment(ConstSharedInstance);
+
+		// Add fragment + tag + sparse + const shared value in one command
+		for (const FMassEntityHandle& Entity : Entities)
+		{
+			FMassArchetypeSharedFragmentValues SharedValues;
+			SharedValues.Add(ConstSharedStruct);
+			EntityManager->Defer().PushCommand<FMassCommandAddElementsWithSharedFragments<FTestFragment_Float, FTestTag_A, FTestFragment_SparseInt>>(
+				Entity, MoveTemp(SharedValues));
+		}
+
+		EntityManager->FlushCommands();
+
+		for (const FMassEntityHandle& Entity : Entities)
+		{
+			AITEST_TRUE(TEXT("Entity should have FTestFragment_Int"), EntityManager->DoesEntityHaveElement(Entity, FTestFragment_Int::StaticStruct()));
+			AITEST_TRUE(TEXT("Entity should have FTestFragment_Float"), EntityManager->DoesEntityHaveElement(Entity, FTestFragment_Float::StaticStruct()));
+			AITEST_TRUE(TEXT("Entity should have TagA"), EntityManager->DoesEntityHaveElement(Entity, FTestTag_A::StaticStruct()));
+			AITEST_TRUE(TEXT("Entity should have const shared"), EntityManager->DoesEntityHaveElement(Entity, FTestConstSharedFragment_Int::StaticStruct()));
+			AITEST_TRUE(TEXT("Entity should have sparse"), EntityManager->DoesEntityHaveElement(Entity, FTestFragment_SparseInt::StaticStruct()));
+		}
+
+		return true;
+	}
+};
+IMPLEMENT_AI_INSTANT_TEST(FCommands_AddElementsWithSharedFragments_AllTypes, "System.Mass.Commands.AddElementsWithSharedFragments.AllTypes");
+
+struct FCommands_AddElementList_AllTypes : FEntityTestBase
+{
+	virtual bool InstantTest() override
+	{
+		constexpr int32 Count = 5;
+		TArray<FMassEntityHandle> Entities;
+		EntityManager->BatchCreateEntities(IntsArchetype, Count, Entities);
+
+		// Add fragment + tag + sparse in one runtime command
+		const TArray<const UScriptStruct*> ElementTypes = { FTestFragment_Float::StaticStruct(), FTestTag_A::StaticStruct(), FTestFragment_SparseInt::StaticStruct() };
+		FMassCommandAddElementList& Command = EntityManager->Defer().PushUniqueCommand<FMassCommandAddElementList>(ElementTypes);
+		for (const FMassEntityHandle& Entity : Entities)
+		{
+			Command.Add(Entity);
+		}
+
+		EntityManager->FlushCommands();
+
+		for (const FMassEntityHandle& Entity : Entities)
+		{
+			AITEST_TRUE(TEXT("Entity should have FTestFragment_Int (original)"), EntityManager->DoesEntityHaveElement(Entity, FTestFragment_Int::StaticStruct()));
+			AITEST_TRUE(TEXT("Entity should have FTestFragment_Float (added)"), EntityManager->DoesEntityHaveElement(Entity, FTestFragment_Float::StaticStruct()));
+			AITEST_TRUE(TEXT("Entity should have TagA (added)"), EntityManager->DoesEntityHaveElement(Entity, FTestTag_A::StaticStruct()));
+			AITEST_TRUE(TEXT("Entity should have sparse (added)"), EntityManager->DoesEntityHaveElement(Entity, FTestFragment_SparseInt::StaticStruct()));
+		}
+
+		return true;
+	}
+};
+IMPLEMENT_AI_INSTANT_TEST(FCommands_AddElementList_AllTypes, "System.Mass.Commands.AddElementList.AllTypes");
+
+struct FCommands_AddFragmentInstances_AllTypes : FEntityTestBase
+{
+	virtual bool InstantTest() override
+	{
+		constexpr int32 Count = 5;
+		TArray<FMassEntityHandle> Entities;
+		EntityManager->BatchCreateEntities(IntsArchetype, Count, Entities);
+
+		// Add fragment with value + tag + sparse fragment + sparse tag in a single command
+		for (const FMassEntityHandle& Entity : Entities)
+		{
+			EntityManager->Defer().PushCommand<FMassCommandAddFragmentInstances<FTestFragment_Float, FTestTag_A, FTestFragment_SparseInt, FTestTag_SparseA>>(
+				Entity, FTestFragment_Float(2.5f), FTestTag_A(), FTestFragment_SparseInt(77), FTestTag_SparseA());
+		}
+
+		EntityManager->FlushCommands();
+
+		for (const FMassEntityHandle& Entity : Entities)
+		{
+			AITEST_TRUE(TEXT("Entity should have FTestFragment_Int (original)"), EntityManager->DoesEntityHaveElement(Entity, FTestFragment_Int::StaticStruct()));
+			AITEST_TRUE(TEXT("Entity should have FTestFragment_Float (added with value)"), EntityManager->DoesEntityHaveElement(Entity, FTestFragment_Float::StaticStruct()));
+			AITEST_TRUE(TEXT("Entity should have TagA (added)"), EntityManager->DoesEntityHaveElement(Entity, FTestTag_A::StaticStruct()));
+			AITEST_TRUE(TEXT("Entity should have SparseInt (added)"), EntityManager->DoesEntityHaveElement(Entity, FTestFragment_SparseInt::StaticStruct()));
+			AITEST_TRUE(TEXT("Entity should have SparseTagA (added)"), EntityManager->DoesEntityHaveElement(Entity, FTestTag_SparseA::StaticStruct()));
+
+			const FMassEntityView EntityView(*EntityManager, Entity);
+			AITEST_EQUAL(TEXT("Float fragment value should be 2.5f"), EntityView.GetFragmentData<FTestFragment_Float>().Value, 2.5f);
+		}
+
+		return true;
+	}
+};
+IMPLEMENT_AI_INSTANT_TEST(FCommands_AddFragmentInstances_AllTypes, "System.Mass.Commands.AddFragmentInstances.AllTypes");
+
+struct FCommands_BuildEntity_AllTypes : FEntityTestBase
+{
+	virtual bool InstantTest() override
+	{
+		constexpr int32 Count = 5;
+		TArray<FMassEntityHandle> Entities;
+		Entities.Reserve(Count);
+		for (int32 i = 0; i < Count; ++i)
+		{
+			Entities.Add(EntityManager->ReserveEntity());
+		}
+
+		// Build entity with fragment values + tag + sparse fragment + sparse tag
+		for (const FMassEntityHandle& Entity : Entities)
+		{
+			EntityManager->Defer().PushCommand<FMassCommandBuildEntity<FTestFragment_Int, FTestFragment_Float, FTestTag_A, FTestFragment_SparseInt, FTestTag_SparseA>>(
+				Entity, FTestFragment_Int(10), FTestFragment_Float(1.5f), FTestTag_A(), FTestFragment_SparseInt(88), FTestTag_SparseA());
+		}
+
+		EntityManager->FlushCommands();
+
+		for (const FMassEntityHandle& Entity : Entities)
+		{
+			AITEST_TRUE(TEXT("Entity should be valid"), EntityManager->IsEntityValid(Entity));
+			AITEST_TRUE(TEXT("Entity should have FTestFragment_Int"), EntityManager->DoesEntityHaveElement(Entity, FTestFragment_Int::StaticStruct()));
+			AITEST_TRUE(TEXT("Entity should have FTestFragment_Float"), EntityManager->DoesEntityHaveElement(Entity, FTestFragment_Float::StaticStruct()));
+			AITEST_TRUE(TEXT("Entity should have TagA"), EntityManager->DoesEntityHaveElement(Entity, FTestTag_A::StaticStruct()));
+			AITEST_TRUE(TEXT("Entity should have SparseInt"), EntityManager->DoesEntityHaveElement(Entity, FTestFragment_SparseInt::StaticStruct()));
+			AITEST_TRUE(TEXT("Entity should have SparseTagA"), EntityManager->DoesEntityHaveElement(Entity, FTestTag_SparseA::StaticStruct()));
+
+			const FMassEntityView EntityView(*EntityManager, Entity);
+			AITEST_EQUAL(TEXT("Int value should be 10"), EntityView.GetFragmentData<FTestFragment_Int>().Value, 10);
+			AITEST_EQUAL(TEXT("Float value should be 1.5f"), EntityView.GetFragmentData<FTestFragment_Float>().Value, 1.5f);
+		}
+
+		return true;
+	}
+};
+IMPLEMENT_AI_INSTANT_TEST(FCommands_BuildEntity_AllTypes, "System.Mass.Commands.BuildEntity.AllTypes");
+
+struct FCommands_AddFragmentInstancesWithSharedFragments_AllTypes : FEntityTestBase
+{
+	virtual bool InstantTest() override
+	{
+		constexpr int32 Count = 5;
+		TArray<FMassEntityHandle> Entities;
+		EntityManager->BatchCreateEntities(IntsArchetype, Count, Entities);
+
+		const FTestConstSharedFragment_Int ConstSharedInstance(42);
+		const FConstSharedStruct ConstSharedStruct = EntityManager->GetOrCreateConstSharedFragment(ConstSharedInstance);
+
+		// Add fragment with value + tag + sparse fragment + sparse tag + shared value in one command
+		for (const FMassEntityHandle& Entity : Entities)
+		{
+			FMassArchetypeSharedFragmentValues SharedValues;
+			SharedValues.Add(ConstSharedStruct);
+			EntityManager->Defer().PushCommand<FMassCommandAddFragmentInstancesWithSharedFragments<FMassArchetypeSharedFragmentValues, FTestFragment_Float, FTestTag_A, FTestFragment_SparseInt, FTestTag_SparseA>>(
+				Entity, MoveTemp(SharedValues), FTestFragment_Float(3.14f), FTestTag_A(), FTestFragment_SparseInt(55), FTestTag_SparseA());
+		}
+
+		EntityManager->FlushCommands();
+
+		for (const FMassEntityHandle& Entity : Entities)
+		{
+			AITEST_TRUE(TEXT("Entity should have FTestFragment_Int (original)"), EntityManager->DoesEntityHaveElement(Entity, FTestFragment_Int::StaticStruct()));
+			AITEST_TRUE(TEXT("Entity should have FTestFragment_Float (added)"), EntityManager->DoesEntityHaveElement(Entity, FTestFragment_Float::StaticStruct()));
+			AITEST_TRUE(TEXT("Entity should have TagA (added)"), EntityManager->DoesEntityHaveElement(Entity, FTestTag_A::StaticStruct()));
+			AITEST_TRUE(TEXT("Entity should have SparseInt (added)"), EntityManager->DoesEntityHaveElement(Entity, FTestFragment_SparseInt::StaticStruct()));
+			AITEST_TRUE(TEXT("Entity should have SparseTagA (added)"), EntityManager->DoesEntityHaveElement(Entity, FTestTag_SparseA::StaticStruct()));
+			AITEST_TRUE(TEXT("Entity should have const shared"), EntityManager->DoesEntityHaveElement(Entity, FTestConstSharedFragment_Int::StaticStruct()));
+
+			const FMassEntityView EntityView(*EntityManager, Entity);
+			AITEST_EQUAL(TEXT("Float value should be 3.14f"), EntityView.GetFragmentData<FTestFragment_Float>().Value, 3.14f);
+			AITEST_EQUAL(TEXT("Const shared value should be 42"), EntityView.GetConstSharedFragmentData<FTestConstSharedFragment_Int>().Value, 42);
+		}
+
+		return true;
+	}
+};
+IMPLEMENT_AI_INSTANT_TEST(FCommands_AddFragmentInstancesWithSharedFragments_AllTypes, "System.Mass.Commands.AddFragmentInstancesWithSharedFragments.AllTypes");
+
+struct FCommands_BuildEntityWithSharedFragments_AllTypes : FEntityTestBase
+{
+	virtual bool InstantTest() override
+	{
+		constexpr int32 Count = 5;
+		TArray<FMassEntityHandle> Entities;
+		Entities.Reserve(Count);
+		for (int32 i = 0; i < Count; ++i)
+		{
+			Entities.Add(EntityManager->ReserveEntity());
+		}
+
+		const FTestConstSharedFragment_Int ConstSharedInstance(99);
+		const FConstSharedStruct ConstSharedStruct = EntityManager->GetOrCreateConstSharedFragment(ConstSharedInstance);
+
+		// Build entity with fragment values + tag + sparse fragment + sparse tag + shared value
+		for (const FMassEntityHandle& Entity : Entities)
+		{
+			FMassArchetypeSharedFragmentValues SharedValues;
+			SharedValues.Add(ConstSharedStruct);
+			EntityManager->Defer().PushCommand<FMassCommandBuildEntityWithSharedFragments<FMassArchetypeSharedFragmentValues, FTestFragment_Int, FTestTag_A, FTestFragment_SparseInt, FTestTag_SparseA>>(
+				Entity, MoveTemp(SharedValues), FTestFragment_Int(7), FTestTag_A(), FTestFragment_SparseInt(33), FTestTag_SparseA());
+		}
+
+		EntityManager->FlushCommands();
+
+		for (const FMassEntityHandle& Entity : Entities)
+		{
+			AITEST_TRUE(TEXT("Entity should be valid"), EntityManager->IsEntityValid(Entity));
+			AITEST_TRUE(TEXT("Entity should have FTestFragment_Int"), EntityManager->DoesEntityHaveElement(Entity, FTestFragment_Int::StaticStruct()));
+			AITEST_TRUE(TEXT("Entity should have TagA"), EntityManager->DoesEntityHaveElement(Entity, FTestTag_A::StaticStruct()));
+			AITEST_TRUE(TEXT("Entity should have SparseInt"), EntityManager->DoesEntityHaveElement(Entity, FTestFragment_SparseInt::StaticStruct()));
+			AITEST_TRUE(TEXT("Entity should have SparseTagA"), EntityManager->DoesEntityHaveElement(Entity, FTestTag_SparseA::StaticStruct()));
+			AITEST_TRUE(TEXT("Entity should have const shared"), EntityManager->DoesEntityHaveElement(Entity, FTestConstSharedFragment_Int::StaticStruct()));
+
+			const FMassEntityView EntityView(*EntityManager, Entity);
+			AITEST_EQUAL(TEXT("Int value should be 7"), EntityView.GetFragmentData<FTestFragment_Int>().Value, 7);
+			AITEST_EQUAL(TEXT("Const shared value should be 99"), EntityView.GetConstSharedFragmentData<FTestConstSharedFragment_Int>().Value, 99);
+		}
+
+		return true;
+	}
+};
+IMPLEMENT_AI_INSTANT_TEST(FCommands_BuildEntityWithSharedFragments_AllTypes, "System.Mass.Commands.BuildEntityWithSharedFragments.AllTypes");
+
+//----------------------------------------------------------------------//
+// Benchmarks: Separate vs Mixed-Element Commands
+// Each iteration runs the Mixed arm first, then the Separate arm. Mixed therefore
+// always operates on a colder cache; Separate benefits from archetype tables and
+// allocator pools already warmed by the Mixed run. Results reflect this ordering bias.
+//----------------------------------------------------------------------//
+DEFINE_LOG_CATEGORY_STATIC(LogMassCommandsBench, Log, All);
+
+namespace
+{
+	constexpr int32 BenchEntityCount = 10000;
+	constexpr int32 BenchIterations = 10;
+
+	double TimedFlush(const TSharedPtr<FMassEntityManager>& EntityManager)
+	{
+		const double Start = FPlatformTime::Seconds();
+		EntityManager->FlushCommands();
+		const double End = FPlatformTime::Seconds();
+		return (End - Start) * 1000.0;
+	}
+}
+
+struct FBench_AddElements_SeparateVsMixed : FEntityTestBase
+{
+	virtual bool InstantTest() override
+	{
+		double MixedTotal = 0.0;
+		double SeparateTotal = 0.0;
+
+		for (int32 Iter = 0; Iter < BenchIterations; ++Iter)
+		{
+			// Mixed
+			{
+				TArray<FMassEntityHandle> Entities;
+				EntityManager->BatchCreateEntities(IntsArchetype, BenchEntityCount, Entities);
+				for (const FMassEntityHandle& Entity : Entities)
+				{
+					EntityManager->Defer().PushCommand<FMassCommandAddElements<FTestFragment_Float, FTestTag_A, FTestFragment_SparseInt>>(Entity);
+				}
+				MixedTotal += TimedFlush(EntityManager);
+				EntityManager->BatchDestroyEntities(Entities);
+			}
+			// Separate
+			{
+				TArray<FMassEntityHandle> Entities;
+				EntityManager->BatchCreateEntities(IntsArchetype, BenchEntityCount, Entities);
+				FMassCommandAddElement& CmdFloat = EntityManager->Defer().PushUniqueCommand<FMassCommandAddElement>(FTestFragment_Float::StaticStruct());
+				FMassCommandAddElement& CmdTag = EntityManager->Defer().PushUniqueCommand<FMassCommandAddElement>(FTestTag_A::StaticStruct());
+				FMassCommandAddElement& CmdSparse = EntityManager->Defer().PushUniqueCommand<FMassCommandAddElement>(FTestFragment_SparseInt::StaticStruct());
+				for (const FMassEntityHandle& Entity : Entities)
+				{
+					CmdFloat.Add(Entity);
+					CmdTag.Add(Entity);
+					CmdSparse.Add(Entity);
+				}
+				SeparateTotal += TimedFlush(EntityManager);
+				EntityManager->BatchDestroyEntities(Entities);
+			}
+		}
+
+		const double SeparateAvg = SeparateTotal / BenchIterations;
+		const double MixedAvg = MixedTotal / BenchIterations;
+		const double DeltaPct = SeparateAvg > 0.0 ? ((MixedAvg - SeparateAvg) / SeparateAvg) * 100.0 : 0.0;
+		UE_LOG(LogMassCommandsBench, Log, TEXT("AddElements: separate %.3f ms, mixed %.3f ms, delta: %+.1f%%"), SeparateAvg, MixedAvg, DeltaPct);
+
+		// Validate once
+		{
+			TArray<FMassEntityHandle> Entities;
+			EntityManager->BatchCreateEntities(IntsArchetype, 5, Entities);
+			for (const FMassEntityHandle& Entity : Entities)
+			{
+				EntityManager->Defer().PushCommand<FMassCommandAddElements<FTestFragment_Float, FTestTag_A, FTestFragment_SparseInt>>(Entity);
+			}
+			EntityManager->FlushCommands();
+			for (const FMassEntityHandle& Entity : Entities)
+			{
+				AITEST_TRUE(TEXT("Should have Float"), EntityManager->DoesEntityHaveElement(Entity, FTestFragment_Float::StaticStruct()));
+				AITEST_TRUE(TEXT("Should have TagA"), EntityManager->DoesEntityHaveElement(Entity, FTestTag_A::StaticStruct()));
+				AITEST_TRUE(TEXT("Should have Sparse"), EntityManager->DoesEntityHaveElement(Entity, FTestFragment_SparseInt::StaticStruct()));
+			}
+		}
+		return true;
+	}
+};
+IMPLEMENT_AI_INSTANT_TEST(FBench_AddElements_SeparateVsMixed, "System.Mass.Commands.Benchmark.AddElements_SeparateVsMixed");
+
+struct FBench_AddFragmentInstances_SeparateVsMixed : FEntityTestBase
+{
+	virtual bool InstantTest() override
+	{
+		double MixedTotal = 0.0;
+		double SeparateTotal = 0.0;
+
+		for (int32 Iter = 0; Iter < BenchIterations; ++Iter)
+		{
+			// Mixed
+			{
+				TArray<FMassEntityHandle> Entities;
+				EntityManager->BatchCreateEntities(IntsArchetype, BenchEntityCount, Entities);
+				for (const FMassEntityHandle& Entity : Entities)
+				{
+					EntityManager->Defer().PushCommand<FMassCommandAddFragmentInstances<FTestFragment_Float, FTestTag_A>>(Entity, FTestFragment_Float(3.14f), FTestTag_A());
+				}
+				MixedTotal += TimedFlush(EntityManager);
+				EntityManager->BatchDestroyEntities(Entities);
+			}
+			// Separate
+			{
+				TArray<FMassEntityHandle> Entities;
+				EntityManager->BatchCreateEntities(IntsArchetype, BenchEntityCount, Entities);
+				for (const FMassEntityHandle& Entity : Entities)
+				{
+					EntityManager->Defer().PushCommand<FMassCommandAddFragmentInstances<FTestFragment_Float>>(Entity, FTestFragment_Float(3.14f));
+					EntityManager->Defer().PushCommand<FMassCommandAddElements<FTestTag_A>>(Entity);
+				}
+				SeparateTotal += TimedFlush(EntityManager);
+				EntityManager->BatchDestroyEntities(Entities);
+			}
+		}
+
+		const double SeparateAvg = SeparateTotal / BenchIterations;
+		const double MixedAvg = MixedTotal / BenchIterations;
+		const double DeltaPct = SeparateAvg > 0.0 ? ((MixedAvg - SeparateAvg) / SeparateAvg) * 100.0 : 0.0;
+		UE_LOG(LogMassCommandsBench, Log, TEXT("AddFragmentInstances: separate %.3f ms, mixed %.3f ms, delta: %+.1f%%"), SeparateAvg, MixedAvg, DeltaPct);
+
+		// Validate once
+		{
+			TArray<FMassEntityHandle> Entities;
+			EntityManager->BatchCreateEntities(IntsArchetype, 5, Entities);
+			for (const FMassEntityHandle& Entity : Entities)
+			{
+				EntityManager->Defer().PushCommand<FMassCommandAddFragmentInstances<FTestFragment_Float, FTestTag_A>>(Entity, FTestFragment_Float(3.14f), FTestTag_A());
+			}
+			EntityManager->FlushCommands();
+			for (const FMassEntityHandle& Entity : Entities)
+			{
+				AITEST_TRUE(TEXT("Should have Float"), EntityManager->DoesEntityHaveElement(Entity, FTestFragment_Float::StaticStruct()));
+				AITEST_TRUE(TEXT("Should have TagA"), EntityManager->DoesEntityHaveElement(Entity, FTestTag_A::StaticStruct()));
+				const FMassEntityView View(*EntityManager, Entity);
+				AITEST_EQUAL(TEXT("Float value"), View.GetFragmentData<FTestFragment_Float>().Value, 3.14f);
+			}
+		}
+		return true;
+	}
+};
+IMPLEMENT_AI_INSTANT_TEST(FBench_AddFragmentInstances_SeparateVsMixed, "System.Mass.Commands.Benchmark.AddFragmentInstances_SeparateVsMixed");
+
+struct FBench_AddFragmentInstancesWithShared_DeferredVsMixed : FEntityTestBase
+{
+	virtual bool InstantTest() override
+	{
+		const FTestConstSharedFragment_Int ConstSharedInstance(42);
+		const FConstSharedStruct ConstSharedStruct = EntityManager->GetOrCreateConstSharedFragment(ConstSharedInstance);
+
+		double MixedTotal = 0.0;
+		double SeparateTotal = 0.0;
+
+		for (int32 Iter = 0; Iter < BenchIterations; ++Iter)
+		{
+			// Mixed
+			{
+				TArray<FMassEntityHandle> Entities;
+				EntityManager->BatchCreateEntities(IntsArchetype, BenchEntityCount, Entities);
+				for (const FMassEntityHandle& Entity : Entities)
+				{
+					FMassArchetypeSharedFragmentValues SharedValues;
+					SharedValues.Add(ConstSharedStruct);
+					EntityManager->Defer().PushCommand<FMassCommandAddFragmentInstancesWithSharedFragments<FMassArchetypeSharedFragmentValues, FTestFragment_Float>>(
+						Entity, MoveTemp(SharedValues), FTestFragment_Float(3.14f));
+				}
+				MixedTotal += TimedFlush(EntityManager);
+				EntityManager->BatchDestroyEntities(Entities);
+			}
+			// Separate
+			{
+				TArray<FMassEntityHandle> Entities;
+				EntityManager->BatchCreateEntities(IntsArchetype, BenchEntityCount, Entities);
+				for (const FMassEntityHandle& Entity : Entities)
+				{
+					EntityManager->Defer().PushCommand<FMassCommandAddFragmentInstances<FTestFragment_Float>>(Entity, FTestFragment_Float(3.14f));
+					EntityManager->Defer().PushCommand<FMassDeferredAddCommand>([Entity, ConstSharedStruct](FMassEntityManager& InEntityManager)
+						{
+							if (InEntityManager.IsEntityValid(Entity))
+							{
+								InEntityManager.AddConstSharedFragmentToEntity(Entity, ConstSharedStruct);
+							}
+						});
+				}
+				SeparateTotal += TimedFlush(EntityManager);
+				EntityManager->BatchDestroyEntities(Entities);
+			}
+		}
+
+		const double SeparateAvg = SeparateTotal / BenchIterations;
+		const double MixedAvg = MixedTotal / BenchIterations;
+		const double DeltaPct = SeparateAvg > 0.0 ? ((MixedAvg - SeparateAvg) / SeparateAvg) * 100.0 : 0.0;
+		UE_LOG(LogMassCommandsBench, Log, TEXT("AddFragmentInstances+Shared: separate %.3f ms, mixed %.3f ms, delta: %+.1f%%"), SeparateAvg, MixedAvg, DeltaPct);
+
+		// Validate once
+		{
+			TArray<FMassEntityHandle> Entities;
+			EntityManager->BatchCreateEntities(IntsArchetype, 5, Entities);
+			for (const FMassEntityHandle& Entity : Entities)
+			{
+				FMassArchetypeSharedFragmentValues SharedValues;
+				SharedValues.Add(ConstSharedStruct);
+				EntityManager->Defer().PushCommand<FMassCommandAddFragmentInstancesWithSharedFragments<FMassArchetypeSharedFragmentValues, FTestFragment_Float>>(
+					Entity, MoveTemp(SharedValues), FTestFragment_Float(3.14f));
+			}
+			EntityManager->FlushCommands();
+			for (const FMassEntityHandle& Entity : Entities)
+			{
+				AITEST_TRUE(TEXT("Should have Float"), EntityManager->DoesEntityHaveElement(Entity, FTestFragment_Float::StaticStruct()));
+				AITEST_TRUE(TEXT("Should have ConstShared"), EntityManager->DoesEntityHaveElement(Entity, FTestConstSharedFragment_Int::StaticStruct()));
+				const FMassEntityView View(*EntityManager, Entity);
+				AITEST_EQUAL(TEXT("Float value"), View.GetFragmentData<FTestFragment_Float>().Value, 3.14f);
+				AITEST_EQUAL(TEXT("Shared value"), View.GetConstSharedFragmentData<FTestConstSharedFragment_Int>().Value, 42);
+			}
+		}
+		return true;
+	}
+};
+IMPLEMENT_AI_INSTANT_TEST(FBench_AddFragmentInstancesWithShared_DeferredVsMixed, "System.Mass.Commands.Benchmark.AddFragmentInstancesWithShared_DeferredVsMixed");
+
+struct FBench_BuildEntity_SeparateVsMixed : FEntityTestBase
+{
+	virtual bool InstantTest() override
+	{
+		double MixedTotal = 0.0;
+		double SeparateTotal = 0.0;
+
+		for (int32 Iter = 0; Iter < BenchIterations; ++Iter)
+		{
+			// Mixed
+			{
+				TArray<FMassEntityHandle> Entities;
+				Entities.Reserve(BenchEntityCount);
+				for (int32 i = 0; i < BenchEntityCount; ++i) { Entities.Add(EntityManager->ReserveEntity()); }
+				for (const FMassEntityHandle& Entity : Entities)
+				{
+					EntityManager->Defer().PushCommand<FMassCommandBuildEntity<FTestFragment_Int, FTestFragment_Float, FTestTag_A>>(
+						Entity, FTestFragment_Int(10), FTestFragment_Float(2.5f), FTestTag_A());
+				}
+				MixedTotal += TimedFlush(EntityManager);
+				EntityManager->BatchDestroyEntities(Entities);
+			}
+			// Separate
+			{
+				TArray<FMassEntityHandle> Entities;
+				Entities.Reserve(BenchEntityCount);
+				for (int32 i = 0; i < BenchEntityCount; ++i) { Entities.Add(EntityManager->ReserveEntity()); }
+				for (const FMassEntityHandle& Entity : Entities)
+				{
+					EntityManager->Defer().PushCommand<FMassCommandBuildEntity<FTestFragment_Int, FTestFragment_Float>>(
+						Entity, FTestFragment_Int(10), FTestFragment_Float(2.5f));
+					EntityManager->Defer().PushCommand<FMassCommandAddElements<FTestTag_A>>(Entity);
+				}
+				SeparateTotal += TimedFlush(EntityManager);
+				EntityManager->BatchDestroyEntities(Entities);
+			}
+		}
+
+		const double SeparateAvg = SeparateTotal / BenchIterations;
+		const double MixedAvg = MixedTotal / BenchIterations;
+		const double DeltaPct = SeparateAvg > 0.0 ? ((MixedAvg - SeparateAvg) / SeparateAvg) * 100.0 : 0.0;
+		UE_LOG(LogMassCommandsBench, Log, TEXT("BuildEntity: separate %.3f ms, mixed %.3f ms, delta: %+.1f%%"), SeparateAvg, MixedAvg, DeltaPct);
+
+		// Validate once
+		{
+			TArray<FMassEntityHandle> Entities;
+			Entities.Reserve(5);
+			for (int32 i = 0; i < 5; ++i) { Entities.Add(EntityManager->ReserveEntity()); }
+			for (const FMassEntityHandle& Entity : Entities)
+			{
+				EntityManager->Defer().PushCommand<FMassCommandBuildEntity<FTestFragment_Int, FTestFragment_Float, FTestTag_A>>(
+					Entity, FTestFragment_Int(10), FTestFragment_Float(2.5f), FTestTag_A());
+			}
+			EntityManager->FlushCommands();
+			for (const FMassEntityHandle& Entity : Entities)
+			{
+				AITEST_TRUE(TEXT("Should have Int"), EntityManager->DoesEntityHaveElement(Entity, FTestFragment_Int::StaticStruct()));
+				AITEST_TRUE(TEXT("Should have Float"), EntityManager->DoesEntityHaveElement(Entity, FTestFragment_Float::StaticStruct()));
+				AITEST_TRUE(TEXT("Should have TagA"), EntityManager->DoesEntityHaveElement(Entity, FTestTag_A::StaticStruct()));
+				const FMassEntityView View(*EntityManager, Entity);
+				AITEST_EQUAL(TEXT("Int value"), View.GetFragmentData<FTestFragment_Int>().Value, 10);
+				AITEST_EQUAL(TEXT("Float value"), View.GetFragmentData<FTestFragment_Float>().Value, 2.5f);
+			}
+		}
+		return true;
+	}
+};
+IMPLEMENT_AI_INSTANT_TEST(FBench_BuildEntity_SeparateVsMixed, "System.Mass.Commands.Benchmark.BuildEntity_SeparateVsMixed");
+
+//----------------------------------------------------------------------//
+// Parity benchmarks: single-type vs multi-type command overhead
+// Measures whether multi-type command instantiations carry dispatch overhead relative
+// to single-type instantiations. Both arms must perform identical composition transitions;
+// where possible the extra type in the multi arm is already present in the source archetype
+// so it is a no-op at the archetype level. BuildEntity is an exception: reserved entities
+// have no existing composition, so the arms build genuinely different archetypes.
+//----------------------------------------------------------------------//
+
+struct FBench_AddElements_SingleVsMultiType : FEntityTestBase
+{
+	virtual bool InstantTest() override
+	{
+		double MultiTotal = 0.0;
+		double SingleTotal = 0.0;
+
+		// FTestTag_A is already in IntsTagAArchetype, so the multi arm's extra type is a
+		// no-op at the composition level. Both arms perform the identical transition
+		// ({Int,Tag_A} -> {Int,Tag_A,Float}); the only variable is command dispatch overhead.
+		const FMassArchetypeHandle IntsTagAArchetype = EntityManager->CreateArchetype(
+			{FTestFragment_Int::StaticStruct(), FTestTag_A::StaticStruct()});
+
+		for (int32 Iter = 0; Iter < BenchIterations; ++Iter)
+		{
+			// Multi-type: FTestTag_A already present — no-op for Tag_A at composition level
+			{
+				TArray<FMassEntityHandle> Entities;
+				EntityManager->BatchCreateEntities(IntsTagAArchetype, BenchEntityCount, Entities);
+				for (const FMassEntityHandle& Entity : Entities)
+				{
+					EntityManager->Defer().PushCommand<FMassCommandAddElements<FTestFragment_Float, FTestTag_A>>(Entity);
+				}
+				MultiTotal += TimedFlush(EntityManager);
+				EntityManager->BatchDestroyEntities(Entities);
+			}
+			// Single-type
+			{
+				TArray<FMassEntityHandle> Entities;
+				EntityManager->BatchCreateEntities(IntsTagAArchetype, BenchEntityCount, Entities);
+				for (const FMassEntityHandle& Entity : Entities)
+				{
+					EntityManager->Defer().PushCommand<FMassCommandAddElements<FTestFragment_Float>>(Entity);
+				}
+				SingleTotal += TimedFlush(EntityManager);
+				EntityManager->BatchDestroyEntities(Entities);
+			}
+		}
+
+		const double SingleAvg = SingleTotal / BenchIterations;
+		const double MultiAvg = MultiTotal / BenchIterations;
+		const double OverheadPct = SingleAvg > 0.0 ? ((MultiAvg - SingleAvg) / SingleAvg) * 100.0 : 0.0;
+		UE_LOG(LogMassCommandsBench, Log, TEXT("AddElements parity: single-type %.3f ms, multi-type %.3f ms, overhead: %+.1f%%"), SingleAvg, MultiAvg, OverheadPct);
+
+		// Validate once
+		{
+			TArray<FMassEntityHandle> Entities;
+			EntityManager->BatchCreateEntities(IntsArchetype, 5, Entities);
+			for (const FMassEntityHandle& Entity : Entities)
+			{
+				EntityManager->Defer().PushCommand<FMassCommandAddElements<FTestFragment_Float, FTestTag_A>>(Entity);
+			}
+			EntityManager->FlushCommands();
+			for (const FMassEntityHandle& Entity : Entities)
+			{
+				AITEST_TRUE(TEXT("Should have Float"), EntityManager->DoesEntityHaveElement(Entity, FTestFragment_Float::StaticStruct()));
+				AITEST_TRUE(TEXT("Should have TagA"), EntityManager->DoesEntityHaveElement(Entity, FTestTag_A::StaticStruct()));
+			}
+		}
+		return true;
+	}
+};
+IMPLEMENT_AI_INSTANT_TEST(FBench_AddElements_SingleVsMultiType, "System.Mass.Commands.Benchmark.AddElements_SingleVsMultiType");
+
+struct FBench_RemoveElements_SingleVsMultiType : FEntityTestBase
+{
+	virtual bool InstantTest() override
+	{
+		double MultiTotal = 0.0;
+		double SingleTotal = 0.0;
+
+		for (int32 Iter = 0; Iter < BenchIterations; ++Iter)
+		{
+			// Multi-type: FTestTag_A is not present in FloatsIntsArchetype, so the removal is a no-op for that type.
+			// Both arms therefore perform the identical archetype transition ({Float,Int} -> {Int}); the only
+			// variable is single-type vs multi-type command dispatch overhead.
+			{
+				TArray<FMassEntityHandle> Entities;
+				EntityManager->BatchCreateEntities(FloatsIntsArchetype, BenchEntityCount, Entities);
+				for (const FMassEntityHandle& Entity : Entities)
+				{
+					EntityManager->Defer().PushCommand<FMassCommandRemoveElements<FTestFragment_Float, FTestTag_A>>(Entity);
+				}
+				MultiTotal += TimedFlush(EntityManager);
+				EntityManager->BatchDestroyEntities(Entities);
+			}
+			// Single-type
+			{
+				TArray<FMassEntityHandle> Entities;
+				EntityManager->BatchCreateEntities(FloatsIntsArchetype, BenchEntityCount, Entities);
+				for (const FMassEntityHandle& Entity : Entities)
+				{
+					EntityManager->Defer().PushCommand<FMassCommandRemoveElements<FTestFragment_Float>>(Entity);
+				}
+				SingleTotal += TimedFlush(EntityManager);
+				EntityManager->BatchDestroyEntities(Entities);
+			}
+		}
+
+		const double SingleAvg = SingleTotal / BenchIterations;
+		const double MultiAvg = MultiTotal / BenchIterations;
+		const double OverheadPct = SingleAvg > 0.0 ? ((MultiAvg - SingleAvg) / SingleAvg) * 100.0 : 0.0;
+		UE_LOG(LogMassCommandsBench, Log, TEXT("RemoveElements parity: single-type %.3f ms, multi-type %.3f ms, overhead: %+.1f%%"), SingleAvg, MultiAvg, OverheadPct);
+
+		// Validate once
+		{
+			TArray<FMassEntityHandle> Entities;
+			EntityManager->BatchCreateEntities(FloatsIntsArchetype, 5, Entities);
+			for (const FMassEntityHandle& Entity : Entities)
+			{
+				EntityManager->Defer().PushCommand<FMassCommandAddElements<FTestTag_A>>(Entity);
+			}
+			EntityManager->FlushCommands();
+			for (const FMassEntityHandle& Entity : Entities)
+			{
+				EntityManager->Defer().PushCommand<FMassCommandRemoveElements<FTestFragment_Float, FTestTag_A>>(Entity);
+			}
+			EntityManager->FlushCommands();
+			for (const FMassEntityHandle& Entity : Entities)
+			{
+				AITEST_FALSE(TEXT("Should not have Float"), EntityManager->DoesEntityHaveElement(Entity, FTestFragment_Float::StaticStruct()));
+				AITEST_FALSE(TEXT("Should not have TagA"), EntityManager->DoesEntityHaveElement(Entity, FTestTag_A::StaticStruct()));
+				AITEST_TRUE(TEXT("Should still have Int"), EntityManager->DoesEntityHaveElement(Entity, FTestFragment_Int::StaticStruct()));
+			}
+		}
+		return true;
+	}
+};
+IMPLEMENT_AI_INSTANT_TEST(FBench_RemoveElements_SingleVsMultiType, "System.Mass.Commands.Benchmark.RemoveElements_SingleVsMultiType");
+
+struct FBench_AddFragmentInstances_SingleVsMultiType : FEntityTestBase
+{
+	virtual bool InstantTest() override
+	{
+		double MultiTotal = 0.0;
+		double SingleTotal = 0.0;
+
+		// FTestTag_A is already in IntsTagAArchetype, so the multi arm's extra type is a
+		// no-op at the composition level. Tags also carry no value payload (GetAsNonTagGenericMultiArray
+		// excludes them), so the command buffers are equivalent in data volume. Both arms perform
+		// the identical transition ({Int,Tag_A} -> {Int,Tag_A,Float}); the only variable is
+		// command dispatch overhead.
+		const FMassArchetypeHandle IntsTagAArchetype = EntityManager->CreateArchetype(
+			{FTestFragment_Int::StaticStruct(), FTestTag_A::StaticStruct()});
+
+		for (int32 Iter = 0; Iter < BenchIterations; ++Iter)
+		{
+			// Multi-type: FTestTag_A already present — no-op for Tag_A at composition level;
+			// tag value excluded from payload by GetAsNonTagGenericMultiArray
+			{
+				TArray<FMassEntityHandle> Entities;
+				EntityManager->BatchCreateEntities(IntsTagAArchetype, BenchEntityCount, Entities);
+				for (const FMassEntityHandle& Entity : Entities)
+				{
+					EntityManager->Defer().PushCommand<FMassCommandAddFragmentInstances<FTestFragment_Float, FTestTag_A>>(Entity, FTestFragment_Float(3.14f), FTestTag_A());
+				}
+				MultiTotal += TimedFlush(EntityManager);
+				EntityManager->BatchDestroyEntities(Entities);
+			}
+			// Single-type
+			{
+				TArray<FMassEntityHandle> Entities;
+				EntityManager->BatchCreateEntities(IntsTagAArchetype, BenchEntityCount, Entities);
+				for (const FMassEntityHandle& Entity : Entities)
+				{
+					EntityManager->Defer().PushCommand<FMassCommandAddFragmentInstances<FTestFragment_Float>>(Entity, FTestFragment_Float(3.14f));
+				}
+				SingleTotal += TimedFlush(EntityManager);
+				EntityManager->BatchDestroyEntities(Entities);
+			}
+		}
+
+		const double SingleAvg = SingleTotal / BenchIterations;
+		const double MultiAvg = MultiTotal / BenchIterations;
+		const double OverheadPct = SingleAvg > 0.0 ? ((MultiAvg - SingleAvg) / SingleAvg) * 100.0 : 0.0;
+		UE_LOG(LogMassCommandsBench, Log, TEXT("AddFragmentInstances parity: single-type %.3f ms, multi-type %.3f ms, overhead: %+.1f%%"), SingleAvg, MultiAvg, OverheadPct);
+
+		// Validate once
+		{
+			TArray<FMassEntityHandle> Entities;
+			EntityManager->BatchCreateEntities(IntsArchetype, 5, Entities);
+			for (const FMassEntityHandle& Entity : Entities)
+			{
+				EntityManager->Defer().PushCommand<FMassCommandAddFragmentInstances<FTestFragment_Float, FTestTag_A>>(Entity, FTestFragment_Float(3.14f), FTestTag_A());
+			}
+			EntityManager->FlushCommands();
+			for (const FMassEntityHandle& Entity : Entities)
+			{
+				AITEST_TRUE(TEXT("Should have Float"), EntityManager->DoesEntityHaveElement(Entity, FTestFragment_Float::StaticStruct()));
+				AITEST_TRUE(TEXT("Should have TagA"), EntityManager->DoesEntityHaveElement(Entity, FTestTag_A::StaticStruct()));
+				const FMassEntityView View(*EntityManager, Entity);
+				AITEST_EQUAL(TEXT("Float value"), View.GetFragmentData<FTestFragment_Float>().Value, 3.14f);
+			}
+		}
+		return true;
+	}
+};
+IMPLEMENT_AI_INSTANT_TEST(FBench_AddFragmentInstances_SingleVsMultiType, "System.Mass.Commands.Benchmark.AddFragmentInstances_SingleVsMultiType");
+
+struct FBench_BuildEntity_SingleVsMultiType : FEntityTestBase
+{
+	virtual bool InstantTest() override
+	{
+		double MultiTotal = 0.0;
+		double SingleTotal = 0.0;
+
+		// Note: reserved entities have no existing composition, so there is no "already-present type"
+		// trick available. The multi arm builds {Int,Float,Tag_A} and the single arm builds {Int,Float};
+		// the delta includes real composition cost (extra archetype lookup + Tag_A in the type set),
+		// not pure dispatch overhead. After the first iteration both archetypes are cached so the
+		// archetype-creation cost is amortized, but the arms remain methodologically asymmetric.
+		for (int32 Iter = 0; Iter < BenchIterations; ++Iter)
+		{
+			// Multi-type
+			{
+				TArray<FMassEntityHandle> Entities;
+				Entities.Reserve(BenchEntityCount);
+				for (int32 i = 0; i < BenchEntityCount; ++i) { Entities.Add(EntityManager->ReserveEntity()); }
+				for (const FMassEntityHandle& Entity : Entities)
+				{
+					EntityManager->Defer().PushCommand<FMassCommandBuildEntity<FTestFragment_Int, FTestFragment_Float, FTestTag_A>>(
+						Entity, FTestFragment_Int(10), FTestFragment_Float(2.5f), FTestTag_A());
+				}
+				MultiTotal += TimedFlush(EntityManager);
+				EntityManager->BatchDestroyEntities(Entities);
+			}
+			// Single-type
+			{
+				TArray<FMassEntityHandle> Entities;
+				Entities.Reserve(BenchEntityCount);
+				for (int32 i = 0; i < BenchEntityCount; ++i) { Entities.Add(EntityManager->ReserveEntity()); }
+				for (const FMassEntityHandle& Entity : Entities)
+				{
+					EntityManager->Defer().PushCommand<FMassCommandBuildEntity<FTestFragment_Int, FTestFragment_Float>>(
+						Entity, FTestFragment_Int(10), FTestFragment_Float(2.5f));
+				}
+				SingleTotal += TimedFlush(EntityManager);
+				EntityManager->BatchDestroyEntities(Entities);
+			}
+		}
+
+		const double SingleAvg = SingleTotal / BenchIterations;
+		const double MultiAvg = MultiTotal / BenchIterations;
+		const double OverheadPct = SingleAvg > 0.0 ? ((MultiAvg - SingleAvg) / SingleAvg) * 100.0 : 0.0;
+		UE_LOG(LogMassCommandsBench, Log, TEXT("BuildEntity parity: single-type %.3f ms, multi-type %.3f ms, overhead: %+.1f%%"), SingleAvg, MultiAvg, OverheadPct);
+
+		// Validate once
+		{
+			TArray<FMassEntityHandle> Entities;
+			Entities.Reserve(5);
+			for (int32 i = 0; i < 5; ++i) { Entities.Add(EntityManager->ReserveEntity()); }
+			for (const FMassEntityHandle& Entity : Entities)
+			{
+				EntityManager->Defer().PushCommand<FMassCommandBuildEntity<FTestFragment_Int, FTestFragment_Float, FTestTag_A>>(
+					Entity, FTestFragment_Int(10), FTestFragment_Float(2.5f), FTestTag_A());
+			}
+			EntityManager->FlushCommands();
+			for (const FMassEntityHandle& Entity : Entities)
+			{
+				AITEST_TRUE(TEXT("Should have Int"), EntityManager->DoesEntityHaveElement(Entity, FTestFragment_Int::StaticStruct()));
+				AITEST_TRUE(TEXT("Should have Float"), EntityManager->DoesEntityHaveElement(Entity, FTestFragment_Float::StaticStruct()));
+				AITEST_TRUE(TEXT("Should have TagA"), EntityManager->DoesEntityHaveElement(Entity, FTestTag_A::StaticStruct()));
+				const FMassEntityView View(*EntityManager, Entity);
+				AITEST_EQUAL(TEXT("Int value"), View.GetFragmentData<FTestFragment_Int>().Value, 10);
+				AITEST_EQUAL(TEXT("Float value"), View.GetFragmentData<FTestFragment_Float>().Value, 2.5f);
+			}
+		}
+		return true;
+	}
+};
+IMPLEMENT_AI_INSTANT_TEST(FBench_BuildEntity_SingleVsMultiType, "System.Mass.Commands.Benchmark.BuildEntity_SingleVsMultiType");
+
 #endif // WITH_MASSENTITY_DEBUG
 
 //----------------------------------------------------------------------//
diff --git a/Engine/Source/Runtime/Mass/MassEngine/Private/AI/MassEngineNavigationProcessors.cpp b/Engine/Source/Runtime/Mass/MassEngine/Private/AI/MassEngineNavigationProcessors.cpp
--- a/Engine/Source/Runtime/Mass/MassEngine/Private/AI/MassEngineNavigationProcessors.cpp
+++ b/Engine/Source/Runtime/Mass/MassEngine/Private/AI/MassEngineNavigationProcessors.cpp
@@ -186,29 +186,21 @@ void UMassNavigationElementProcessor::Execute(FMassEntityManager& EntityManager,
 			{
 				const FMassEntityHandle EntityHandle = Context.GetEntity(EntityIt);
 
-				// Mark as evaluated so this entity is excluded from future runs
-				Context.Defer().AddTag<FMassNavigationEvaluatedTag>(EntityHandle);
-
+				// Both paths below add FMassNavigationEvaluatedTag so this entity is excluded from future runs.
 				if (!IsNavigationRelevant(CollisionSettingsFragment, MeshFragment, TransformFragments[EntityIt]))
 				{
+					Context.Defer().AddTag<FMassNavigationEvaluatedTag>(EntityHandle);
+
 					continue;
 				}
 
+				FMassArchetypeSharedFragmentValues SharedValues;
+				SharedValues.Add(ParametersSharedFragment);
+
 				FNavigationElement Element(*EntityManagerOwner, EntityHandle.AsNumber());
 				ConfigureElement(EntityManagerOwner, Element, MeshFragment, TransformFragments[EntityIt], ParametersSharedFragment.Get<const FMassNavigationRelevantParameters>());
 				const FNavigationElementHandle NavigationElementHandle = FNavigationSystem::AddNavigationElement(EntityManagerOwner->GetWorld(), MoveTemp(Element));
-				Context.Defer().PushCommand<FMassCommandAddFragmentInstances>(EntityHandle, FMassNavigationRelevantFragment(NavigationElementHandle));
-				Context.Defer().PushCommand<FMassDeferredAddCommand>([EntityHandle, ParametersSharedFragment](FMassEntityManager& InEntityManager)
-					{
-						if (InEntityManager.IsEntityValid(EntityHandle))
-						{
-							InEntityManager.AddConstSharedFragmentToEntity(EntityHandle, ParametersSharedFragment);
-						}
-						else
-						{
-							UE_LOGF(LogMass, Log, "FMassNavigationRelevantParameters shared fragment not added: Entity is no longer valid");
-						}
-					});
+				Context.Defer().AddFragmentInstancesWithSharedFragments(EntityHandle, MoveTemp(SharedValues), FMassNavigationRelevantFragment(NavigationElementHandle), FMassNavigationEvaluatedTag{});
 			}
 		});
 }
@@ -333,13 +325,13 @@ void UMassNavigationRelevantMassDeinitializer::Execute(FMassEntityManager& Entit
 				}
 			}
 
-			Context.Defer().PushCommand<FMassCommandRemoveElements<FMassNavigationRelevantFragment, FMassNavigationRelevantParameters, FMassNavigationEvaluatedTag>>(Context.GetEntities());
+			Context.Defer().RemoveElements<FMassNavigationRelevantFragment, FMassNavigationRelevantParameters, FMassNavigationEvaluatedTag>(Context.GetEntities());
 		});
 
 	// Remove the evaluated tag from non-navigation-relevant entities so they can be re-evaluated
 	NonRelevantEntityQuery.ForEachEntityChunk(ExecutionContext,
 		[](FMassExecutionContext& Context)
 		{
-			Context.Defer().PushCommand<FMassCommandRemoveElements<FMassNavigationEvaluatedTag>>(Context.GetEntities());
+			Context.Defer().RemoveElements<FMassNavigationEvaluatedTag>(Context.GetEntities());
 		});
 }
diff --git a/Engine/Source/Runtime/Mass/MassEngine/Private/Mesh/MassEngineRenderISMProcessors.cpp b/Engine/Source/Runtime/Mass/MassEngine/Private/Mesh/MassEngineRenderISMProcessors.cpp
--- a/Engine/Source/Runtime/Mass/MassEngine/Private/Mesh/MassEngineRenderISMProcessors.cpp
+++ b/Engine/Source/Runtime/Mass/MassEngine/Private/Mesh/MassEngineRenderISMProcessors.cpp
@@ -235,23 +235,13 @@ void UMassInstantiateStaticMeshAndCreateISMRenderStateProcessor::Execute(FMassEn
 				FMassRenderISMInstanceFragment ISMInfoFragment;
 				ISMInfoFragment.InstanceIndex = ISMDescFragment.PerInstanceSMData.Add(TransformFragments[*EntityIt].GetTransform().ToMatrixWithScale());
 #if WITH_EDITOR
-				ISMInfoFragment.HitProxyIndex = ISMDescFragment.PerInstanceHitProxy.Add( new HTypeElementHandleHitProxy(TypeElementFragments.Num()? TypeElementFragments[*EntityIt].TypeElementHandle : FTypedElementHandle{}));
+				ISMInfoFragment.HitProxyIndex = ISMDescFragment.PerInstanceHitProxy.Add(new HTypeElementHandleHitProxy(TypeElementFragments.Num()? TypeElementFragments[*EntityIt].TypeElementHandle : FTypedElementHandle{}));
 #endif // WITH_EDITOR
 
-				// Add new fragments @todo find a way to move to new archetype in one shot
-				FMassEntityHandle EntityHandle = Context.GetEntity(*EntityIt);
-				Context.Defer().PushCommand<FMassCommandAddFragmentInstances>(EntityHandle, MoveTemp(ISMInfoFragment));
-				Context.Defer().PushCommand<FMassDeferredAddCommand>([EntityHandle, LinkedISMEntityFragment](FMassEntityManager& InEntityManager)
-					{
-						if (InEntityManager.IsEntityValid(EntityHandle))
-						{
-							InEntityManager.AddConstSharedFragmentToEntity(EntityHandle, LinkedISMEntityFragment);
-						}
-						else
-						{
-							UE_LOGF(LogMassEngine, Log, "FMassRenderISMLinkFragment shared fragment not added: Entity is no longer valid");
-						}
-					});
+				const FMassEntityHandle EntityHandle = Context.GetEntity(*EntityIt);
+				FMassArchetypeSharedFragmentValues SharedValues;
+				SharedValues.Add(LinkedISMEntityFragment);
+				Context.Defer().AddFragmentInstancesWithSharedFragments(EntityHandle, MoveTemp(SharedValues), MoveTemp(ISMInfoFragment));
 
 				ISMEntityInfo->CurrentInstance++;
 			}
@@ -366,19 +356,7 @@ void UMassDirtyRenderISMInstanceProcessor::SignalEntities(FMassEntityManager& En
 				}
 				else
 				{
-					FMassEntityHandle EntityHandle = Context.GetEntity(*EntityIt);
-					Context.Defer().RemoveFragment<FMassRenderISMInstanceFragment>(EntityHandle);
-					Context.Defer().PushCommand<FMassDeferredRemoveCommand>([EntityHandle](FMassEntityManager& InEntityManager)
-						{
-							if (InEntityManager.IsEntityValid(EntityHandle))
-							{
-								InEntityManager.RemoveConstSharedFragmentFromEntity(EntityHandle, FMassRenderISMLinkFragment::StaticStruct());
-							}
-							else
-							{
-								UE_LOGF(LogMassEngine, Log, "FMassRenderISMLinkFragment shared fragment not added: Entity is no longer valid");
-							}
-						});
+					Context.Defer().RemoveElements<FMassRenderISMInstanceFragment, FMassRenderISMLinkFragment>(Context.GetEntity(*EntityIt));
 
 					// Delete the ism entity if this is the last instance in the desc
 					if (ISMDescFragment->PerInstanceSMData.Num() == 1)
@@ -479,7 +457,7 @@ void UMassISMDestroyRenderStateProcessor::Execute(FMassEntityManager& EntityMana
 				FMassEntityHandle EntityHandle = Context.GetEntity(*EntityIt);
 
 				ISMInfoFragments[*EntityIt].InstanceIndex = INDEX_NONE;
-				Context.Defer().PushCommand<FMassCommandRemoveElements<FMassRenderISMInstanceFragment, FMassRenderISMLinkFragment>>(EntityHandle);
+				Context.Defer().RemoveElements<FMassRenderISMInstanceFragment, FMassRenderISMLinkFragment>(EntityHandle);
 			}
 		});
 
diff --git a/Engine/Source/Runtime/MassEntity/Private/MassEntityManager.cpp b/Engine/Source/Runtime/MassEntity/Private/MassEntityManager.cpp
--- a/Engine/Source/Runtime/MassEntity/Private/MassEntityManager.cpp
+++ b/Engine/Source/Runtime/MassEntity/Private/MassEntityManager.cpp
@@ -739,8 +739,11 @@ int32 FMassEntityManager::BatchReserveEntities(TArrayView<FMassEntityHandle> InO
 	return GetEntityStorageInterface().Acquire(InOutEntities);
 }
 
-TSharedRef<FMassEntityManager::FEntityCreationContext> FMassEntityManager::BatchBuildEntities(const FMassArchetypeEntityCollectionWithPayload& EncodedEntitiesWithPayload
-	, const FMassFragmentBitSet& FragmentsAffected, const FMassArchetypeSharedFragmentValues& SharedFragmentValues, const FMassArchetypeCreationParams& CreationParams)
+TSharedRef<FMassEntityManager::FEntityCreationContext> FMassEntityManager::BatchBuildEntities(
+	const FMassArchetypeEntityCollectionWithPayload& EncodedEntitiesWithPayload
+	, const FMassFragmentBitSet& FragmentsAffected
+	, const FMassArchetypeSharedFragmentValues& SharedFragmentValues
+	, const FMassArchetypeCreationParams& CreationParams)
 {
 	CHECK_SYNC_API_RETURN(return FMassObserverManager::FCreationContext::DebugCreateDummyCreationContext());
 	check(SharedFragmentValues.IsSorted());
@@ -761,6 +764,15 @@ TSharedRef<FMassEntityManager::FEntityCreationContext> FMassEntityManager::Batch
 TSharedRef<FMassEntityManager::FEntityCreationContext> FMassEntityManager::BatchBuildEntities(const FMassArchetypeEntityCollectionWithPayload& EncodedEntitiesWithPayload
 	, const FMassArchetypeCompositionDescriptor& Composition
 	, const FMassArchetypeSharedFragmentValues& SharedFragmentValues, const FMassArchetypeCreationParams& CreationParams)
+{
+	return BatchBuildEntities(EncodedEntitiesWithPayload, Composition.GetElementsBitSet(), SharedFragmentValues, CreationParams);
+}
+
+TSharedRef<FMassEntityManager::FEntityCreationContext> FMassEntityManager::BatchBuildEntities(
+	const FMassArchetypeEntityCollectionWithPayload& EncodedEntitiesWithPayload
+	, const FMassElementBitSet& Composition
+	, const FMassArchetypeSharedFragmentValues& SharedFragmentValues
+	, const FMassArchetypeCreationParams& CreationParams)
 {
 	CHECK_SYNC_API_RETURN(return FMassObserverManager::FCreationContext::DebugCreateDummyCreationContext());
 
@@ -790,12 +802,12 @@ TSharedRef<FMassEntityManager::FEntityCreationContext> FMassEntityManager::Batch
 
 		const int32 SerialNumber = GetEntityStorageInterface().GetSerialNumber(Entity.Index);
 		Entity.SerialNumber = SerialNumber;
-		
+
 		GetEntityStorageInterface().SetArchetypeFromShared(Entity.Index, TargetArchetypeHandle.DataPtr);
 	}
 
 	TargetArchetypeHandle.DataPtr->BatchAddEntities(EntityHandles, SharedFragmentValues, TargetArchetypeEntityRanges);
-	UE_TRACE_MASS_ENTITIES_CREATED(EntityHandles, *TargetArchetypeHandle.DataPtr.Get())
+	UE_TRACE_MASS_ENTITIES_CREATED(EntityHandles, *TargetArchetypeHandle.DataPtr.Get());
 
 	if (EncodedEntitiesWithPayload.GetPayload().IsEmpty() == false)
 	{
@@ -1878,19 +1890,35 @@ bool FMassEntityManager::RemoveSharedFragmentFromEntity(const FMassEntityHandle
 
 void FMassEntityManager::BatchChangeCompositionForEntities(TConstArrayView<FMassArchetypeEntityCollection> EntityCollections, const FMassElementBitSet& ElementsToAdd, const FMassElementBitSet& ElementsToRemove)
 {
-	CHECK_SYNC_API();
-
-	// Adding shared or sparse element types via this API is not supported — shared fragments require
-	// values (use AddSharedFragmentToEntity), sparse elements have separate storage,
-	// and chunk fragments require dedicated APIs.
-	// Removing shared fragments is supported (extracted below and passed to BatchMoveEntitiesToAnotherArchetype).
+	// Adding shared or sparse element types via this overload is not supported -- shared fragments require
+	// values (use the overload with FMassArchetypeSharedFragmentValues, or AddSharedFragmentToEntity),
+	// sparse elements have separate storage, and chunk fragments require dedicated APIs.
+	// Removing shared fragments is supported (extracted in the implementation and passed to BatchMoveEntitiesToAnotherArchetype).
 	testableCheckfReturn(ElementsToAdd.ContainsOnlyTagsAndFragments() && ElementsToAdd.GetSparseElements().IsEmpty()
 		, return
 		, TEXT("%hs: ElementsToAdd contains shared, chunk, or sparse element types."
 			" Only non-sparse tags and fragments can be added via this API.")
 		, __FUNCTION__);
 
-	TRACE_CPUPROFILER_EVENT_SCOPE(Mass_BatchChangeCompositionForEntities);
+	static const FMassArchetypeSharedFragmentValues EmptySharedValues;
+	BatchChangeCompositionForEntities(EntityCollections, ElementsToAdd, ElementsToRemove, EmptySharedValues);
+}
+
+void FMassEntityManager::BatchChangeCompositionForEntities(
+	TConstArrayView<FMassArchetypeEntityCollection> EntityCollections,
+	const FMassElementBitSet& ElementsToAdd,
+	const FMassElementBitSet& ElementsToRemove,
+	const FMassArchetypeSharedFragmentValues& SharedFragmentValuesToAdd)
+{
+	CHECK_SYNC_API();
+
+	TRACE_CPUPROFILER_EVENT_SCOPE(Mass_BatchChangeCompositionForEntitiesWithSharedFragments);
+
+	testableCheckfReturn(ElementsToAdd.GetSparseElements().IsEmpty()
+		, return
+		, TEXT("%hs: ElementsToAdd contains sparse element types."
+			" Sparse elements have separate storage and cannot be added via composition change.")
+		, __FUNCTION__);
 
 	// Extract shared fragment removal bitsets for BatchMoveEntitiesToAnotherArchetype
 	const FMassSharedFragmentBitSet SharedToRemove = ElementsToRemove.Get<FMassSharedFragmentBitSet>();
@@ -1906,7 +1934,7 @@ void FMassEntityManager::BatchChangeCompositionForEntities(TConstArrayView<FMass
 		{
 			FMassElementBitSet ElementsAdded;
 			FMassElementBitSet NewComposition = CurrentArchetype->GetCompositionBitSet() + ElementsToAdd - ElementsToRemove;
-			check(NewComposition.IsEmpty()==false);
+			check(NewComposition.IsEmpty() == false);
 			if (UNLIKELY(CurrentArchetype->GetCompositionBitSet().IsEquivalent(NewComposition)))
 			{
 				continue;
@@ -1932,9 +1960,10 @@ void FMassEntityManager::BatchChangeCompositionForEntities(TConstArrayView<FMass
 			// Move the entity over; single move handles fragment/tag changes + shared fragment removal
 			FMassArchetypeEntityCollection::FEntityRangeArray NewArchetypeEntityRanges;
 			TArray<FMassEntityHandle> EntitiesBeingMoved;
+			const bool bAddSharedFragments = !SharedFragmentValuesToAdd.IsEmpty();
 			CurrentArchetype->BatchMoveEntitiesToAnotherArchetype(Collection, *NewArchetypeHandle.DataPtr.Get(), EntitiesBeingMoved
 				, bElementsAddedAreObserved && (ElementsAdded.IsEmpty() == false) ? &NewArchetypeEntityRanges : nullptr
-				, /*SharedFragmentValuesToAdd=*/nullptr
+				, bAddSharedFragments ? &SharedFragmentValuesToAdd : nullptr
 				, bRemoveSharedFragments ? &SharedToRemove : nullptr
 				, bRemoveSharedFragments ? &ConstSharedToRemove : nullptr);
 
@@ -1956,22 +1985,49 @@ void FMassEntityManager::BatchChangeCompositionForEntities(TConstArrayView<FMass
 		else // only valid if there are fragments in NewComposition
 		{
 			FMassElementBitSet NewComposition = (ElementsToAdd - ElementsToRemove);
-			BatchBuildEntities(FMassArchetypeEntityCollectionWithPayload(Collection), FMassArchetypeCompositionDescriptor(MoveTemp(NewComposition)));
+			BatchBuildEntities(FMassArchetypeEntityCollectionWithPayload(Collection), MoveTemp(NewComposition), SharedFragmentValuesToAdd);
 		}
 	}
 }
 
 void FMassEntityManager::BatchAddFragmentInstancesForEntities(TConstArrayView<FMassArchetypeEntityCollectionWithPayload> EntityCollections, const FMassFragmentBitSet& FragmentsAffected)
+{
+	static const FMassArchetypeSharedFragmentValues EmptySharedValues;
+	BatchAddFragmentInstancesForEntities(EntityCollections, FragmentsAffected, FMassTagBitSet(), EmptySharedValues);
+}
+
+void FMassEntityManager::BatchAddFragmentInstancesForEntities(TConstArrayView<FMassArchetypeEntityCollectionWithPayload> EntityCollections
+	, const FMassFragmentBitSet& FragmentsAffected, const FMassTagBitSet& TagsToAdd)
+{
+	static const FMassArchetypeSharedFragmentValues EmptySharedValues;
+	BatchAddFragmentInstancesForEntities(EntityCollections, FragmentsAffected, TagsToAdd, EmptySharedValues);
+}
+
+void FMassEntityManager::BatchAddFragmentInstancesForEntities(TConstArrayView<FMassArchetypeEntityCollectionWithPayload> EntityCollections
+	, const FMassFragmentBitSet& FragmentsAffected, const FMassArchetypeSharedFragmentValues& SharedFragmentValuesToAdd)
+{
+	BatchAddFragmentInstancesForEntities(EntityCollections, FragmentsAffected, FMassTagBitSet(), SharedFragmentValuesToAdd);
+}
+
+void FMassEntityManager::BatchAddFragmentInstancesForEntities(TConstArrayView<FMassArchetypeEntityCollectionWithPayload> EntityCollections
+	, const FMassFragmentBitSet& FragmentsAffected, const FMassTagBitSet& TagsToAdd, const FMassArchetypeSharedFragmentValues& SharedFragmentValuesToAdd)
 {
 	CHECK_SYNC_API();
 
 	TRACE_CPUPROFILER_EVENT_SCOPE(Mass_BatchAddFragmentInstancesForEntities);
 
+	testableCheckfReturn(FMassElementBitSet(TagsToAdd).GetSparseElements().IsEmpty()
+		, return
+		, TEXT("%hs: TagsToAdd contains sparse tag types."
+			" Sparse elements must be handled at the command level via BatchAddSparseElementToEntities before calling this function.")
+		, __FUNCTION__);
+
+	// Sparse fragments are excluded from composition but their values are set via BatchSetFragmentValues.
 	const FMassElementBitSet SparseElementsAffected = FragmentsAffected & FMassElementBitSet::GetAllSparseElements();
-	FMassElementBitSet RegularElementsAffected = FragmentsAffected;
+	FMassElementBitSet RegularFragmentsBits = FragmentsAffected;
 	if (SparseElementsAffected.IsEmpty() == false)
 	{
-		RegularElementsAffected -= SparseElementsAffected;
+		RegularFragmentsBits -= SparseElementsAffected;
 	}
 
 	// here's the scenario:
@@ -1981,10 +2037,11 @@ void FMassEntityManager::BatchAddFragmentInstancesForEntities(TConstArrayView<FM
 	//		* there might be an "empty" archetype in the mix - then step A results in archetype creation and assigning
 	//		* if step A is required then the initial FMassArchetypeEntityCollection instance is no longer valid
 	// * setting value can be done uniformly for all entities, remembering some might be in different chunks already
-	// * @todo note that after adding fragment type some entities originally in different archetypes end up in the same 
+	// * @todo note that after adding fragment type some entities originally in different archetypes end up in the same
 	//		archetype. This could be utilized as a basis for optimization. To be investigated.
 	//
-	const bool bFragmentsAddedAreObserved = ObserverManager.HasObserversForBitSet(FragmentsAffected, EMassObservedOperation::AddElement);
+	const FMassElementBitSet AllRegularElementsToAdd = RegularFragmentsBits + TagsToAdd + SharedFragmentValuesToAdd.GetBitSet();
+	const bool bElementsAddedAreObserved = ObserverManager.HasObserversForBitSet(AllRegularElementsToAdd, EMassObservedOperation::AddElement);
 
 	for (const FMassArchetypeEntityCollectionWithPayload& EntityRangesWithPayload : EntityCollections)
 	{
@@ -1994,28 +2051,29 @@ void FMassEntityManager::BatchAddFragmentInstancesForEntities(TConstArrayView<FM
 		if (CurrentArchetype)
 		{
 			FMassArchetypeEntityCollection::FEntityRangeArray TargetArchetypeEntityRanges;
-			FMassElementBitSet NewComposition = RegularElementsAffected + CurrentArchetype->GetCompositionBitSet();
 			FMassElementBitSet ElementsAdded;
 
+			// New composition includes regular fragments, non-sparse tags, and shared fragment type bits
+			FMassElementBitSet NewComposition = AllRegularElementsToAdd + CurrentArchetype->GetCompositionBitSet();
+
 			if (CurrentArchetype->GetCompositionBitSet().IsEquivalent(NewComposition) == false)
 			{
-				if (bFragmentsAddedAreObserved)
+				if (bElementsAddedAreObserved)
 				{
-					ElementsAdded = RegularElementsAffected - CurrentArchetype->GetCompositionBitSet();
+					ElementsAdded = AllRegularElementsToAdd - CurrentArchetype->GetCompositionBitSet();
 				}
 
 				FMassArchetypeHandle NewArchetypeHandle = InternalCreateSimilarArchetype(TargetArchetypeHandle.DataPtr, MoveTemp(NewComposition));
 				checkSlow(NewArchetypeHandle.IsValid());
 
-				// Move the entity over
 				TArray<FMassEntityHandle> EntitiesBeingMoved;
+				const bool bAddSharedFragments = !SharedFragmentValuesToAdd.IsEmpty();
 				CurrentArchetype->BatchMoveEntitiesToAnotherArchetype(EntityRangesWithPayload.GetEntityCollection(), *NewArchetypeHandle.DataPtr.Get()
-					, EntitiesBeingMoved, &TargetArchetypeEntityRanges);
+					, EntitiesBeingMoved, &TargetArchetypeEntityRanges, bAddSharedFragments ? &SharedFragmentValuesToAdd : nullptr);
 
 				for (const FMassEntityHandle& Entity : EntitiesBeingMoved)
 				{
 					check(GetEntityStorageInterface().IsValidIndex(Entity.Index));
-
 					GetEntityStorageInterface().SetArchetypeFromShared(Entity.Index, NewArchetypeHandle.DataPtr);
 				}
 
@@ -2025,14 +2083,25 @@ void FMassEntityManager::BatchAddFragmentInstancesForEntities(TConstArrayView<FM
 			{
 				// @todo consider adding sparse elements to ElementsAdded once we start observing those.
 				TargetArchetypeEntityRanges = EntityRangesWithPayload.GetEntityCollection().GetRanges();
+				// The entity's composition already contains all target element types (including the shared fragment type bits from SharedFragmentValuesToAdd),
+				// so no archetype move occurs. Updating shared fragment VALUES for already-present types requires a chunk move,
+				// which BatchMoveEntitiesToAnotherArchetype does not support within the same archetype.
+				// This matches the explicit limitation in BatchAddSharedFragmentsForEntities.
+				testableCheckfReturn(SharedFragmentValuesToAdd.IsEmpty()
+					, continue
+					, TEXT("%hs: SharedFragmentValuesToAdd is non-empty but the entity composition is unchanged."
+						" Updating shared fragment values for already-present types is not supported via this API.")
+					, __FUNCTION__);
 			}
 
-			// at this point all the entities are in the target archetype, we can set the values
-			// note that even though the "subchunk" information could have changed the order of entities is the same and 
+			// Set per-entity fragment values (order of entities is preserved through the move).
+			// At this point all the entities are in the target archetype, we can set the values.
+			// Note that even though the "subchunk" information could have changed, the order of entities is the same and 
 			// corresponds to the order in FMassArchetypeEntityCollectionWithPayload's payload
+			// The payload must NOT contain tag arrays; tags have zero-size storage and are not in FragmentIndexMap.
 			TargetArchetypeHandle.DataPtr->BatchSetFragmentValues(TargetArchetypeEntityRanges, EntityRangesWithPayload.GetPayload());
 
-			if (bFragmentsAddedAreObserved && (ElementsAdded.IsEmpty() == false))
+			if (bElementsAddedAreObserved && (ElementsAdded.IsEmpty() == false))
 			{
 				ObserverManager.OnCompositionChanged(
 					FMassArchetypeEntityCollection(TargetArchetypeHandle, MoveTemp(TargetArchetypeEntityRanges))
@@ -2040,9 +2109,12 @@ void FMassEntityManager::BatchAddFragmentInstancesForEntities(TConstArrayView<FM
 					, EMassObservedOperation::AddElement);
 			}
 		}
-		else 
+		else
 		{
-			BatchBuildEntities(EntityRangesWithPayload, FragmentsAffected, FMassArchetypeSharedFragmentValues());
+			// Archetype-less entity: build with full composition (sparse elements handled at command level).
+			// AllRegularElementsToAdd already combines non-sparse fragments, tags, and shared fragment type bits,
+			// so a single path handles both the tag and no-tag cases.
+			BatchBuildEntities(EntityRangesWithPayload, AllRegularElementsToAdd, SharedFragmentValuesToAdd);
 		}
 	}
 }
diff --git a/Engine/Source/Runtime/MassEntity/Public/MassCommandBuffer.h b/Engine/Source/Runtime/MassEntity/Public/MassCommandBuffer.h
--- a/Engine/Source/Runtime/MassEntity/Public/MassCommandBuffer.h
+++ b/Engine/Source/Runtime/MassEntity/Public/MassCommandBuffer.h
@@ -287,6 +287,22 @@ struct FMassCommandBuffer
 		PushCommand<FMassCommandAddElements<TTypes...>>(Entities);
 	}
 
+	/** Adds non-shared element types AND shared fragment values to an entity. Per-hash grouping for batching. */
+	template<typename... TTypes>
+	void AddElementsWithSharedFragments(FMassEntityHandle Entity, FMassArchetypeSharedFragmentValues&& SharedValues)
+	{
+		PushCommand<FMassCommandAddElementsWithSharedFragments<TTypes...>>(Entity, MoveTemp(SharedValues));
+	}
+
+	/** Adds element types with per-entity values AND shared fragment values to an entity in a single move. Per-hash grouping.
+	 *  TElements accepts FMassFragment types (per-entity values required) and FMassTag types including sparse
+	 *  (tags participate in composition but are excluded from the value payload due to zero-size storage). */
+	template<typename... TElements>
+	void AddFragmentInstancesWithSharedFragments(FMassEntityHandle Entity, FMassArchetypeSharedFragmentValues&& SharedValues, TElements&&... InElements)
+	{
+		PushCommand<FMassCommandAddFragmentInstancesWithSharedFragments<FMassArchetypeSharedFragmentValues, UE::Mass::Clean<TElements>...>>(Entity, MoveTemp(SharedValues), Forward<TElements>(InElements)...);
+	}
+
 	/** Adds any mix of non-shared element types to an entity. Types provided at runtime. Uses PushUniqueCommand. */
 	void AddElements(FMassEntityHandle Entity, TConstArrayView<const UScriptStruct*> ElementTypes)
 	{
diff --git a/Engine/Source/Runtime/MassEntity/Public/MassCommands.h b/Engine/Source/Runtime/MassEntity/Public/MassCommands.h
--- a/Engine/Source/Runtime/MassEntity/Public/MassCommands.h
+++ b/Engine/Source/Runtime/MassEntity/Public/MassCommands.h
@@ -12,80 +12,94 @@
 #include "MassCommands.generated.h"
 
 /**
- * ========================================================================================================
+ * =====================================================================================================================
  * Mass Commands Quick Reference
- * ========================================================================================================
+ * =====================================================================================================================
  *
  * ADDING ELEMENTS (composition-only, no initial values):
  * -------------------------------------------------------
- * | Need                           | Command                                    | Batching     | BP  |
- * |--------------------------------|--------------------------------------------|--------------|-----|
- * | Add fragments/tags/sparse      | FMassCommandAddElements<T...>              | Cross-chunk  | Yes |
- * | (compile-time types)           | Convenience: AddElements<T...>()           |              |     |
- * |--------------------------------|--------------------------------------------|--------------|-----|
- * | Add fragments/tags/sparse      | FMassCommandAddElementList                 | Per-instance | No  |
- * | (runtime types)                | Convenience: AddElements(E, T)             |              |     |
- * |--------------------------------|--------------------------------------------|--------------|-----|
- * | Add single type (runtime)      | FMassCommandAddElement                     | Per-instance | No  |
+ * | Need                       | Command                                                    | Batching    | BP  |
+ * |----------------------------|------------------------------------------------------------|-------------|-----|
+ * | Add fragments/tags/sparse  | FMassCommandAddElements<T...>                              | Cross-chunk | Yes |
+ * | (compile-time types)       | Convenience: AddElements<T...>()                           |             |     |
+ * |----------------------------|------------------------------------------------------------|-------------|-----|
+ * | Add fragments/tags/sparse  | FMassCommandAddElementList                                 | Per-inst.   | No  |
+ * | (runtime types)            | Convenience: AddElements(E, T)                             |             |     |
+ * |----------------------------|------------------------------------------------------------|-------------|-----|
+ * | Add single type (runtime)  | FMassCommandAddElement                                     | Per-inst.   | No  |
+ *
+ * ADDING FRAGMENTS WITH VALUES (+ optional tags/sparse):
+ * --------------------------------------------------------
+ * | Need                       | Command                                                    | Batching    | BP  |
+ * |----------------------------|------------------------------------------------------------|-------------|-----|
+ * | Add fragments with values  | FMassCommandAddFragmentInstances<T...>                     | Cross-chunk | Yes |
+ * | + optional tags/sparse     | T... accepts FMassFragment, FMassTag, and sparse subtypes  |             |     |
+ * |                            | Tags/sparse tags: composition only. Sparse frags: in-place |             |     |
  *
- * ADDING FRAGMENTS WITH VALUES:
- * ------------------------------
- * | Need                           | Command                                    | Batching     | BP  |
- * |--------------------------------|--------------------------------------------|--------------|-----|
- * | Add fragments with values      | FMassCommandAddFragmentInstances<T...>     | Cross-chunk  | Yes |
- * | (sets initial data per entity) |                                            |              |     |
+ * ADDING ELEMENTS WITH SHARED FRAGMENT VALUES (+ optional tags/sparse):
+ * ----------------------------------------------------------------------
+ * | Need                       | Command                                                    | Batching    | BP  |
+ * |----------------------------|------------------------------------------------------------|-------------|-----|
+ * | Add fragments/tags/sparse  | FMassCommandAddElementsWithSharedFragments<T...>           | Per-hash    | No  |
+ * | + shared fragment values   | Convenience: AddElementsWithSharedFragments<T...>(E, V)    | group       |     |
+ * |----------------------------|------------------------------------------------------------|-------------|-----|
+ * | Add fragments with values  | FMassCommandAddFragmentInstancesWithSharedFragments        | Per-hash    | Yes |
+ * | + shared values + tags     | <TShared, T...> T... accepts fragments, tags, sparse       | group       |     |
  *
  * REMOVING ELEMENTS:
  * -------------------
- * | Need                           | Command                                    | Batching     | BP  |
- * |--------------------------------|--------------------------------------------|--------------|-----|
- * | Remove any element types       | FMassCommandRemoveElements<T...>           | Cross-chunk  | Yes |
- * | (compile-time types)           | Convenience: RemoveElements<T...>()        |              |     |
- * |--------------------------------|--------------------------------------------|--------------|-----|
- * | Remove any element types       | FMassCommandRemoveElementList              | Per-instance | No  |
- * | (runtime types)                | Convenience: RemoveElements(E, T)          |              |     |
- * |--------------------------------|--------------------------------------------|--------------|-----|
- * | Remove single type (runtime)   | FMassCommandRemoveElement                  | Per-instance | No  |
+ * | Need                       | Command                                                    | Batching    | BP  |
+ * |----------------------------|------------------------------------------------------------|-------------|-----|
+ * | Remove any element types   | FMassCommandRemoveElements<T...>                           | Cross-chunk | Yes |
+ * | (compile-time types)       | Convenience: RemoveElements<T...>()                        |             |     |
+ * |----------------------------|------------------------------------------------------------|-------------|-----|
+ * | Remove any element types   | FMassCommandRemoveElementList                              | Per-inst.   | No  |
+ * | (runtime types)            | Convenience: RemoveElements(E, T)                          |             |     |
+ * |----------------------------|------------------------------------------------------------|-------------|-----|
+ * | Remove single (runtime)    | FMassCommandRemoveElement                                  | Per-inst.   | No  |
  *
  * CHANGING COMPOSITION:
  * ----------------------
- * | Need                           | Command                                    | Batching     | BP  |
- * |--------------------------------|--------------------------------------------|--------------|-----|
- * | Swap one tag for another       | FMassCommandSwapTags<TOld, TNew>           | Cross-chunk  | No  |
- * | (compile-time types)           | Convenience: SwapTags<TOld, TNew>()        |              |     |
+ * | Need                       | Command                                                    | Batching    | BP  |
+ * |----------------------------|------------------------------------------------------------|-------------|-----|
+ * | Swap one tag for another   | FMassCommandSwapTags<TOld, TNew>                           | Cross-chunk | No  |
+ * | (compile-time types)       | Convenience: SwapTags<TOld, TNew>()                        |             |     |
  *
- * CREATING ENTITIES:
- * -------------------
- * | Need                           | Command                                    | Batching     | BP  |
- * |--------------------------------|--------------------------------------------|--------------|-----|
- * | Create entity with values      | FMassCommandBuildEntity<T...>              | Cross-chunk  | Yes |
- * | Create with shared fragments   | FMassCommandBuildEntity                    | Cross-chunk  | Yes |
- * |                                | WithSharedFragments<TShared, T...>         |              |     |
+ * CREATING ENTITIES (+ optional tags/sparse):
+ * ----------------------------------------------
+ * | Need                       | Command                                                    | Batching    | BP  |
+ * |----------------------------|------------------------------------------------------------|-------------|-----|
+ * | Create entity with values  | FMassCommandBuildEntity<T...>                              | Cross-chunk | Yes |
+ * | + optional tags/sparse     | T... accepts FMassFragment, FMassTag, and sparse subtypes  |             |     |
+ * |----------------------------|------------------------------------------------------------|-------------|-----|
+ * | Create with shared values  | FMassCommandBuildEntityWithSharedFragments<TShared, T...>  | Per-hash    | Yes |
+ * | + optional tags/sparse     | T... accepts fragments, tags, sparse                       | group       |     |
  *
  * DESTROYING ENTITIES:
  * ---------------------
- * | Need                           | Command                                    | Batching     | BP  |
- * |--------------------------------|--------------------------------------------|--------------|-----|
- * | Destroy entities               | FMassCommandDestroyEntities                | Cross-chunk  | Yes |
- * |                                | Convenience: DestroyEntity/Entities()      |              |     |
+ * | Need                       | Command                                                    | Batching    | BP  |
+ * |----------------------------|------------------------------------------------------------|-------------|-----|
+ * | Destroy entities           | FMassCommandDestroyEntities                                | Cross-chunk | Yes |
+ * |                            | Convenience: DestroyEntity/Entities()                      |             |     |
  *
  * DEFERRED (LAMBDA):
  * -------------------
- * | Need                           | Command                                    | Batching     | BP  |
- * |--------------------------------|--------------------------------------------|--------------|-----|
- * | Arbitrary logic at flush time  | FMassDeferredCreateCommand                 | Cross-chunk  | No  |
- * | (type depends on operation)    | FMassDeferredAddCommand                    |              |     |
- * |                                | FMassDeferredRemoveCommand                 |              |     |
- * |                                | FMassDeferredChangeCompositionCommand      |              |     |
- * |                                | FMassDeferredSetCommand                    |              |     |
- * |                                | FMassDeferredDestroyCommand                |              |     |
+ * | Need                       | Command                                                    | Batching    | BP  |
+ * |----------------------------|------------------------------------------------------------|-------------|-----|
+ * | Arbitrary logic at flush   | FMassDeferredCreateCommand                                 | Cross-chunk | No  |
+ * | (type depends on op)       | FMassDeferredAddCommand                                    |             |     |
+ * |                            | FMassDeferredRemoveCommand                                 |             |     |
+ * |                            | FMassDeferredChangeCompositionCommand                      |             |     |
+ * |                            | FMassDeferredSetCommand                                    |             |     |
+ * |                            | FMassDeferredDestroyCommand                                |             |     |
  *
  * NOTES:
- * - Cross-chunk batching: PushCommand reuses one instance across all chunks; one Run() for all entities.
- * - Per-instance batching: PushUniqueCommand creates a new instance per call. Entities batch via Add().
- * - Adding shared/const shared fragments requires values; use AddSharedFragmentToEntity directly.
+ * - Cross-chunk: PushCommand reuses one instance across all chunks; one Run() for all entities.
+ * - Per-inst.: PushUniqueCommand creates a new instance per call. Entities batch via Add().
+ * - Per-hash group: entities grouped by shared fragment value hash; one move per group.
+ * - Adding shared/const shared fragments requires values; use the WithSharedFragments variant.
  * - Removing shared/const shared fragments is supported by RemoveElements and RemoveElement.
- * - Sparse elements are handled in-place (no archetype move). Non-sparse types are combined into a single entity move.
+ * - Sparse elements are handled in-place (no archetype move). Non-sparse types: single entity move.
  * - BP = Breakpoint support (Mass debugger breakpoints triggered at PushCommand time).
  * - Deferred commands store a TFunction<void(FMassEntityManager&)>; no type safety, full flexibility.
  * ========================================================================================================
@@ -482,9 +496,8 @@ protected:
 		TArray<FMassArchetypeEntityCollection> EntityCollections;
 		UE::Mass::Utils::CreateEntityCollections(EntityManager, TargetEntities, FMassArchetypeEntityCollection::FoldDuplicates, EntityCollections);
 
-		// Sparse removal is always in-place (no entity move)
-		// Process first before changing the composition of the entities that will invalidate
-		// the entity collection
+		// Sparse element modification is always in-place (no entity move)
+		// Process first before changing the composition of the entities that will invalidate the entity collections
 		for (const UScriptStruct* SparseType : SparseTypes)
 		{
 			EntityManager.BatchRemoveSparseElementFromEntities(EntityCollections, SparseType);
@@ -601,9 +614,8 @@ protected:
 		TArray<FMassArchetypeEntityCollection> EntityCollections;
 		UE::Mass::Utils::CreateEntityCollections(EntityManager, TargetEntities, FMassArchetypeEntityCollection::FoldDuplicates, EntityCollections);
 
-		// Sparse removal is always in-place (no entity move)
-		// Process first before changing the composition of the entities that will invalidate
-		// the entity collection
+		// Sparse element modification is always in-place (no entity move)
+		// Process first before changing the composition of the entities that will invalidate the entity collections
 		for (const UScriptStruct* SparseType : SparseTypes)
 		{
 			EntityManager.BatchRemoveSparseElementFromEntities(EntityCollections, SparseType);
@@ -695,19 +707,22 @@ protected:
 		TArray<FMassArchetypeEntityCollection> EntityCollections;
 		UE::Mass::Utils::CreateEntityCollections(EntityManager, TargetEntities, FMassArchetypeEntityCollection::FoldDuplicates, EntityCollections);
 
-		if (!NonSparseToAdd.IsEmpty())
-		{
-			EntityManager.BatchChangeCompositionForEntities(EntityCollections, NonSparseToAdd, FMassElementBitSet());
-		}
+		// Sparse element modification is always in-place (no entity move)
+		// Process first before changing the composition of the entities that will invalidate the entity collections
 		for (const UScriptStruct* SparseType : SparseTypes)
 		{
 			EntityManager.BatchAddSparseElementToEntities(EntityCollections, SparseType);
 		}
+
+		if (!NonSparseTypes.IsEmpty())
+		{
+			EntityManager.BatchChangeCompositionForEntities(EntityCollections, NonSparseTypes, FMassElementBitSet());
+		}
 	}
 
 	virtual SIZE_T GetAllocatedSize() const override
 	{
-		return Super::GetAllocatedSize() + NonSparseToAdd.GetAllocatedSize() + SparseTypes.GetAllocatedSize();
+		return Super::GetAllocatedSize() + NonSparseTypes.GetAllocatedSize() + SparseTypes.GetAllocatedSize();
 	}
 
 private:
@@ -721,11 +736,11 @@ private:
 		}
 		else
 		{
-			NonSparseToAdd.Add(FClean::StaticStruct());
+			NonSparseTypes.Add(FClean::StaticStruct());
 		}
 	}
 
-	FMassElementBitSet NonSparseToAdd;
+	FMassElementBitSet NonSparseTypes;
 	TArray<const UScriptStruct*> SparseTypes;
 };
 
@@ -766,7 +781,7 @@ struct FMassCommandAddElementList : public FMassBatchedEntityCommand
 			}
 			else
 			{
-				NonSparseToAdd.Add(ElementType);
+				NonSparseTypes.Add(ElementType);
 			}
 		}
 	}
@@ -779,19 +794,22 @@ protected:
 		TArray<FMassArchetypeEntityCollection> EntityCollections;
 		UE::Mass::Utils::CreateEntityCollections(EntityManager, TargetEntities, FMassArchetypeEntityCollection::FoldDuplicates, EntityCollections);
 
-		if (!NonSparseToAdd.IsEmpty())
-		{
-			EntityManager.BatchChangeCompositionForEntities(EntityCollections, NonSparseToAdd, FMassElementBitSet());
-		}
+		// Sparse element modification is always in-place (no entity move)
+		// Process first before changing the composition of the entities that will invalidate the entity collections
 		for (const UScriptStruct* SparseType : SparseTypes)
 		{
 			EntityManager.BatchAddSparseElementToEntities(EntityCollections, SparseType);
 		}
+
+		if (!NonSparseTypes.IsEmpty())
+		{
+			EntityManager.BatchChangeCompositionForEntities(EntityCollections, NonSparseTypes, FMassElementBitSet());
+		}
 	}
 
 	virtual SIZE_T GetAllocatedSize() const override
 	{
-		return Super::GetAllocatedSize() + NonSparseToAdd.GetAllocatedSize() + SparseTypes.GetAllocatedSize();
+		return Super::GetAllocatedSize() + NonSparseTypes.GetAllocatedSize() + SparseTypes.GetAllocatedSize();
 	}
 
 #if CSV_PROFILER_STATS || WITH_MASSENTITY_DEBUG
@@ -801,7 +819,7 @@ protected:
 	}
 #endif // CSV_PROFILER_STATS || WITH_MASSENTITY_DEBUG
 
-	FMassElementBitSet NonSparseToAdd;
+	FMassElementBitSet NonSparseTypes;
 	TArray<const UScriptStruct*> SparseTypes;
 };
 
@@ -814,6 +832,343 @@ struct UE::Mass::Command::TCommandTraits<FMassCommandAddElementList> final
 	};
 };
 
+/**
+ * Adds any mix of non-shared element types AND shared fragment values to existing entities.
+ * Types are template parameters for non-shared elements.
+ * Shared fragment values are provided at runtime.
+ * Entities are grouped by shared fragment value hash for efficient per-group archetype changes.
+ * Performs a single entity move per group (non-sparse composition + shared values combined).
+ *
+ * When to use: when you need to add fragments/tags AND shared fragment values to existing entities
+ * in a single operation. Replaces separate AddElements + BatchAddSharedFragmentsForEntities calls.
+ *
+ * Example:
+ *   FMassArchetypeSharedFragmentValues SharedValues;
+ *   SharedValues.Add(EntityManager->GetOrCreateConstSharedFragment(FMySharedFragment(42)));
+ *   Context.Defer().PushCommand<FMassCommandAddElementsWithSharedFragments<FMyFragment, FMyTag>>(
+ *       EntityHandle, MoveTemp(SharedValues));
+ *
+ * Supported types (TTypes): FMassFragment + FMassTag (including sparse).
+ * Rejected types (TTypes): FMassSharedFragment, FMassConstSharedFragment (use shared values param instead), FMassChunkFragment.
+ * Breakpoint support: No (composition-only command).
+ * Entity moves: 1 move per shared-value group. Sparse addition is in-place.
+ */
+template<typename... TTypes>
+struct FMassCommandAddElementsWithSharedFragments : public FMassBatchedCommand
+{
+	static_assert((!TIsDerivedFrom<UE::Mass::Clean<TTypes>, FMassSharedFragment>::Value && ...),
+		"Shared fragments cannot be passed as TTypes -- provide them via FMassArchetypeSharedFragmentValues instead.");
+	static_assert((!TIsDerivedFrom<UE::Mass::Clean<TTypes>, FMassConstSharedFragment>::Value && ...),
+		"Const shared fragments cannot be passed as TTypes -- provide them via FMassArchetypeSharedFragmentValues instead.");
+	static_assert((!TIsDerivedFrom<UE::Mass::Clean<TTypes>, FMassChunkFragment>::Value && ...),
+		"Chunk fragments cannot be passed to FMassCommandAddElementsWithSharedFragments.");
+
+	using Super = FMassBatchedCommand;
+
+	FMassCommandAddElementsWithSharedFragments()
+		: Super(EMassCommandOperationType::Add DEBUG_NAME("AddElementsWithSharedFragments"))
+	{
+		(CategorizeElementType<TTypes>(), ...);
+	}
+
+	/** @note Callers must not Add() the same entity with different shared values — it would end up in two hash groups,
+	 *  and the second group's archetype move would overwrite the first group's without merging fragment values. */
+	void Add(FMassEntityHandle Entity, FMassArchetypeSharedFragmentValues&& InSharedFragments)
+	{
+		InSharedFragments.Sort();
+		const uint32 Hash = GetTypeHash(InSharedFragments);
+		FPerSharedFragmentsHashData& Instance = Data.FindOrAdd(Hash, MoveTemp(InSharedFragments));
+		Instance.TargetEntities.Add(Entity);
+		bHasWork = true;
+	}
+
+protected:
+	virtual void Run(FMassEntityManager& EntityManager) override
+	{
+		TRACE_CPUPROFILER_EVENT_SCOPE(MassCommandAddElementsWithSharedFragments_Execute);
+
+		// @todo: add inline arrays support in CreateEntityCollections
+		TArray<FMassArchetypeEntityCollection> EntityCollections;
+
+		for (auto& It : Data)
+		{
+			UE::Mass::Utils::CreateEntityCollections(EntityManager, It.Value.TargetEntities,
+				FMassArchetypeEntityCollection::FoldDuplicates, EntityCollections);
+
+			// Sparse element modification is always in-place (no entity move)
+			// Process first before changing the composition of the entities that will invalidate the entity collections
+			for (const UScriptStruct* SparseType : SparseTypes)
+			{
+				EntityManager.BatchAddSparseElementToEntities(EntityCollections, SparseType);
+			}
+
+			// Combine non-sparse element types + shared fragment bits from values
+			FMassElementBitSet AllToAdd = NonSparseTypes;
+			AllToAdd += It.Value.SharedFragmentValues.GetBitSet();
+
+			EntityManager.BatchChangeCompositionForEntities(
+				EntityCollections, AllToAdd, /*ElementsToRemove*/FMassElementBitSet(), It.Value.SharedFragmentValues);
+
+			EntityCollections.Reset();
+		}
+	}
+
+	virtual void Reset() override
+	{
+		Data.Reset();
+		Super::Reset();
+	}
+
+	virtual SIZE_T GetAllocatedSize() const override
+	{
+		SIZE_T TotalSize = NonSparseTypes.GetAllocatedSize() + SparseTypes.GetAllocatedSize();
+		for (const auto& KeyValue : Data)
+		{
+			TotalSize += KeyValue.Value.TargetEntities.GetAllocatedSize()
+				+ KeyValue.Value.SharedFragmentValues.GetAllocatedSize();
+		}
+		TotalSize += Data.GetAllocatedSize();
+		return TotalSize;
+	}
+
+#if CSV_PROFILER_STATS || WITH_MASSENTITY_DEBUG
+	virtual int32 GetNumOperationsStat() const override
+	{
+		int32 TotalCount = 0;
+		for (const auto& KeyValue : Data)
+		{
+			TotalCount += KeyValue.Value.TargetEntities.Num();
+		}
+		return TotalCount;
+	}
+#endif // CSV_PROFILER_STATS || WITH_MASSENTITY_DEBUG
+
+private:
+	template<typename T>
+	void CategorizeElementType()
+	{
+		using FClean = UE::Mass::Clean<T>;
+		if constexpr (UE::Mass::CSparse<FClean>)
+		{
+			SparseTypes.Add(FClean::StaticStruct());
+		}
+		else
+		{
+			NonSparseTypes.Add(FClean::StaticStruct());
+		}
+	}
+
+	struct FPerSharedFragmentsHashData
+	{
+		FPerSharedFragmentsHashData(FMassArchetypeSharedFragmentValues&& InSharedFragmentValues)
+			: SharedFragmentValues(MoveTemp(InSharedFragmentValues))
+		{}
+
+		SIZE_T GetAllocatedSize() const
+		{
+			return TargetEntities.GetAllocatedSize() + SharedFragmentValues.GetAllocatedSize();
+		}
+
+		TArray<FMassEntityHandle> TargetEntities;
+		FMassArchetypeSharedFragmentValues SharedFragmentValues;
+	};
+
+	FMassElementBitSet NonSparseTypes;
+	TArray<const UScriptStruct*> SparseTypes;
+	TMap<uint32, FPerSharedFragmentsHashData> Data;
+};
+
+/**
+ * Adds fragment types with per-entity values AND shared fragment values to existing entities in a single entity move.
+ * Extends FMassCommandAddFragmentInstances with shared fragment value support; entities are partitioned by shared
+ * fragment value hash at Add() time and dispatched in one BatchAddFragmentInstancesForEntities call per group.
+ * When no shared fragment values are needed, prefer FMassCommandAddFragmentInstances — its flat single-batch
+ * storage avoids the per-Add hash computation and TMap overhead.
+ * Uses PushCommand for cross-chunk batching.
+ *
+ * Note: TSharedFragmentValues is always FMassArchetypeSharedFragmentValues but is a template param
+ * to maintain uniform PushCommand interface (all params in one typename... list).
+ *
+ * When to use: when you need to add fragments with initial values AND shared fragment values to existing entities in one operation.
+ * Replaces separate AddFragmentInstances + deferred lambda AddConstSharedFragmentToEntity calls.
+ *
+ * Example:
+ *   FMassArchetypeSharedFragmentValues SharedValues;
+ *   SharedValues.Add(EntityManager->GetOrCreateConstSharedFragment(FMySharedFragment(42)));
+ *   Context.Defer().PushCommand<FMassCommandAddFragmentInstancesWithSharedFragments>(
+ *       EntityHandle, MoveTemp(SharedValues), FMyFragmentA(1), FMyFragmentB(2.0f));
+ *
+ * Supported types (TTypes): FMassFragment + FMassTag (including sparse).
+ * Tags are included in composition but excluded from value payload (zero-size storage).
+ * Shared fragments via FMassArchetypeSharedFragmentValues.
+ * Breakpoint support: Yes -- CheckFragmentAddBreakpoints called at PushCommand time.
+ * Entity moves: 1 move per shared-value group.
+ * Fragment values set in-place after move.
+ */
+template<typename TSharedFragmentValues, typename... TTypes>
+struct FMassCommandAddFragmentInstancesWithSharedFragments : public FMassBatchedCommand
+{
+	static_assert(((TIsDerivedFrom<UE::Mass::Clean<TTypes>, FMassFragment>::Value
+		|| TIsDerivedFrom<UE::Mass::Clean<TTypes>, FMassTag>::Value) && ...),
+		"All types passed to FMassCommandAddFragmentInstancesWithSharedFragments must derive from FMassFragment or FMassTag.");
+	static_assert((!TIsDerivedFrom<UE::Mass::Clean<TTypes>, FMassSharedFragment>::Value && ...),
+		"Shared fragments cannot be passed as TTypes -- provide them via FMassArchetypeSharedFragmentValues instead.");
+	static_assert((!TIsDerivedFrom<UE::Mass::Clean<TTypes>, FMassConstSharedFragment>::Value && ...),
+		"Const shared fragments cannot be passed as TTypes -- provide them via FMassArchetypeSharedFragmentValues instead.");
+	static_assert((!TIsDerivedFrom<UE::Mass::Clean<TTypes>, FMassChunkFragment>::Value && ...),
+		"Chunk fragments cannot be passed to FMassCommandAddFragmentInstancesWithSharedFragments.");
+
+	using Super = FMassBatchedCommand;
+
+	FMassCommandAddFragmentInstancesWithSharedFragments(EMassCommandOperationType OperationType = EMassCommandOperationType::Set DEBUG_NAME_PARAM("AddFragmentInstancesWithSharedFragments"))
+		: Super(OperationType FORWARD_DEBUG_NAME_PARAM)
+	{
+		(CategorizeElementType<TTypes>(), ...);
+	}
+
+	/** @note Callers must not Add() the same entity with different shared values — it would end up in two hash groups,
+	 *  and the second group's archetype move would overwrite the first group's without merging fragment values. */
+	void Add(FMassEntityHandle Entity, FMassArchetypeSharedFragmentValues&& InSharedFragments, TTypes... InFragments)
+	{
+		InSharedFragments.Sort();
+
+		// Compute hash before adding to the map since evaluation order is not guaranteed
+		// and MoveTemp will invalidate InSharedFragments
+		const uint32 Hash = GetTypeHash(InSharedFragments);
+
+		FPerSharedFragmentsHashData& Instance = Data.FindOrAdd(Hash, MoveTemp(InSharedFragments));
+		Instance.Fragments.Add(InFragments...);
+		Instance.TargetEntities.Add(Entity);
+
+		bHasWork = true;
+	}
+
+#if WITH_MASSENTITY_DEBUG
+	static bool CheckBreakpoints(const FMassEntityHandle Entity, TTypes... InFragments)
+	{
+		// Function pointer to force non-template overload resolution.
+		// Without this, MSVC deduces TFragments=UScriptStruct* from StaticStruct() return type,
+		// matching the variadic template overload instead of the (FMassEntityHandle, const UScriptStruct*) overload.
+		using FCheckFn = bool(*)(const FMassEntityHandle&, const UScriptStruct*);
+		const FCheckFn Check = &UE::Mass::Debug::FBreakpoint::CheckFragmentAddBreakpoints;
+		return (Check(Entity, UE::Mass::Clean<TTypes>::StaticStruct()) || ...);
+	}
+#endif // WITH_MASSENTITY_DEBUG
+
+protected:
+	virtual SIZE_T GetAllocatedSize() const override
+	{
+		SIZE_T TotalSize = FragmentsAffected.GetAllocatedSize() + TagsAffected.GetAllocatedSize() + SparseTypes.GetAllocatedSize();
+		for (const auto& KeyValue : Data)
+		{
+			TotalSize += KeyValue.Value.GetAllocatedSize();
+		}
+		TotalSize += Data.GetAllocatedSize();
+		return TotalSize;
+	}
+
+	virtual void Run(FMassEntityManager& EntityManager) override
+	{
+		TRACE_CPUPROFILER_EVENT_SCOPE(MassCommandAddFragmentInstancesWithSharedFragments_Execute);
+
+		constexpr int FragmentTypesCount = UE::Mass::TMultiTypeList<TTypes...>::Ordinal + 1;
+		TArray<FStructArrayView, TInlineAllocator<8>> GenericMultiArray;
+		GenericMultiArray.Reserve(FragmentTypesCount);
+		TArray<FMassArchetypeEntityCollectionWithPayload> EntityCollections;
+		TArray<FMassArchetypeEntityCollection, TInlineAllocator<16>> NonPayloadCollections;
+
+		for (auto& It : Data)
+		{
+			// Exclude tag types from the payload — tags have zero-size storage
+			It.Value.Fragments.GetAsNonTagGenericMultiArray(GenericMultiArray);
+
+			FMassArchetypeEntityCollectionWithPayload::CreateEntityRangesWithPayload(EntityManager, It.Value.TargetEntities, FMassArchetypeEntityCollection::FoldDuplicates
+				, FMassGenericPayloadView(GenericMultiArray), EntityCollections);
+
+			// Sparse element modification is always in-place (no entity move).
+			// Process before the composition change that would invalidate entity collections.
+			if (!SparseTypes.IsEmpty())
+			{
+				NonPayloadCollections.Reserve(EntityCollections.Num());
+				for (const FMassArchetypeEntityCollectionWithPayload& WithPayload : EntityCollections)
+				{
+					NonPayloadCollections.Add(WithPayload.GetEntityCollection());
+				}
+				for (const UScriptStruct* SparseType : SparseTypes)
+				{
+					EntityManager.BatchAddSparseElementToEntities(NonPayloadCollections, SparseType);
+				}
+				NonPayloadCollections.Reset();
+			}
+
+			EntityManager.BatchAddFragmentInstancesForEntities(EntityCollections, FragmentsAffected, TagsAffected, It.Value.SharedFragmentValues);
+
+			GenericMultiArray.Reset();
+			EntityCollections.Reset();
+		}
+	}
+
+	virtual void Reset() override
+	{
+		Data.Reset();
+		Super::Reset();
+	}
+
+#if CSV_PROFILER_STATS || WITH_MASSENTITY_DEBUG
+	virtual int32 GetNumOperationsStat() const override
+	{
+		int32 TotalCount = 0;
+		for (const auto& KeyValue : Data)
+		{
+			TotalCount += KeyValue.Value.TargetEntities.Num();
+		}
+		return TotalCount;
+	}
+#endif // CSV_PROFILER_STATS || WITH_MASSENTITY_DEBUG
+
+	FMassFragmentBitSet FragmentsAffected;
+	FMassTagBitSet TagsAffected;
+	TArray<const UScriptStruct*> SparseTypes;
+
+	struct FPerSharedFragmentsHashData
+	{
+		FPerSharedFragmentsHashData(FMassArchetypeSharedFragmentValues&& InSharedFragmentValues)
+			: SharedFragmentValues(MoveTemp(InSharedFragmentValues))
+		{
+		}
+
+		SIZE_T GetAllocatedSize() const
+		{
+			return TargetEntities.GetAllocatedSize() + Fragments.GetAllocatedSize() + SharedFragmentValues.GetAllocatedSize();
+		}
+
+		TArray<FMassEntityHandle> TargetEntities;
+		mutable UE::Mass::TMultiArray<TTypes...> Fragments;
+		FMassArchetypeSharedFragmentValues SharedFragmentValues;
+	};
+
+	TMap<uint32, FPerSharedFragmentsHashData> Data;
+
+private:
+	template<typename T>
+	void CategorizeElementType()
+	{
+		using FClean = UE::Mass::Clean<T>;
+		if constexpr (UE::Mass::CSparse<FClean>)
+		{
+			SparseTypes.Add(FClean::StaticStruct());
+		}
+		else if constexpr (TIsDerivedFrom<FClean, FMassTag>::Value)
+		{
+			TagsAffected.Add(FClean::StaticStruct());
+		}
+		else
+		{
+			FragmentsAffected.Add(FClean::StaticStruct());
+		}
+	}
+};
+
 //-----------------------------------------------------------------------------
 // Simple fragment composition change
 //-----------------------------------------------------------------------------
@@ -1008,31 +1363,46 @@ using FMassCommandSwapTags = FMassCommandSwapTagsInternal<EMassCommandCheckTime:
 //-----------------------------------------------------------------------------
 /**
  * Adds fragments with initial values to entities. Each entity receives typed fragment data.
+ * Tags (FMassTag, FMassSparseTag) are also accepted — they participate in the composition change but are excluded from the value payload (tags have zero-size storage).
  * Uses PushCommand for cross-chunk batching. Sparse fragments are supported (handled internally).
+ * All entities are accumulated in a single flat batch and dispatched in one BatchAddFragmentInstancesForEntities call.
  *
- * When to use: when you need to add fragments AND set their initial values in one operation.
+ * When to use: when you need to add fragments AND set their initial values in one operation, optionally adding tags in the same entity move.
  * For composition-only changes (no values), prefer FMassCommandAddElements<T...>.
+ * For adding shared fragment values in the same move, use FMassCommandAddFragmentInstancesWithSharedFragments.
  *
  * Example:
- *   // Add with values — entity + fragment instances passed together:
- *   Context.Defer().PushCommand<FMassCommandAddFragmentInstances>(EntityHandle, FMyFragmentA(42), FMyFragmentB(3.14f));
+ *   // Add fragments with values + tags in a single entity move:
+ *   Context.Defer().PushCommand<FMassCommandAddFragmentInstances>(EntityHandle, FMyFragmentA(42), FMyTag());
  *
- * Supported types: FMassFragment subtypes (including sparse fragments).
+ * Supported types: FMassFragment + FMassTag (including sparse).
+ * Tag instances are accepted by Add() for uniform syntax but their values are ignored.
  * Breakpoint support: Yes — CheckFragmentAddBreakpoints called at PushCommand time.
  */
 template<typename... TOthers>
 struct FMassCommandAddFragmentInstances : public FMassBatchedEntityCommand
 {
-	static_assert((TIsDerivedFrom<UE::Mass::Clean<TOthers>, FMassFragment>::Value && ...), "All types passed to FMassCommandAddFragmentInstances must derive from FMassFragment.");
-	// Note: sparse fragments are intentionally allowed here
-	// BatchAddFragmentInstancesForEntities handles them by splitting sparse elements out and dispatching to BatchAddSparseElementToEntities internally.
+	static_assert(((TIsDerivedFrom<UE::Mass::Clean<TOthers>, FMassFragment>::Value
+		|| TIsDerivedFrom<UE::Mass::Clean<TOthers>, FMassTag>::Value) && ...),
+		"All types passed to FMassCommandAddFragmentInstances must derive from FMassFragment or FMassTag.");
+	static_assert((!TIsDerivedFrom<UE::Mass::Clean<TOthers>, FMassSharedFragment>::Value && ...),
+		"Shared fragments cannot be passed to FMassCommandAddFragmentInstances -- use FMassCommandAddFragmentInstancesWithSharedFragments instead.");
+	static_assert((!TIsDerivedFrom<UE::Mass::Clean<TOthers>, FMassConstSharedFragment>::Value && ...),
+		"Const shared fragments cannot be passed to FMassCommandAddFragmentInstances -- use FMassCommandAddFragmentInstancesWithSharedFragments instead.");
+	static_assert((!TIsDerivedFrom<UE::Mass::Clean<TOthers>, FMassChunkFragment>::Value && ...),
+		"Chunk fragments cannot be passed to FMassCommandAddFragmentInstances.");
+
+	// Note: sparse fragments and sparse tags are intentionally allowed in TOthers.
+	// CategorizeElementType routes them to SparseTypes for in-place handling at the command level
+	// before the composition change. They do not appear in FragmentsAffected or TagsAffected.
 
 	using Super = FMassBatchedEntityCommand;
 
 	FMassCommandAddFragmentInstances(EMassCommandOperationType OperationType = EMassCommandOperationType::Set DEBUG_NAME_PARAM("AddFragmentInstanceList"))
 		: Super(OperationType FORWARD_DEBUG_NAME_PARAM)
-		, FragmentsAffected(UE::Mass::Utils::ConstructFragmentBitSet<EMassCommandCheckTime::CompileTimeCheck, TOthers...>())
-	{}
+	{
+		(CategorizeElementType<TOthers>(), ...);
+	}
 
 	template<typename... TInstances>
 	void Add(FMassEntityHandle Entity, TInstances... InFragments)
@@ -1044,7 +1414,12 @@ struct FMassCommandAddFragmentInstances : public FMassBatchedEntityCommand
 #if WITH_MASSENTITY_DEBUG
 	static bool CheckBreakpoints(const FMassEntityHandle Entity, TOthers... InFragments)
 	{
-		return UE::Mass::Debug::FBreakpoint::CheckFragmentAddBreakpoints(Entity, Forward<TOthers>(InFragments)...);
+		// Function pointer to force non-template overload resolution.
+		// Without this, MSVC deduces TFragments=UScriptStruct* from StaticStruct() return type,
+		// matching the variadic template overload instead of the (FMassEntityHandle, const UScriptStruct*) overload.
+		using FCheckFn = bool(*)(const FMassEntityHandle&, const UScriptStruct*);
+		const FCheckFn Check = &UE::Mass::Debug::FBreakpoint::CheckFragmentAddBreakpoints;
+		return (Check(Entity, UE::Mass::Clean<TOthers>::StaticStruct()) || ...);
 	}
 #endif // WITH_MASSENTITY_DEBUG
 
@@ -1057,38 +1432,77 @@ protected:
 
 	virtual SIZE_T GetAllocatedSize() const override
 	{
-		return Super::GetAllocatedSize() + Fragments.GetAllocatedSize() + FragmentsAffected.GetAllocatedSize();
+		return Super::GetAllocatedSize() + Fragments.GetAllocatedSize() + FragmentsAffected.GetAllocatedSize() + TagsAffected.GetAllocatedSize() + SparseTypes.GetAllocatedSize();
 	}
 
 	virtual void Run(FMassEntityManager& EntityManager) override
 	{
 		TRACE_CPUPROFILER_EVENT_SCOPE(MassCommandAddFragmentInstances_Execute);
 
-		TArray<FStructArrayView> GenericMultiArray;
+		TArray<FStructArrayView, TInlineAllocator<8>> GenericMultiArray;
 		GenericMultiArray.Reserve(Fragments.GetNumArrays());
-		Fragments.GetAsGenericMultiArray(GenericMultiArray);
+		// Use GetAsNonTagGenericMultiArray to exclude tag types from the payload —
+		// tags have zero-size storage and BatchSetFragmentValues would crash on them.
+		Fragments.GetAsNonTagGenericMultiArray(GenericMultiArray);
 
 		TArray<FMassArchetypeEntityCollectionWithPayload> EntityCollections;
 		FMassArchetypeEntityCollectionWithPayload::CreateEntityRangesWithPayload(EntityManager, TargetEntities, FMassArchetypeEntityCollection::FoldDuplicates
 			, FMassGenericPayloadView(GenericMultiArray), EntityCollections);
 
-		EntityManager.BatchAddFragmentInstancesForEntities(EntityCollections, FragmentsAffected);
+		// Sparse element modification is always in-place (no entity move).
+		// Process before the composition change that would invalidate entity collections.
+		if (!SparseTypes.IsEmpty())
+		{
+			TArray<FMassArchetypeEntityCollection, TInlineAllocator<16>> NonPayloadCollections;
+			NonPayloadCollections.Reserve(EntityCollections.Num());
+			for (const FMassArchetypeEntityCollectionWithPayload& WithPayload : EntityCollections)
+			{
+				NonPayloadCollections.Add(WithPayload.GetEntityCollection());
+			}
+			for (const UScriptStruct* SparseType : SparseTypes)
+			{
+				EntityManager.BatchAddSparseElementToEntities(NonPayloadCollections, SparseType);
+			}
+		}
+
+		EntityManager.BatchAddFragmentInstancesForEntities(EntityCollections, FragmentsAffected, TagsAffected);
 	}
 
 	mutable UE::Mass::TMultiArray<TOthers...> Fragments;
-	const FMassFragmentBitSet FragmentsAffected;
+	FMassFragmentBitSet FragmentsAffected;
+	FMassTagBitSet TagsAffected;
+	TArray<const UScriptStruct*> SparseTypes;
+
+private:
+	template<typename T>
+	void CategorizeElementType()
+	{
+		using FClean = UE::Mass::Clean<T>;
+		if constexpr (UE::Mass::CSparse<FClean>)
+		{
+			SparseTypes.Add(FClean::StaticStruct());
+		}
+		else if constexpr (TIsDerivedFrom<FClean, FMassTag>::Value)
+		{
+			TagsAffected.Add(FClean::StaticStruct());
+		}
+		else
+		{
+			FragmentsAffected.Add(FClean::StaticStruct());
+		}
+	}
 };
 
 /**
- * Builds reserved entities with fragment values.
+ * Builds reserved entities with fragment values (and optionally tags).
  * Entities must be pre-reserved via EntityManager->ReserveEntity().
  * Uses PushCommand for cross-chunk batching.
  *
  * Example:
  *   FMassEntityHandle Entity = EntityManager->ReserveEntity();
- *   Context.Defer().PushCommand<FMassCommandBuildEntity>(Entity, FMyFragmentA(42), FMyFragmentB(3.14f));
+ *   Context.Defer().PushCommand<FMassCommandBuildEntity>(Entity, FMyFragmentA(42), FMyTag());
  *
- * Supported types: FMassFragment subtypes for TOthers.
+ * Supported types: FMassFragment + FMassTag (including sparse).
  * Breakpoint support: Yes — CheckCreateEntityBreakpoints called at PushCommand time.
  */
 template<typename... TOthers>
@@ -1114,9 +1528,9 @@ protected:
 	{
 		TRACE_CPUPROFILER_EVENT_SCOPE(MassCommandBuildEntity_Execute);
 
-		TArray<FStructArrayView> GenericMultiArray;
+		TArray<FStructArrayView, TInlineAllocator<8>> GenericMultiArray;
 		GenericMultiArray.Reserve(Super::Fragments.GetNumArrays());
-		Super::Fragments.GetAsGenericMultiArray(GenericMultiArray);
+		Super::Fragments.GetAsNonTagGenericMultiArray(GenericMultiArray);
 
 		TArray<FMassArchetypeEntityCollectionWithPayload> EntityCollections;
 		FMassArchetypeEntityCollectionWithPayload::CreateEntityRangesWithPayload(EntityManager, Super::TargetEntities, FMassArchetypeEntityCollection::FoldDuplicates
@@ -1125,7 +1539,28 @@ protected:
 		check(EntityCollections.Num() <= 1);
 		if (EntityCollections.Num())
 		{
-			EntityManager.BatchBuildEntities(EntityCollections[0], Super::FragmentsAffected, FMassArchetypeSharedFragmentValues());
+			if (Super::TagsAffected.IsEmpty())
+			{
+				EntityManager.BatchBuildEntities(EntityCollections[0], Super::FragmentsAffected, FMassArchetypeSharedFragmentValues());
+			}
+			else
+			{
+				FMassElementBitSet Composition(Super::FragmentsAffected);
+				Composition += FMassElementBitSet(Super::TagsAffected);
+				EntityManager.BatchBuildEntities(EntityCollections[0], FMassArchetypeCompositionDescriptor(MoveTemp(Composition)), FMassArchetypeSharedFragmentValues());
+			}
+
+			// Sparse elements are in-place — add after entities are built.
+			// Rebuild collections from TargetEntities since the original collections are stale after build.
+			if (!Super::SparseTypes.IsEmpty())
+			{
+				TArray<FMassArchetypeEntityCollection> FreshCollections;
+				UE::Mass::Utils::CreateEntityCollections(EntityManager, Super::TargetEntities, FMassArchetypeEntityCollection::FoldDuplicates, FreshCollections);
+				for (const UScriptStruct* SparseType : Super::SparseTypes)
+				{
+					EntityManager.BatchAddSparseElementToEntities(FreshCollections, SparseType);
+				}
+			}
 		}
 	}
 };
@@ -1133,6 +1568,7 @@ protected:
 /**
  * Builds reserved entities with fragment values AND shared fragment values.
  * Entities with the same shared fragment values are grouped for efficient archetype creation.
+ * Inherits data management from FMassCommandAddFragmentInstancesWithSharedFragments; only Run() differs.
  * Uses PushCommand for cross-chunk batching.
  *
  * Note: TSharedFragmentValues is always FMassArchetypeSharedFragmentValues but is a template param
@@ -1145,34 +1581,18 @@ protected:
  *   Context.Defer().PushCommand<FMassCommandBuildEntityWithSharedFragments>(
  *       Entity, MoveTemp(SharedValues), FMyFragmentA(1), FMyFragmentB(2.0f));
  *
- * Supported types: FMassFragment subtypes for TOthers. Shared fragments via FMassArchetypeSharedFragmentValues.
+ * Supported types: FMassFragment + FMassTag (including sparse).
+ * Shared fragments via FMassArchetypeSharedFragmentValues.
  * Breakpoint support: Yes — CheckCreateEntityBreakpoints called at PushCommand time.
  */
 template<typename TSharedFragmentValues, typename... TOthers>
-struct FMassCommandBuildEntityWithSharedFragments : public FMassBatchedCommand
+struct FMassCommandBuildEntityWithSharedFragments : public FMassCommandAddFragmentInstancesWithSharedFragments<TSharedFragmentValues, TOthers...>
 {
-	static_assert((TIsDerivedFrom<UE::Mass::Clean<TOthers>, FMassFragment>::Value && ...), "All fragment types passed to FMassCommandBuildEntityWithSharedFragments must derive from FMassFragment.");
-
-	using Super = FMassBatchedCommand;
+	using Super = FMassCommandAddFragmentInstancesWithSharedFragments<TSharedFragmentValues, TOthers...>;
 
 	FMassCommandBuildEntityWithSharedFragments()
-		: Super(EMassCommandOperationType::Create DEBUG_NAME("FMassCommandBuildEntityWithSharedFragments"))
-		, FragmentsAffected(UE::Mass::Utils::ConstructFragmentBitSet<EMassCommandCheckTime::CompileTimeCheck, TOthers...>())
-	{}
-
-	void Add(FMassEntityHandle Entity, FMassArchetypeSharedFragmentValues&& InSharedFragments, TOthers... InFragments)
+		: Super(EMassCommandOperationType::Create DEBUG_NAME("BuildEntityWithSharedFragments"))
 	{
-		InSharedFragments.Sort();
-
-		// Compute hash before adding to the map since evaluation order is not guaranteed
-		// and MoveTemp will invalidate InSharedFragments
-		const uint32 Hash = GetTypeHash(InSharedFragments);
-
-		FPerSharedFragmentsHashData& Instance = Data.FindOrAdd(Hash, MoveTemp(InSharedFragments));
-		Instance.Fragments.Add(InFragments...);
-		Instance.TargetEntities.Add(Entity);
-
-		bHasWork = true;
 	}
 
 #if WITH_MASSENTITY_DEBUG
@@ -1184,82 +1604,56 @@ struct FMassCommandBuildEntityWithSharedFragments : public FMassBatchedCommand
 #endif // WITH_MASSENTITY_DEBUG
 
 protected:
-	virtual SIZE_T GetAllocatedSize() const override
-	{
-		SIZE_T TotalSize = 0;
-		for (const auto& KeyValue : Data)
-		{
-			TotalSize += KeyValue.Value.GetAllocatedSize();
-		}
-		TotalSize += Data.GetAllocatedSize();
-		TotalSize += FragmentsAffected.GetAllocatedSize();
-		return TotalSize;
-	}
-
 	virtual void Run(FMassEntityManager& EntityManager) override
 	{
 		TRACE_CPUPROFILER_EVENT_SCOPE(MassCommandBuildEntityWithSharedFragments_Execute);
 
 		constexpr int FragmentTypesCount = UE::Mass::TMultiTypeList<TOthers...>::Ordinal + 1;
-		TArray<FStructArrayView> GenericMultiArray;
+		TArray<FStructArrayView, TInlineAllocator<8>> GenericMultiArray;
 		GenericMultiArray.Reserve(FragmentTypesCount);
+		TArray<FMassArchetypeEntityCollectionWithPayload> EntityCollections;
+		TArray<FMassArchetypeEntityCollection> FreshCollections;
 
-		for (auto It : Data)
-		{			
-			It.Value.Fragments.GetAsGenericMultiArray(GenericMultiArray);
+		for (auto& It : Super::Data)
+		{
+			It.Value.Fragments.GetAsNonTagGenericMultiArray(GenericMultiArray);
 
-			TArray<FMassArchetypeEntityCollectionWithPayload> EntityCollections;
 			FMassArchetypeEntityCollectionWithPayload::CreateEntityRangesWithPayload(EntityManager, It.Value.TargetEntities, FMassArchetypeEntityCollection::FoldDuplicates
 				, FMassGenericPayloadView(GenericMultiArray), EntityCollections);
-			checkf(EntityCollections.Num() <= 1, TEXT("We expect TargetEntities to only contain archetype-less entities, ones that need to be \'build\'"));
+			checkf(EntityCollections.Num() <= 1, TEXT("We expect TargetEntities to only contain archetype-less entities, ones that need to be \'built\'"));
 
 			if (EntityCollections.Num())
 			{
-				EntityManager.BatchBuildEntities(EntityCollections[0], FragmentsAffected, It.Value.SharedFragmentValues);
+				if (Super::TagsAffected.IsEmpty())
+				{
+					EntityManager.BatchBuildEntities(EntityCollections[0], Super::FragmentsAffected, It.Value.SharedFragmentValues);
+				}
+				else
+				{
+					// Build an FMassElementBitSet combining fragments + tags + shared type bits from values
+					FMassElementBitSet Composition(Super::FragmentsAffected);
+					Composition += FMassElementBitSet(Super::TagsAffected);
+					Composition += It.Value.SharedFragmentValues.GetBitSet();
+					EntityManager.BatchBuildEntities(EntityCollections[0], FMassArchetypeCompositionDescriptor(MoveTemp(Composition)), It.Value.SharedFragmentValues);
+				}
+
+				// Sparse elements are in-place — add after entities are built.
+				// Rebuild collections from this group's TargetEntities since originals are stale after build.
+				if (!Super::SparseTypes.IsEmpty())
+				{
+					UE::Mass::Utils::CreateEntityCollections(EntityManager, It.Value.TargetEntities, FMassArchetypeEntityCollection::FoldDuplicates, FreshCollections);
+					for (const UScriptStruct* SparseType : Super::SparseTypes)
+					{
+						EntityManager.BatchAddSparseElementToEntities(FreshCollections, SparseType);
+					}
+					FreshCollections.Reset();
+				}
 			}
 
 			GenericMultiArray.Reset();
+			EntityCollections.Reset();
 		}
 	}
-
-	virtual void Reset() override
-	{
-		Data.Reset();
-		Super::Reset();
-	}
-
-#if CSV_PROFILER_STATS || WITH_MASSENTITY_DEBUG
-	virtual int32 GetNumOperationsStat() const override
-	{
-		int32 TotalCount = 0;
-		for (const auto& KeyValue : Data)
-		{
-			TotalCount += KeyValue.Value.TargetEntities.Num();
-		}
-		return TotalCount;
-	}
-#endif // CSV_PROFILER_STATS || WITH_MASSENTITY_DEBUG
-
-	FMassFragmentBitSet FragmentsAffected;
-
-	struct FPerSharedFragmentsHashData
-	{
-		FPerSharedFragmentsHashData(FMassArchetypeSharedFragmentValues&& InSharedFragmentValues)
-			: SharedFragmentValues(MoveTemp(InSharedFragmentValues))
-		{	
-		}
-
-		SIZE_T GetAllocatedSize() const
-		{
-			return TargetEntities.GetAllocatedSize() + Fragments.GetAllocatedSize() + SharedFragmentValues.GetAllocatedSize();
-		}
-
-		TArray<FMassEntityHandle> TargetEntities;
-		mutable UE::Mass::TMultiArray<TOthers...> Fragments;
-		FMassArchetypeSharedFragmentValues SharedFragmentValues;
-	};
-
-	TMap<uint32, FPerSharedFragmentsHashData> Data;
 };
 
 //-----------------------------------------------------------------------------
diff --git a/Engine/Source/Runtime/MassEntity/Public/MassEntityManager.h b/Engine/Source/Runtime/MassEntity/Public/MassEntityManager.h
--- a/Engine/Source/Runtime/MassEntity/Public/MassEntityManager.h
+++ b/Engine/Source/Runtime/MassEntity/Public/MassEntityManager.h
@@ -592,12 +592,39 @@ public:
 	 */
 	UE_API int32 BatchReserveEntities(TArrayView<FMassEntityHandle> InOutEntities);
 
-	UE_API TSharedRef<FEntityCreationContext> BatchBuildEntities(const FMassArchetypeEntityCollectionWithPayload& EncodedEntitiesWithPayload, const FMassFragmentBitSet& FragmentsAffected
-		, const FMassArchetypeSharedFragmentValues& SharedFragmentValues = {}, const FMassArchetypeCreationParams& CreationParams = FMassArchetypeCreationParams());
-	UE_API TSharedRef<FEntityCreationContext> BatchBuildEntities(const FMassArchetypeEntityCollectionWithPayload& EncodedEntitiesWithPayload, const FMassArchetypeCompositionDescriptor& Composition
-		, const FMassArchetypeSharedFragmentValues& SharedFragmentValues = {}, const FMassArchetypeCreationParams& CreationParams = FMassArchetypeCreationParams());
-	
-	UE_API void BatchChangeCompositionForEntities(TConstArrayView<FMassArchetypeEntityCollection> EntityCollections, const FMassElementBitSet& ElementsToAdd, const FMassElementBitSet& ElementsToRemove);
+	UE_API TSharedRef<FEntityCreationContext> BatchBuildEntities(
+		const FMassArchetypeEntityCollectionWithPayload& EncodedEntitiesWithPayload
+		, const FMassFragmentBitSet& FragmentsAffected
+		, const FMassArchetypeSharedFragmentValues& SharedFragmentValues = {}
+		, const FMassArchetypeCreationParams& CreationParams = FMassArchetypeCreationParams());
+
+	UE_API TSharedRef<FEntityCreationContext> BatchBuildEntities(
+		const FMassArchetypeEntityCollectionWithPayload& EncodedEntitiesWithPayload
+		, const FMassArchetypeCompositionDescriptor& Composition
+		, const FMassArchetypeSharedFragmentValues& SharedFragmentValues = {}
+		, const FMassArchetypeCreationParams& CreationParams = FMassArchetypeCreationParams());
+
+	UE_API TSharedRef<FEntityCreationContext> BatchBuildEntities(
+		const FMassArchetypeEntityCollectionWithPayload& EncodedEntitiesWithPayload
+		, const FMassElementBitSet& Composition
+		, const FMassArchetypeSharedFragmentValues& SharedFragmentValues = {}
+		, const FMassArchetypeCreationParams& CreationParams = FMassArchetypeCreationParams());
+
+	UE_API void BatchChangeCompositionForEntities(
+		TConstArrayView<FMassArchetypeEntityCollection> EntityCollections
+		, const FMassElementBitSet& ElementsToAdd
+		, const FMassElementBitSet& ElementsToRemove);
+
+	/**
+	 * Overload that also adds shared fragment values during the same entity move.
+	 * ElementsToAdd must include the shared fragment bits from SharedFragmentValuesToAdd.
+	 * Sparse element types must not appear in ElementsToAdd; use the dedicated sparse element APIs instead.
+	 */
+	UE_API void BatchChangeCompositionForEntities(
+		TConstArrayView<FMassArchetypeEntityCollection> EntityCollections,
+		const FMassElementBitSet& ElementsToAdd,
+		const FMassElementBitSet& ElementsToRemove,
+		const FMassArchetypeSharedFragmentValues& SharedFragmentValuesToAdd);
 
 	void BatchChangeTagsForEntities(TConstArrayView<FMassArchetypeEntityCollection> EntityCollections, const FMassTagBitSet& TagsToAdd, const FMassTagBitSet& TagsToRemove)
 	{
@@ -611,8 +638,27 @@ public:
 	}
 
 	UE_API void BatchAddFragmentInstancesForEntities(TConstArrayView<FMassArchetypeEntityCollectionWithPayload> EntityCollections, const FMassFragmentBitSet& FragmentsAffected);
-	/** 
-	 * Adds a new const and non-const shared fragments to all entities provided via EntityCollections 
+
+	/**
+	 * Adds fragment types with per-entity values AND tags to entities in a single entity move.
+	 * Tags are included in the composition change but excluded from value payload (tags have zero-size storage).
+	 */
+	UE_API void BatchAddFragmentInstancesForEntities(TConstArrayView<FMassArchetypeEntityCollectionWithPayload> EntityCollections, const FMassFragmentBitSet& FragmentsAffected, const FMassTagBitSet& TagsToAdd);
+
+	/**
+	 * Adds fragment types with per-entity values, tags, and shared fragment values — all in a single entity move.
+	 * Tags are included in composition but excluded from value payload. Shared fragment values are set per-chunk.
+	 */
+	UE_API void BatchAddFragmentInstancesForEntities(TConstArrayView<FMassArchetypeEntityCollectionWithPayload> EntityCollections, const FMassFragmentBitSet& FragmentsAffected, const FMassTagBitSet& TagsToAdd, const FMassArchetypeSharedFragmentValues& SharedFragmentValuesToAdd);
+
+	/**
+	 * Adds fragment types, sets their per-entity values, and adds shared fragment values — all in a single entity move.
+	 * Combines the functionality of BatchAddFragmentInstancesForEntities and BatchAddSharedFragmentsForEntities.
+	 */
+	UE_API void BatchAddFragmentInstancesForEntities(TConstArrayView<FMassArchetypeEntityCollectionWithPayload> EntityCollections, const FMassFragmentBitSet& FragmentsAffected, const FMassArchetypeSharedFragmentValues& SharedFragmentValuesToAdd);
+
+	/**
+	 * Adds a new const and non-const shared fragments to all entities provided via EntityCollections
 	 */
 	UE_API void BatchAddSharedFragmentsForEntities(TConstArrayView<FMassArchetypeEntityCollection> EntityCollections, const FMassArchetypeSharedFragmentValues& AddedFragmentValues);
 
diff --git a/Engine/Source/Runtime/MassEntity/Public/MassEntityTypes.h b/Engine/Source/Runtime/MassEntity/Public/MassEntityTypes.h
--- a/Engine/Source/Runtime/MassEntity/Public/MassEntityTypes.h
+++ b/Engine/Source/Runtime/MassEntity/Public/MassEntityTypes.h
@@ -879,12 +879,25 @@ namespace UE::Mass
 			Super::Add(Forward<TInstances>(Rest)...);
 		}
 
-		void GetAsGenericMultiArray(TArray<FStructArrayView>& A) /*const*/
+		template<typename AllocatorType = FDefaultAllocator>
+		void GetAsGenericMultiArray(TArray<FStructArrayView, AllocatorType>& A) /*const*/
 		{
 			Super::GetAsGenericMultiArray(A);
 			A.Add(FStructArrayView(FragmentInstances));
 		}
 
+		/** Like GetAsGenericMultiArray but skips FMassTag-derived types (including FMassSparseTag).
+		 *  Tags have zero-size storage and must be excluded from payload arrays passed to BatchSetFragmentValues. */
+		template<typename AllocatorType = FDefaultAllocator>
+		void GetAsNonTagGenericMultiArray(TArray<FStructArrayView, AllocatorType>& A) /*const*/
+		{
+			Super::GetAsNonTagGenericMultiArray(A);
+			if constexpr (!TIsDerivedFrom<FType, FMassTag>::Value)
+			{
+				A.Add(FStructArrayView(FragmentInstances));
+			}
+		}
+
 		void GatherAffectedFragments(FMassFragmentBitSet& OutBitSet) const
 		{
 			Super::GatherAffectedFragments(OutBitSet);
@@ -933,11 +946,22 @@ namespace UE::Mass
 			FragmentInstances.Emplace(Forward<FType>(Item));
 		}
 
-		void GetAsGenericMultiArray(TArray<FStructArrayView>& A) /*const*/
+		template<typename AllocatorType = FDefaultAllocator>
+		void GetAsGenericMultiArray(TArray<FStructArrayView, AllocatorType>& A) /*const*/
 		{
 			A.Add(FStructArrayView(FragmentInstances));
 		}
 
+		/** Like GetAsGenericMultiArray but skips FMassTag-derived types (including FMassSparseTag). */
+		template<typename AllocatorType = FDefaultAllocator>
+		void GetAsNonTagGenericMultiArray(TArray<FStructArrayView, AllocatorType>& A) /*const*/
+		{
+			if constexpr (!TIsDerivedFrom<FType, FMassTag>::Value)
+			{
+				A.Add(FStructArrayView(FragmentInstances));
+			}
+		}
+
 		void GatherAffectedFragments(FMassFragmentBitSet& OutBitSet) const
 		{
 			OutBitSet += FMassFragmentBitSet::GetTypeBitSet<FType>();
